NetworkPrediction plugin initial checkin.
-This is a WIP plugin, disabled by default.
-Goal is to generalize player prediction and eventually improve our networked movement system.
-See readme.txt in NetworkPrediction plugin folder for additional info

#rxn

diff --git a/Engine/Plugins/Runtime/NetworkPrediction/NetworkPrediction.uplugin b/Engine/Plugins/Runtime/NetworkPrediction/NetworkPrediction.uplugin
new file mode 100644
--- /dev/null
+++ b/Engine/Plugins/Runtime/NetworkPrediction/NetworkPrediction.uplugin
@@ -0,0 +1,24 @@
+{
+	"FileVersion": 3,
+	"Version": 1,
+	"VersionName": "0.1",
+	"FriendlyName": "Network Prediction",
+	"Description": "Generalized framework for writing network prediction friendly gameplay systems",
+	"Category": "Gameplay",
+	"CreatedBy": "Epic Games, Inc.",
+	"CreatedByURL": "http://epicgames.com",
+	"DocsURL": "",
+	"MarketplaceURL": "",
+	"SupportURL": "",
+	"EnabledByDefault": false,
+	"CanContainContent": false,
+	"IsBetaVersion": true,
+	"Installed": false,
+	"Modules": [
+		{
+			"Name": "NetworkPrediction",
+			"Type": "Runtime",
+			"LoadingPhase": "Default"
+		}		
+	]
+}
\ No newline at end of file
diff --git a/Engine/Plugins/Runtime/NetworkPrediction/Source/NetworkPrediction/NetworkPrediction.Build.cs b/Engine/Plugins/Runtime/NetworkPrediction/Source/NetworkPrediction/NetworkPrediction.Build.cs
new file mode 100644
--- /dev/null
+++ b/Engine/Plugins/Runtime/NetworkPrediction/Source/NetworkPrediction/NetworkPrediction.Build.cs
@@ -0,0 +1,45 @@
+// Copyright 1998-2018 Epic Games, Inc. All Rights Reserved.
+
+namespace UnrealBuildTool.Rules
+{
+	public class NetworkPrediction : ModuleRules
+	{
+		public NetworkPrediction(ReadOnlyTargetRules Target) : base(Target)
+		{
+			PublicIncludePaths.AddRange(
+				new string[] {
+                    ModuleDirectory + "/Public",
+				}
+				);
+
+			PrivateIncludePaths.AddRange(
+				new string[] {
+                    "NetworkPrediction/Private",
+				}
+				);
+
+			PublicDependencyModuleNames.AddRange(
+				new string[]
+				{
+					"Core",
+					"CoreUObject",
+                    "Engine",
+                    "RenderCore",
+				}
+				);
+
+			PrivateDependencyModuleNames.AddRange(
+				new string[]
+				{
+				}
+				);
+
+			DynamicallyLoadedModuleNames.AddRange(
+				new string[]
+				{
+				}
+				);
+
+		}
+	}
+}
diff --git a/Engine/Plugins/Runtime/NetworkPrediction/Source/NetworkPrediction/Private/Movement/FlyingMovementComponent.cpp b/Engine/Plugins/Runtime/NetworkPrediction/Source/NetworkPrediction/Private/Movement/FlyingMovementComponent.cpp
new file mode 100644
--- /dev/null
+++ b/Engine/Plugins/Runtime/NetworkPrediction/Source/NetworkPrediction/Private/Movement/FlyingMovementComponent.cpp
@@ -0,0 +1,478 @@
+// Copyright 1998-2019 Epic Games, Inc. All Rights Reserved.
+
+
+#include "Movement/FlyingMovement.h"
+#include "GameFramework/Actor.h"
+#include "Engine/World.h"
+#include "Components/PrimitiveComponent.h"
+#include "GameFramework/CharacterMovementComponent.h" // Temp
+#include "GameFramework/PlayerController.h"
+#include "Net/UnrealNetwork.h"
+#include "Engine/LocalPlayer.h"
+#include "Misc/CoreDelegates.h"
+#include "UObject/CoreNet.h"
+#include "Engine/NetDriver.h"
+#include "Engine/NetConnection.h"
+#include "UObject/UObjectIterator.h"
+#include "Components/CapsuleComponent.h"
+#include "VisualLogger/VisualLogger.h"
+#include "Math/Color.h"
+#include "DrawDebugHelpers.h"
+#include "Debug/DebugDrawService.h"
+#include "Engine/Canvas.h"
+#include "Engine/Engine.h"
+#include "Debug/ReporterGraph.h"
+#include "NetworkSimulationModelDebugger.h"
+
+DEFINE_LOG_CATEGORY_STATIC(LogFlyingMovement, Log, All);
+
+
+namespace FlyingMovementCVars
+{
+
+static float PenetrationPullbackDistance = 0.125f;
+static FAutoConsoleVariableRef CVarPenetrationPullbackDistance(TEXT("fp.PenetrationPullbackDistance"),
+	PenetrationPullbackDistance,
+	TEXT("Pull out from penetration of an object by this extra distance.\n")
+	TEXT("Distance added to penetration fix-ups."),
+	ECVF_Default);
+
+static float PenetrationOverlapCheckInflation = 0.100f;
+static FAutoConsoleVariableRef CVarPenetrationOverlapCheckInflation(TEXT("motion.PenetrationOverlapCheckInflation"),
+	PenetrationOverlapCheckInflation,
+	TEXT("Inflation added to object when checking if a location is free of blocking collision.\n")
+	TEXT("Distance added to inflation in penetration overlap check."),
+	ECVF_Default);
+
+static int32 RequestMispredict = 0;
+static FAutoConsoleVariableRef CVarRequestMispredict(TEXT("fp.RequestMispredict"),
+	RequestMispredict, TEXT("Causes a misprediction by inserting random value into stream on authority side"), ECVF_Default);
+}
+
+// ----------------------------------------------------------------------------------------------------------
+//	UFlyingMovementComponent setup/init
+// ----------------------------------------------------------------------------------------------------------
+
+UFlyingMovementComponent::UFlyingMovementComponent()
+{
+	PrimaryComponentTick.TickGroup = TG_PrePhysics;
+	PrimaryComponentTick.bCanEverTick = true;
+	
+	bWantsInitializeComponent = true;
+	bAutoActivate = true;
+	bReplicates = true;
+
+	bWantsInitializeComponent = true;
+}
+
+void UFlyingMovementComponent::RegisterComponentTickFunctions(bool bRegister)
+{
+	Super::RegisterComponentTickFunctions(bRegister);
+
+	// Super may start up the tick function when we don't want to.
+	UpdateTickRegistration();
+
+	// If the owner ticks, make sure we tick first. This is to ensure the owner's location will be up to date when it ticks.
+	AActor* Owner = GetOwner();
+	
+	if (bRegister && PrimaryComponentTick.bCanEverTick && Owner && Owner->CanEverTick())
+	{
+		Owner->PrimaryActorTick.AddPrerequisite(this, PrimaryComponentTick);
+	}
+}
+
+void UFlyingMovementComponent::UpdateTickRegistration()
+{
+	const bool bHasUpdatedComponent = (UpdatedComponent != NULL);
+	SetComponentTickEnabled(bHasUpdatedComponent && bAutoActivate);
+}
+
+void UFlyingMovementComponent::InitializeComponent()
+{
+	TGuardValue<bool> InInitializeComponentGuard(bInInitializeComponent, true);
+	Super::InitializeComponent();
+
+	// RootComponent is null in OnRegister for blueprint (non-native) root components.
+	if (!UpdatedComponent)
+	{
+		// Auto-register owner's root component if found.
+		if (AActor* MyActor = GetOwner())
+		{
+			if (USceneComponent* NewUpdatedComponent = MyActor->GetRootComponent())
+			{
+				SetUpdatedComponent(NewUpdatedComponent);
+			}
+		}
+	}
+}
+
+void UFlyingMovementComponent::OnRegister()
+{
+	TGuardValue<bool> InOnRegisterGuard(bInOnRegister, true);
+
+	UpdatedPrimitive = Cast<UPrimitiveComponent>(UpdatedComponent);
+	Super::OnRegister();
+
+	const UWorld* MyWorld = GetWorld();
+	if (MyWorld && MyWorld->IsGameWorld())
+	{
+		USceneComponent* NewUpdatedComponent = UpdatedComponent;
+		if (!UpdatedComponent)
+		{
+			// Auto-register owner's root component if found.
+			AActor* MyActor = GetOwner();
+			if (MyActor)
+			{
+				NewUpdatedComponent = MyActor->GetRootComponent();
+			}
+		}
+
+		SetUpdatedComponent(NewUpdatedComponent);
+	}
+}
+
+void UFlyingMovementComponent::SetUpdatedComponent(USceneComponent* NewUpdatedComponent)
+{
+	if (UpdatedComponent && UpdatedComponent != NewUpdatedComponent)
+	{
+		UpdatedComponent->SetShouldUpdatePhysicsVolume(false);
+		if (!UpdatedComponent->IsPendingKill())
+		{
+			UpdatedComponent->SetPhysicsVolume(NULL, true);
+			UpdatedComponent->PhysicsVolumeChangedDelegate.RemoveDynamic(this, &UFlyingMovementComponent::PhysicsVolumeChanged);
+		}
+
+		// remove from tick prerequisite
+		UpdatedComponent->PrimaryComponentTick.RemovePrerequisite(this, PrimaryComponentTick); 
+	}
+
+	// Don't assign pending kill components, but allow those to null out previous UpdatedComponent.
+	UpdatedComponent = IsValid(NewUpdatedComponent) ? NewUpdatedComponent : NULL;
+	UpdatedPrimitive = Cast<UPrimitiveComponent>(UpdatedComponent);
+
+	// Assign delegates
+	if (UpdatedComponent && !UpdatedComponent->IsPendingKill())
+	{
+		UpdatedComponent->SetShouldUpdatePhysicsVolume(true);
+		UpdatedComponent->PhysicsVolumeChangedDelegate.AddUniqueDynamic(this, &UFlyingMovementComponent::PhysicsVolumeChanged);
+
+		if (!bInOnRegister && !bInInitializeComponent)
+		{
+			// UpdateOverlaps() in component registration will take care of this.
+			UpdatedComponent->UpdatePhysicsVolume(true);
+		}
+		
+		// force ticks after movement component updates
+		UpdatedComponent->PrimaryComponentTick.AddPrerequisite(this, PrimaryComponentTick); 
+	}
+	UpdateTickRegistration();
+}
+
+void UFlyingMovementComponent::PhysicsVolumeChanged(APhysicsVolume* NewVolume)
+{
+	// This itself feels bad. When will this be called? Its impossible to know what is allowed and not allowed to be done in this callback.
+	// Callbacks instead should be trapped within the simulation update function. This isn't really possible though since the UpdateComponent
+	// is the one that will call this.
+}
+
+// ----------------------------------------------------------------------------------------------------------
+//	Core Network Prediction functions
+// ----------------------------------------------------------------------------------------------------------
+
+IReplicationProxy* UFlyingMovementComponent::InstantiateNetworkSimulation()
+{
+	NetworkSim.Reset(new FlyingMovement::FMovementSystem());
+
+#if NETSIM_MODEL_DEBUG
+	FNetworkSimulationModelDebuggerManager::Get().RegisterNetworkSimulationModel(NetworkSim.Get(), (IFlyingMovementDriver*)this, GetOwner(), TEXT("FlyingPawnMovement"));
+#endif
+	return NetworkSim.Get();
+}
+
+	// Child classes should override this an initialize their NetworkSim here
+void UFlyingMovementComponent::InitializeForNetworkRole(ENetRole Role)
+{
+	check(NetworkSim);
+
+	FNetworkSimulationModelInitParameters InitParams;
+	InitParams.InputBufferSize = Role != ROLE_SimulatedProxy ? 32 : 0;
+	InitParams.SyncedBufferSize = Role != ROLE_AutonomousProxy ? 2 : 32;
+	InitParams.AuxBufferSize = Role != ROLE_AutonomousProxy ? 2 : 32;
+	InitParams.DebugBufferSize = 32;
+	InitParams.HistoricBufferSize = 128;
+	NetworkSim->InitializeForNetworkRole(Role, IsLocallyControlled(), InitParams);
+}
+
+void UFlyingMovementComponent::TickComponent(float DeltaTime, enum ELevelTick TickType, FActorComponentTickFunction* ThisTickFunction)
+{
+	// TEMP! Disable existing CMC if it is activate. Just makes A/B testing eaiser for now.
+	if (AActor* Owner = GetOwner())
+	{
+		if (UCharacterMovementComponent* OldComp = Owner->FindComponentByClass<UCharacterMovementComponent>())
+		{
+			if (OldComp->IsActive())
+			{
+				OldComp->Deactivate();
+				Owner->bReplicateMovement = false;
+			}
+		}
+	}
+
+
+
+	Super::TickComponent(DeltaTime, TickType, ThisTickFunction);
+	const ENetRole OwnerRole = GetOwnerRole();
+
+	// -------------------------------------
+	// Tick the network sim
+	// -------------------------------------
+	if (NetworkSim)
+	{
+		// Check if we should trip a mispredict. (Note how its not possible to do this inside the Update function!)
+		if (OwnerRole == ROLE_Authority && FlyingMovementCVars::RequestMispredict)
+		{
+			FlyingMovement::FMovementSystem::ForceMispredict = true;
+			FlyingMovementCVars::RequestMispredict = 0;
+		}
+
+		FlyingMovement::FMovementSystem::FTickParameters Parameters;
+		Parameters.Role = OwnerRole;
+		Parameters.LocalDeltaTimeSeconds = DeltaTime;
+
+		// Tick the core network sim, this will consume input and generate new sync state
+		NetworkSim->Tick((IFlyingMovementDriver*)this, Parameters);
+
+		// Client->Server communication
+		if (NetworkSim->ShouldSendServerRPC(OwnerRole, DeltaTime))
+		{
+			// Temp hack to make sure the ServerRPC doesn't get suppressed from bandwidth limiting
+			// (system hasn't been optimized and not mature enough yet to handle gaps in input stream)
+			FScopedBandwidthLimitBypass BandwidthBypass(GetOwner());
+
+			FServerReplicationRPCParameter ProxyParameter(ReplicationProxy_ServerRPC);
+			ServerRecieveClientInput(ProxyParameter);
+		}
+	}
+}
+
+FlyingMovement::FInputCmd* UFlyingMovementComponent::GetNextClientInputCmdForWrite(float DeltaTime)
+{
+	// This is unfortunate but we must check here to initialize the InputCmdBuffer.
+	// -Tick alone is not appropriate since the PC may tick first and generate new user cmds (before our first tick)
+	// -Startup functions (BeginPlay, PostInitializeComponent etc) are not appropriate because Net Role is not known.
+	// -PostNetReceive could work but would be inefficient since we would be checking each simulated proxy too, ever PostNetReceive.
+	//
+	// This works out best because its only the locally controlled actor that will call into this.
+	// Its just unfortunate it will only pass the first time but we must keep checking.
+	CheckOwnerRoleChange();
+
+	if (NetworkSim.IsValid())
+	{
+		return NetworkSim->GetNextInputForWrite(DeltaTime);
+	}
+	return nullptr;
+}
+
+// ----------------------------------------------------------------------------------------------------------
+//	Movement System Driver
+//
+//	NOTE: Most of the Movement Driver is not ideal! We are at the mercy of the UpdateComponent since it is the
+//	the object that owns its collision data and its MoveComponent function. Ideally we would have everything within
+//	the movement simulation code and it do its own collision queries. But instead we have to come back to the Driver/Component
+//	layer to do this kind of stuff.
+//
+// ----------------------------------------------------------------------------------------------------------
+
+FVector UFlyingMovementComponent::GetPenetrationAdjustment(const FHitResult& Hit)
+{
+	if (!Hit.bStartPenetrating)
+	{
+		return FVector::ZeroVector;
+	}
+
+	FVector Result;
+	const float PullBackDistance = FMath::Abs(FlyingMovementCVars::PenetrationPullbackDistance);
+	const float PenetrationDepth = (Hit.PenetrationDepth > 0.f ? Hit.PenetrationDepth : 0.125f);
+
+	Result = Hit.Normal * (PenetrationDepth + PullBackDistance);
+	
+	return Result;
+}
+
+bool UFlyingMovementComponent::OverlapTest(const FVector& Location, const FQuat& RotationQuat, const ECollisionChannel CollisionChannel, const FCollisionShape& CollisionShape, const AActor* IgnoreActor) const
+{
+	FCollisionQueryParams QueryParams(SCENE_QUERY_STAT(MovementOverlapTest), false, IgnoreActor);
+	FCollisionResponseParams ResponseParam;
+	InitCollisionParams(QueryParams, ResponseParam);
+	return GetWorld()->OverlapBlockingTestByChannel(Location, RotationQuat, CollisionChannel, CollisionShape, QueryParams, ResponseParam);
+}
+
+void UFlyingMovementComponent::InitCollisionParams(FCollisionQueryParams &OutParams, FCollisionResponseParams& OutResponseParam) const
+{
+	if (UpdatedPrimitive)
+	{
+		UpdatedPrimitive->InitSweepCollisionParams(OutParams, OutResponseParam);
+	}
+}
+
+bool UFlyingMovementComponent::ResolvePenetration(const FVector& ProposedAdjustment, const FHitResult& Hit, const FQuat& NewRotationQuat) const
+{
+	// SceneComponent can't be in penetration, so this function really only applies to PrimitiveComponent.
+	const FVector Adjustment = ProposedAdjustment; //ConstrainDirectionToPlane(ProposedAdjustment);
+	if (!Adjustment.IsZero() && UpdatedPrimitive)
+	{
+		QUICK_SCOPE_CYCLE_COUNTER(STAT_FlyingMovementComponent_ResolvePenetration);
+		// See if we can fit at the adjusted location without overlapping anything.
+		AActor* ActorOwner = UpdatedComponent->GetOwner();
+		if (!ActorOwner)
+		{
+			return false;
+		}
+
+		UE_LOG(LogFlyingMovement, Verbose, TEXT("ResolvePenetration: %s.%s at location %s inside %s.%s at location %s by %.3f (netmode: %d)"),
+			   *ActorOwner->GetName(),
+			   *UpdatedComponent->GetName(),
+			   *UpdatedComponent->GetComponentLocation().ToString(),
+			   *GetNameSafe(Hit.GetActor()),
+			   *GetNameSafe(Hit.GetComponent()),
+			   Hit.Component.IsValid() ? *Hit.GetComponent()->GetComponentLocation().ToString() : TEXT("<unknown>"),
+			   Hit.PenetrationDepth,
+			   (uint32)GetNetMode());
+
+		// We really want to make sure that precision differences or differences between the overlap test and sweep tests don't put us into another overlap,
+		// so make the overlap test a bit more restrictive.
+		const float OverlapInflation = FlyingMovementCVars::PenetrationOverlapCheckInflation;
+		bool bEncroached = OverlapTest(Hit.TraceStart + Adjustment, NewRotationQuat, UpdatedPrimitive->GetCollisionObjectType(), UpdatedPrimitive->GetCollisionShape(OverlapInflation), ActorOwner);
+		if (!bEncroached)
+		{
+			// Move without sweeping.
+			MoveUpdatedComponent(Adjustment, NewRotationQuat, false, nullptr, ETeleportType::TeleportPhysics);
+			UE_LOG(LogFlyingMovement, Verbose, TEXT("ResolvePenetration:   teleport by %s"), *Adjustment.ToString());
+			return true;
+		}
+		else
+		{
+			// Disable MOVECOMP_NeverIgnoreBlockingOverlaps if it is enabled, otherwise we wouldn't be able to sweep out of the object to fix the penetration.
+			TGuardValue<EMoveComponentFlags> ScopedFlagRestore(MoveComponentFlags, EMoveComponentFlags(MoveComponentFlags & (~MOVECOMP_NeverIgnoreBlockingOverlaps)));
+
+			// Try sweeping as far as possible...
+			FHitResult SweepOutHit(1.f);
+			bool bMoved = MoveUpdatedComponent(Adjustment, NewRotationQuat, true, &SweepOutHit, ETeleportType::TeleportPhysics);
+			UE_LOG(LogFlyingMovement, Verbose, TEXT("ResolvePenetration:   sweep by %s (success = %d)"), *Adjustment.ToString(), bMoved);
+			
+			// Still stuck?
+			if (!bMoved && SweepOutHit.bStartPenetrating)
+			{
+				// Combine two MTD results to get a new direction that gets out of multiple surfaces.
+				const FVector SecondMTD = GetPenetrationAdjustment(SweepOutHit);
+				const FVector CombinedMTD = Adjustment + SecondMTD;
+				if (SecondMTD != Adjustment && !CombinedMTD.IsZero())
+				{
+					bMoved = MoveUpdatedComponent(CombinedMTD, NewRotationQuat, true, nullptr, ETeleportType::TeleportPhysics);
+					UE_LOG(LogFlyingMovement, Verbose, TEXT("ResolvePenetration:   sweep by %s (MTD combo success = %d)"), *CombinedMTD.ToString(), bMoved);
+				}
+			}
+
+			// Still stuck?
+			if (!bMoved)
+			{
+				// Try moving the proposed adjustment plus the attempted move direction. This can sometimes get out of penetrations with multiple objects
+				const FVector MoveDelta = (Hit.TraceEnd - Hit.TraceStart); //ConstrainDirectionToPlane(Hit.TraceEnd - Hit.TraceStart);
+				if (!MoveDelta.IsZero())
+				{
+					bMoved = MoveUpdatedComponent(Adjustment + MoveDelta, NewRotationQuat, true, nullptr, ETeleportType::TeleportPhysics);
+					UE_LOG(LogFlyingMovement, Verbose, TEXT("ResolvePenetration:   sweep by %s (adjusted attempt success = %d)"), *(Adjustment + MoveDelta).ToString(), bMoved);
+				}
+			}	
+
+			return bMoved;
+		}
+	}
+
+	return false;
+}
+
+bool UFlyingMovementComponent::SafeMoveUpdatedComponent(const FVector& Delta, const FQuat& NewRotation, bool bSweep, FHitResult& OutHit, ETeleportType Teleport) const
+{
+	if (UpdatedComponent == NULL)
+	{
+		OutHit.Reset(1.f);
+		return false;
+	}
+
+	bool bMoveResult = false;
+
+	// Scope for move flags
+	{
+		bMoveResult = MoveUpdatedComponent(Delta, NewRotation, bSweep, &OutHit, Teleport);
+	}
+
+	// Handle initial penetrations
+	if (OutHit.bStartPenetrating && UpdatedComponent)
+	{
+		const FVector RequestedAdjustment = GetPenetrationAdjustment(OutHit);
+		if (ResolvePenetration(RequestedAdjustment, OutHit, NewRotation))
+		{
+			// Retry original move
+			bMoveResult = MoveUpdatedComponent(Delta, NewRotation, bSweep, &OutHit, Teleport);
+		}
+	}
+
+	return bMoveResult;
+}
+
+bool UFlyingMovementComponent::MoveUpdatedComponent(const FVector& Delta, const FQuat& NewRotation, bool bSweep, FHitResult* OutHit, ETeleportType Teleport) const
+{
+	if (UpdatedComponent)
+	{
+		const FVector NewDelta = Delta; //ConstrainDirectionToPlane(Delta);
+		return UpdatedComponent->MoveComponent(NewDelta, NewRotation, bSweep, OutHit, MoveComponentFlags, Teleport);
+	}
+
+	return false;
+}
+
+FTransform UFlyingMovementComponent::GetUpdateComponentTransform() const
+{
+	if (ensure(UpdatedComponent))
+	{
+		return UpdatedComponent->GetComponentTransform();		
+	}
+	return FTransform::Identity;
+}
+
+void UFlyingMovementComponent::InitSyncState(FlyingMovement::FMoveState& OutSyncState) const
+{
+	OutSyncState.Location = UpdatedComponent->GetComponentLocation();
+	OutSyncState.Rotation = UpdatedComponent->GetComponentQuat().Rotator();	
+}
+
+void UFlyingMovementComponent::SyncTo(const FlyingMovement::FMoveState& SyncState)
+{
+	// Does checking equality make any sense here? This is unfortunate
+	if (UpdatedComponent->GetComponentLocation().Equals(SyncState.Location) == false || UpdatedComponent->GetComponentQuat().Rotator().Equals(SyncState.Rotation, FlyingMovement::ROTATOR_TOLERANCE) == false)
+	{
+		FTransform Transform(SyncState.Rotation.Quaternion(), SyncState.Location, UpdatedComponent->GetComponentTransform().GetScale3D() );
+		UpdatedComponent->SetWorldTransform(Transform, false, nullptr, ETeleportType::TeleportPhysics);
+
+		UpdatedComponent->ComponentVelocity = SyncState.Velocity;
+	}
+}
+
+void UFlyingMovementComponent::GetCapsuleDimensions(float &Radius, float& HalfHeight) const
+{
+	if (UCapsuleComponent const* CapsuleComponent = Cast<UCapsuleComponent>(UpdatedComponent))
+	{
+		CapsuleComponent->GetScaledCapsuleSize(Radius, HalfHeight);
+	}
+}
+
+FTransform UFlyingMovementComponent::GetDebugWorldTransform() const
+{
+	return GetOwner()->GetActorTransform();
+}
+
+UObject* UFlyingMovementComponent::GetVLogOwner() const
+{
+	return GetOwner();
+}
\ No newline at end of file
diff --git a/Engine/Plugins/Runtime/NetworkPrediction/Source/NetworkPrediction/Private/Movement/FlyingMovementSimulation.cpp b/Engine/Plugins/Runtime/NetworkPrediction/Source/NetworkPrediction/Private/Movement/FlyingMovementSimulation.cpp
new file mode 100644
--- /dev/null
+++ b/Engine/Plugins/Runtime/NetworkPrediction/Source/NetworkPrediction/Private/Movement/FlyingMovementSimulation.cpp
@@ -0,0 +1,319 @@
+// Copyright 1998-2019 Epic Games, Inc. All Rights Reserved.
+
+
+#include "Movement/FlyingMovement.h"
+#include "GameFramework/Actor.h"
+#include "Engine/World.h"
+#include "Components/PrimitiveComponent.h"
+#include "GameFramework/CharacterMovementComponent.h" // Temp
+#include "GameFramework/PlayerController.h"
+#include "Net/UnrealNetwork.h"
+#include "Engine/LocalPlayer.h"
+#include "Misc/CoreDelegates.h"
+#include "UObject/CoreNet.h"
+#include "Engine/NetDriver.h"
+#include "Engine/NetConnection.h"
+#include "UObject/UObjectIterator.h"
+#include "Components/CapsuleComponent.h"
+#include "VisualLogger/VisualLogger.h"
+#include "Math/Color.h"
+#include "DrawDebugHelpers.h"
+#include "Debug/DebugDrawService.h"
+#include "Engine/Canvas.h"
+#include "Engine/Engine.h"
+#include "Debug/ReporterGraph.h"
+#include "NetworkSimulationModelDebugger.h"
+
+DEFINE_LOG_CATEGORY_STATIC(LogFlyingPawnSimulation, Log, All);
+
+namespace FlyingPawnSimCVars
+{
+static int32 UseVLogger = 1;
+static FAutoConsoleVariableRef CVarUseVLogger(TEXT("fp.Debug.UseUnrealVLogger"),
+	UseVLogger,	TEXT("Use Unreal Visual Logger\n"),	ECVF_Default);
+
+static int32 UseDrawDebug = 1;
+static FAutoConsoleVariableRef CVarUseDrawDebug(TEXT("fp.Debug.UseDrawDebug"),
+	UseVLogger,	TEXT("Use built in DrawDebug* functions for visual logging\n"), ECVF_Default);
+
+static float DrawDebugDefaultLifeTime = 30.f;
+static FAutoConsoleVariableRef CVarDrawDebugDefaultLifeTime(TEXT("fp.Debug.UseDrawDebug.DefaultLifeTime"),
+	DrawDebugDefaultLifeTime, TEXT("Use built in DrawDebug* functions for visual logging"), ECVF_Default);
+
+}
+
+// ============================================================================================
+
+namespace FlyingMovement
+{
+	bool FMovementSystem::ForceMispredict = false;
+	static FVector ForceMispredictVelocityMagnitude = FVector(2000.f, 0.f, 0.f);
+
+	bool IsExceedingMaxSpeed(const FVector& Velocity, float InMaxSpeed)
+	{
+		InMaxSpeed = FMath::Max(0.f, InMaxSpeed);
+		const float MaxSpeedSquared = FMath::Square(InMaxSpeed);
+	
+		// Allow 1% error tolerance, to account for numeric imprecision.
+		const float OverVelocityPercent = 1.01f;
+		return (Velocity.SizeSquared() > MaxSpeedSquared * OverVelocityPercent);
+	}
+
+	FVector ComputeSlideVector(const FVector& Delta, const float Time, const FVector& Normal, const FHitResult& Hit)
+	{
+		return FVector::VectorPlaneProject(Delta, Normal) * Time;
+	}
+
+	void TwoWallAdjust(FVector& OutDelta, const FHitResult& Hit, const FVector& OldHitNormal)
+	{
+		FVector Delta = OutDelta;
+		const FVector HitNormal = Hit.Normal;
+
+		if ((OldHitNormal | HitNormal) <= 0.f) //90 or less corner, so use cross product for direction
+		{
+			const FVector DesiredDir = Delta;
+			FVector NewDir = (HitNormal ^ OldHitNormal);
+			NewDir = NewDir.GetSafeNormal();
+			Delta = (Delta | NewDir) * (1.f - Hit.Time) * NewDir;
+			if ((DesiredDir | Delta) < 0.f)
+			{
+				Delta = -1.f * Delta;
+			}
+		}
+		else //adjust to new wall
+		{
+			const FVector DesiredDir = Delta;
+			Delta = ComputeSlideVector(Delta, 1.f - Hit.Time, HitNormal, Hit);
+			if ((Delta | DesiredDir) <= 0.f)
+			{
+				Delta = FVector::ZeroVector;
+			}
+			else if ( FMath::Abs((HitNormal | OldHitNormal) - 1.f) < KINDA_SMALL_NUMBER )
+			{
+				// we hit the same wall again even after adjusting to move along it the first time
+				// nudge away from it (this can happen due to precision issues)
+				Delta += HitNormal * 0.01f;
+			}
+		}
+
+		OutDelta = Delta;
+	}
+
+	float SlideAlongSurface(FMovementSystem::IMovementDriver* Driver, const FVector& Delta, float Time, const FQuat Rotation, const FVector& Normal, FHitResult& Hit, bool bHandleImpact)
+	{
+		if (!Hit.bBlockingHit)
+		{
+			return 0.f;
+		}
+
+		float PercentTimeApplied = 0.f;
+		const FVector OldHitNormal = Normal;
+
+		FVector SlideDelta = ComputeSlideVector(Delta, Time, Normal, Hit);
+
+		if ((SlideDelta | Delta) > 0.f)
+		{
+			Driver->SafeMoveUpdatedComponent(SlideDelta, Rotation, true, Hit, ETeleportType::None);
+
+			const float FirstHitPercent = Hit.Time;
+			PercentTimeApplied = FirstHitPercent;
+			if (Hit.IsValidBlockingHit())
+			{
+				// Notify first impact
+				if (bHandleImpact)
+				{
+					// !HandleImpact(Hit, FirstHitPercent * Time, SlideDelta);
+				}
+
+				// Compute new slide normal when hitting multiple surfaces.
+				TwoWallAdjust(SlideDelta, Hit, OldHitNormal);
+
+				// Only proceed if the new direction is of significant length and not in reverse of original attempted move.
+				if (!SlideDelta.IsNearlyZero(1e-3f) && (SlideDelta | Delta) > 0.f)
+				{
+					// Perform second move
+					Driver->SafeMoveUpdatedComponent(SlideDelta, Rotation, true, Hit, ETeleportType::None);
+					const float SecondHitPercent = Hit.Time * (1.f - FirstHitPercent);
+					PercentTimeApplied += SecondHitPercent;
+
+					// Notify second impact
+					if (bHandleImpact && Hit.bBlockingHit)
+					{
+						// !HandleImpact(Hit, SecondHitPercent * Time, SlideDelta);
+					}
+				}
+			}
+
+			return FMath::Clamp(PercentTimeApplied, 0.f, 1.f);
+		}
+
+		return 0.f;
+	}
+
+	void FMovementSystem::Update(IMovementDriver* Driver, const FInputCmd& InputCmd, const FMoveState& InputState, FMoveState& OutputState, const FAuxState& AuxState)
+	{
+		// ----------------------------------
+		// AuxState todo
+		float MaxSpeed = 1200.f;
+		float TurningBoost = 8.f;
+		float Deceleration = 8000.f;
+		float Acceleration = 4000.f;
+
+		// ----------------------------------------------------------------
+
+		const float DeltaSeconds = InputCmd.FrameDeltaTime;
+		OutputState = InputState;
+
+		// --------------------------------------------------------------
+		//	Rotation Update
+		//	We do the rotational update inside the movement sim so that things like server side teleport will work.
+		//	(We want rotation to be treated the same as location, with respect to how its updated, corrected, etc).
+		//	In this simulation, the rotation update isn't allowed to "fail". We don't expect the collision query to be able to fail the rotational update.
+		// --------------------------------------------------------------
+
+		OutputState.Rotation += (InputCmd.RotationInput * InputCmd.FrameDeltaTime);
+		OutputState.Rotation.Normalize();
+
+		const FQuat OutputQuat = OutputState.Rotation.Quaternion();
+
+		// After the rotation is known, we need to sync our driver to this state. This is unfortunate and probably not best for perf, but since the Driver owned
+		// primitive component is ultimately what does our scene queries, we have no choice: we must get the primitive component in the state we say it should be in.
+		Driver->SyncTo(OutputState);
+
+		const FVector LocalSpaceMovementInput = OutputState.Rotation.RotateVector( InputCmd.MovementInput );
+	   	
+		// --------------------------------------------------------------
+		// Calculate OutputState.RelativeVelocity based on Input
+		// --------------------------------------------------------------
+		{
+			const FVector ControlAcceleration = LocalSpaceMovementInput.GetClampedToMaxSize(1.f);
+			FVector Velocity = InputState.Velocity;
+
+			const float AnalogInputModifier = (ControlAcceleration.SizeSquared() > 0.f ? ControlAcceleration.Size() : 0.f);
+			const float MaxPawnSpeed = MaxSpeed * AnalogInputModifier;
+			const bool bExceedingMaxSpeed = IsExceedingMaxSpeed(Velocity, MaxPawnSpeed);
+
+			if (AnalogInputModifier > 0.f && !bExceedingMaxSpeed)
+			{
+				// Apply change in velocity direction
+				if (Velocity.SizeSquared() > 0.f)
+				{
+					// Change direction faster than only using acceleration, but never increase velocity magnitude.
+					const float TimeScale = FMath::Clamp(DeltaSeconds * TurningBoost, 0.f, 1.f);
+					Velocity = Velocity + (ControlAcceleration * Velocity.Size() - Velocity) * TimeScale;
+				}
+			}
+			else
+			{
+				// Dampen velocity magnitude based on deceleration.
+				if (Velocity.SizeSquared() > 0.f)
+				{
+					const FVector OldVelocity = Velocity;
+					const float VelSize = FMath::Max(Velocity.Size() - FMath::Abs(Deceleration) * DeltaSeconds, 0.f);
+					Velocity = Velocity.GetSafeNormal() * VelSize;
+
+					// Don't allow braking to lower us below max speed if we started above it.
+					if (bExceedingMaxSpeed && Velocity.SizeSquared() < FMath::Square(MaxPawnSpeed))
+					{
+						Velocity = OldVelocity.GetSafeNormal() * MaxPawnSpeed;
+					}
+				}
+			}
+
+			// Apply acceleration and clamp velocity magnitude.
+			const float NewMaxSpeed = (IsExceedingMaxSpeed(Velocity, MaxPawnSpeed)) ? Velocity.Size() : MaxPawnSpeed;
+			Velocity += ControlAcceleration * FMath::Abs(Acceleration) * DeltaSeconds;
+			Velocity = Velocity.GetClampedToMaxSize(NewMaxSpeed);
+
+			// Finally, output velocity that we calculated
+			OutputState.Velocity = Velocity;
+		
+			if (ForceMispredict)
+			{
+				OutputState.Velocity += ForceMispredictVelocityMagnitude;
+				ForceMispredict = false;
+			}
+		}
+
+		// ===================================================
+
+		// --------------------------------------------------------------
+		//	Calculate the final movement delta and move the update component
+		// --------------------------------------------------------------
+		FVector Delta = OutputState.Velocity * DeltaSeconds;
+
+		if (!Delta.IsNearlyZero(1e-6f))
+		{
+			FHitResult Hit(1.f);
+			Driver->SafeMoveUpdatedComponent(Delta, OutputQuat, true, Hit, ETeleportType::None);
+
+			if (Hit.IsValidBlockingHit())
+			{
+				// Try to slide the remaining distance along the surface.
+				SlideAlongSurface(Driver, Delta, 1.f-Hit.Time, OutputQuat, Hit.Normal, Hit, true);
+			}
+
+			// Update velocity
+			// We don't want position changes to vastly reverse our direction (which can happen due to penetration fixups etc)
+			/*
+			if (!bPositionCorrected)
+			{
+				const FVector NewLocation = UpdatedComponent->GetComponentLocation();
+				Velocity = ((NewLocation - OldLocation) / DeltaTime);
+			}*/
+		}
+
+		// Finalize. This is unfortunate. The component mirrors our internal motion state and since we call into it to update, at this point, it has the real position.
+		const FTransform UpdateComponentTransform = Driver->GetUpdateComponentTransform();
+		OutputState.Location = UpdateComponentTransform.GetLocation();
+
+		// Note that we don't pull the rotation out of the final update transform. Converting back from a quat will lead to a different FRotator than what we are storing
+		// here in the simulation layer. This may not be the best choice for all movement simulations, but is ok for this one.
+	}
+
+	void FMoveState::VisualLog(const FVisualLoggingParameters& Parameters, IFlyingMovementDriver* Driver, IFlyingMovementDriver* LogDriver)
+	{
+		const bool DrawPointOnly = (Parameters.Context == EVisualLoggingContext::OtherMispredicted || Parameters.Context == EVisualLoggingContext::OtherPredicted);
+		const bool DrawPersistent = (Parameters.Lifetime == EVisualLoggingLifetime::Persistent);
+		const float Lifetime = Parameters.Lifetime == EVisualLoggingLifetime::Persistent ? FlyingPawnSimCVars::DrawDebugDefaultLifeTime : -1.f;
+
+		float Radius=0.f;
+		float HalfHeight=0.f;
+		Driver->GetCapsuleDimensions(Radius, HalfHeight);
+
+		if (FlyingPawnSimCVars::UseDrawDebug)
+		{
+			FVector DrawPosition = Location;
+			const float Thickness = 2.f;
+
+			if (DrawPointOnly)
+			{
+				static float PointSize = 3.f;
+				DrawDebugPoint(LogDriver->GetDriverWorld(), DrawPosition, PointSize, Parameters.GetDebugColor(), DrawPersistent, Lifetime);
+			}
+			else
+			{
+				DrawDebugCapsule(LogDriver->GetDriverWorld(), DrawPosition, HalfHeight, Radius, Rotation.Quaternion(), Parameters.GetDebugColor(), DrawPersistent, Lifetime, 0.f, Thickness);
+			}
+		}
+
+	
+		if (FlyingPawnSimCVars::UseVLogger)
+		{
+			FVector DrawPosition = Location;
+			if (DrawPointOnly)
+			{
+				static float CrumbRadius = 3.f;
+				UE_VLOG_LOCATION(LogDriver->GetVLogOwner(), LogFlyingPawnSimulation, Log, DrawPosition, CrumbRadius, Parameters.GetDebugColor(), TEXT("%d"), Parameters.Keyframe);
+			}
+			else
+			{
+				DrawPosition.Z -= HalfHeight;
+				UE_VLOG_CAPSULE(LogDriver->GetVLogOwner(), LogFlyingPawnSimulation, Log, DrawPosition, HalfHeight, Radius, Rotation.Quaternion(), Parameters.GetDebugColor(), TEXT("%d"), Parameters.Keyframe);
+			}
+
+			UE_VLOG(LogDriver->GetVLogOwner(), LogFlyingPawnSimulation, Log, TEXT("[%d] %s. Location: %s."), Parameters.Keyframe, *LexToString(Parameters.Context), *Location.ToString());
+		}
+	}
+
+} // end namespace
\ No newline at end of file
diff --git a/Engine/Plugins/Runtime/NetworkPrediction/Source/NetworkPrediction/Private/NetworkPredictionComponent.cpp b/Engine/Plugins/Runtime/NetworkPrediction/Source/NetworkPrediction/Private/NetworkPredictionComponent.cpp
new file mode 100644
--- /dev/null
+++ b/Engine/Plugins/Runtime/NetworkPrediction/Source/NetworkPrediction/Private/NetworkPredictionComponent.cpp
@@ -0,0 +1,112 @@
+// Copyright 1998-2019 Epic Games, Inc. All Rights Reserved.
+
+
+#include "NetworkPredictionComponent.h"
+#include "Net/UnrealNetwork.h"
+#include "GameFramework/Actor.h"
+#include "GameFramework/Pawn.h"
+
+UNetworkPredictionComponent::UNetworkPredictionComponent()
+{
+	bReplicates = true;
+}
+
+void UNetworkPredictionComponent::InitializeComponent()
+{
+	Super::InitializeComponent();
+
+	IReplicationProxy* NetworkSim = InstantiateNetworkSimulation();
+	check(NetworkSim);
+
+	ReplicationProxy_ServerRPC.Init(NetworkSim, EReplicationProxyTarget::ServerRPC);
+	ReplicationProxy_Autonomous.Init(NetworkSim, EReplicationProxyTarget::AutonomousProxy);
+	ReplicationProxy_Simulated.Init(NetworkSim, EReplicationProxyTarget::SimulatedProxy);
+	ReplicationProxy_Replay.Init(NetworkSim, EReplicationProxyTarget::Replay);
+#if NETSIM_MODEL_DEBUG
+	ReplicationProxy_Debug.Init(NetworkSim, EReplicationProxyTarget::Debug);
+#endif
+}
+
+void UNetworkPredictionComponent::PreReplication(IRepChangedPropertyTracker & ChangedPropertyTracker)
+{
+	Super::PreReplication(ChangedPropertyTracker);
+
+	// We have to update our replication proxies so they can be accurately compared against client shadowstate during property replication. ServerRPC proxy does not need to do this.
+	ReplicationProxy_Autonomous.OnPreReplication();
+	ReplicationProxy_Simulated.OnPreReplication();
+	ReplicationProxy_Replay.OnPreReplication();
+#if NETSIM_MODEL_DEBUG
+	ReplicationProxy_Debug.OnPreReplication();
+#endif
+}
+
+void UNetworkPredictionComponent::PreNetReceive()
+{
+	Super::PreNetReceive();
+	CheckOwnerRoleChange();
+}
+
+void UNetworkPredictionComponent::GetLifetimeReplicatedProps(TArray< FLifetimeProperty > & OutLifetimeProps) const
+{
+	Super::GetLifetimeReplicatedProps(OutLifetimeProps);
+
+	DOREPLIFETIME_CONDITION( UNetworkPredictionComponent, ReplicationProxy_Autonomous, COND_AutonomousOnly);
+	DOREPLIFETIME_CONDITION( UNetworkPredictionComponent, ReplicationProxy_Simulated, COND_SimulatedOnly);
+	DOREPLIFETIME_CONDITION( UNetworkPredictionComponent, ReplicationProxy_Replay, COND_ReplayOnly);
+
+#if NETSIM_MODEL_DEBUG
+	DOREPLIFETIME_CONDITION( UNetworkPredictionComponent, ReplicationProxy_Debug, COND_ReplayOrOwner);
+#endif
+}
+
+bool UNetworkPredictionComponent::IsLocallyControlled()
+{
+	// This awkward because, engine wide, "Is Locally Controlled" is really only a built in concept to Pawns.
+	// We don't want the UNetworkPredictionComponent to be explicitly coupled to APawn though.
+	// Role isn't enough on its own because standalone/listen servers means Role=ROLE_Authority can be locally controlled.
+	// This seems like the best compromise: things will "just work" if the component is a child of an APawn.
+	// If you want to have some concept of locally controlled network sims that are not childs of APawn actors, you can always
+	// override this function.
+
+	bool bIsLocallyControlled = false;
+	if (GetOwner()->Role == ROLE_AutonomousProxy)
+	{
+		bIsLocallyControlled = true;
+	}
+	else if (APawn* PawnOwner = Cast<APawn>(GetOwner()))
+	{
+		bIsLocallyControlled = PawnOwner->IsLocallyControlled();
+	}
+
+	return bIsLocallyControlled;
+}
+
+void UNetworkPredictionComponent::TickComponent(float DeltaTime, enum ELevelTick TickType, FActorComponentTickFunction *ThisTickFunction)
+{
+	CheckOwnerRoleChange();
+	if (PreTickLocallyControlledSim.IsBound() && IsLocallyControlled())
+	{
+		PreTickLocallyControlledSim.Execute(DeltaTime);
+	}
+}
+
+void UNetworkPredictionComponent::CheckOwnerRoleChange()
+{
+	ENetRole CurrentRole = GetOwner()->Role;
+	if (CurrentRole != OwnerCachedNetRole)
+	{
+		OwnerCachedNetRole = CurrentRole;
+		InitializeForNetworkRole(CurrentRole);		
+	}
+}
+
+bool UNetworkPredictionComponent::ServerRecieveClientInput_Validate(const FServerReplicationRPCParameter& ProxyParameter)
+{
+	return true;
+}
+
+void UNetworkPredictionComponent::ServerRecieveClientInput_Implementation(const FServerReplicationRPCParameter& ProxyParameter)
+{
+	// The const_cast is unavoidable here because the replication system only allows by value (forces copy, bad) or by const reference. This use case is unique because we are using the RPC parameter as a temp buffer.
+	const_cast<FServerReplicationRPCParameter&>(ProxyParameter).NetSerializeToProxy(ReplicationProxy_ServerRPC);
+}
\ No newline at end of file
diff --git a/Engine/Plugins/Runtime/NetworkPrediction/Source/NetworkPrediction/Private/NetworkPredictionModule.cpp b/Engine/Plugins/Runtime/NetworkPrediction/Source/NetworkPrediction/Private/NetworkPredictionModule.cpp
new file mode 100644
--- /dev/null
+++ b/Engine/Plugins/Runtime/NetworkPrediction/Source/NetworkPrediction/Private/NetworkPredictionModule.cpp
@@ -0,0 +1,41 @@
+// Copyright 1998-2019 Epic Games, Inc. All Rights Reserved.
+
+#include "NetworkPredictionModule.h"
+#include "CoreMinimal.h"
+#include "Modules/ModuleManager.h"
+#include "Internationalization/Internationalization.h"
+#include "Engine/World.h"
+
+#define LOCTEXT_NAMESPACE "FNetworkPredictionModule"
+
+class FNetworkPredictionModule : public INetworkPredictionModule
+{
+	/** IModuleInterface implementation */
+	virtual void StartupModule() override;
+	virtual void ShutdownModule() override;
+};
+
+void FNetworkPredictionModule::StartupModule()
+{
+/*
+	FWorldDelegates::OnPostWorldInitialization.AddLambda([](UWorld* World, const UWorld::InitializationValues IVS)
+	{
+		World->AddOnActorSpawnedHandler(FOnActorSpawned::FDelegate::CreateLambda([](AActor* Actor)
+		{
+			if (APawn* Pawn = Cast<APawn>(Actor))
+			{
+
+			}
+		}));
+	});
+*/
+}
+
+
+void FNetworkPredictionModule::ShutdownModule()
+{
+}
+
+IMPLEMENT_MODULE( FNetworkPredictionModule, NetworkPrediction )
+#undef LOCTEXT_NAMESPACE
+
diff --git a/Engine/Plugins/Runtime/NetworkPrediction/Source/NetworkPrediction/Private/NetworkSimulationModel.cpp b/Engine/Plugins/Runtime/NetworkPrediction/Source/NetworkPrediction/Private/NetworkSimulationModel.cpp
new file mode 100644
--- /dev/null
+++ b/Engine/Plugins/Runtime/NetworkPrediction/Source/NetworkPrediction/Private/NetworkSimulationModel.cpp
@@ -0,0 +1,114 @@
+// Copyright 1998-2019 Epic Games, Inc. All Rights Reserved.
+
+#include "NetworkSimulationModel.h"
+#include "Engine/NetConnection.h"
+
+DEFINE_LOG_CATEGORY(LogNetworkSim);
+
+FColor FVisualLoggingParameters::DebugColors[(int32)EVisualLoggingContext::MAX] =
+{
+	FColor::Cyan,		// LastPredicted
+	FColor::Cyan,		// OtherPredicted
+	FColor::Green,		// LastConfirmed
+	FColor::Red,		// FirstMispredicted
+	FColor(255,81,101),	// OtherMispredicted
+	FColor(255,91,91),	// LastMispredicted
+	FColor::Purple,		// CurrentServerPIE
+};
+
+// -------------------------------------------------------------------------------------------------------------------------------
+//	FReplicationProxy
+// -------------------------------------------------------------------------------------------------------------------------------
+
+void FReplicationProxy::Init(class IReplicationProxy* InNetworkSimModel, EReplicationProxyTarget InReplicationTarget)
+{
+	NetworkSimModel = InNetworkSimModel;
+	ReplicationTarget = InReplicationTarget;
+}
+
+bool FReplicationProxy::NetSerialize(FArchive& Ar, class UPackageMap* Map, bool& bOutSuccess)
+{
+	NetworkSimModel->NetSerializeProxy(ReplicationTarget, Ar);
+	return true;
+}
+
+void FReplicationProxy::OnPreReplication()
+{
+	CachedDirtyCount = NetworkSimModel->GetProxyDirtyCount(ReplicationTarget);
+}
+	
+bool FReplicationProxy::Identical(const FReplicationProxy* Other, uint32 PortFlags) const
+{
+	return (CachedDirtyCount == Other->CachedDirtyCount);
+}
+
+// -------------------------------------------------------------------------------------------------------------------------------
+//	FServerRPCProxyParameter
+// -------------------------------------------------------------------------------------------------------------------------------
+
+TArray<uint8> FServerReplicationRPCParameter::TempStorage;
+
+bool FServerReplicationRPCParameter::NetSerialize(FArchive& Ar, class UPackageMap* Map, bool& bOutSuccess)
+{
+	if (Ar.IsLoading())
+	{
+		// Loading: serialize to temp storage. We'll do the real deserialize in a manual call to ::NetSerializeToProxy
+		FNetBitReader& BitReader = (FNetBitReader&)Ar;
+		CachedNumBits = BitReader.GetBitsLeft();
+		CachedPackageMap = Map;
+
+		const int64 BytesLeft = BitReader.GetBytesLeft();
+		check(BytesLeft > 0); // Should not possibly be able to get here with an empty archive
+		TempStorage.Reset(BytesLeft);
+		TempStorage.SetNumUninitialized(BytesLeft);
+		TempStorage.Last() = 0;
+
+		BitReader.SerializeBits(TempStorage.GetData(), CachedNumBits);
+	}
+	else
+	{
+		// Saving: directly call into the proxy's NetSerialize. No need for temp storage.
+		check(Proxy); // Must have been set before, via ctor.
+		return Proxy->NetSerialize(Ar, Map, bOutSuccess);
+	}
+
+	return true;
+}
+
+void FServerReplicationRPCParameter::NetSerializeToProxy(FReplicationProxy& InProxy)
+{
+	check(CachedPackageMap != nullptr);
+	check(CachedNumBits != -1);
+
+	FNetBitReader BitReader(CachedPackageMap, TempStorage.GetData(), CachedNumBits);
+
+	bool bOutSuccess = true;
+	InProxy.NetSerialize(BitReader, CachedPackageMap, bOutSuccess);
+
+	CachedNumBits = -1;
+	CachedPackageMap = nullptr;
+}
+
+// -------------------------------------------------------------------------------------------------------------------------------
+//	FScopedBandwidthLimitBypass
+// -------------------------------------------------------------------------------------------------------------------------------
+
+FScopedBandwidthLimitBypass::FScopedBandwidthLimitBypass(AActor* OwnerActor)
+{
+	if (OwnerActor)
+	{
+		CachedNetConnection = OwnerActor->GetNetConnection();
+		if (CachedNetConnection)
+		{
+			RestoreBits = CachedNetConnection->QueuedBits + CachedNetConnection->SendBuffer.GetNumBits();
+		}
+	}
+}
+
+FScopedBandwidthLimitBypass::~FScopedBandwidthLimitBypass()
+{
+	if (CachedNetConnection)
+	{
+		CachedNetConnection->QueuedBits = RestoreBits - CachedNetConnection->SendBuffer.GetNumBits();
+	}
+}
\ No newline at end of file
diff --git a/Engine/Plugins/Runtime/NetworkPrediction/Source/NetworkPrediction/Private/NetworkSimulationModelDebugger.cpp b/Engine/Plugins/Runtime/NetworkPrediction/Source/NetworkPrediction/Private/NetworkSimulationModelDebugger.cpp
new file mode 100644
--- /dev/null
+++ b/Engine/Plugins/Runtime/NetworkPrediction/Source/NetworkPrediction/Private/NetworkSimulationModelDebugger.cpp
@@ -0,0 +1,104 @@
+// Copyright 1998-2019 Epic Games, Inc. All Rights Reserved.
+
+#include "NetworkSimulationModelDebugger.h"
+#include "NetworkSimulationModel.h"
+#include "Engine/World.h"
+#include "Engine/LocalPlayer.h"
+#include "GameFramework/Pawn.h"
+#include "GameFramework/PlayerController.h"
+#include "Misc/NetworkGuid.h"
+#include "Engine/NetConnection.h"
+#include "Engine/NetDriver.h"
+
+FNetworkSimulationModelDebuggerManager& FNetworkSimulationModelDebuggerManager::Get()
+{
+	static FNetworkSimulationModelDebuggerManager Manager;
+	return Manager;
+}
+
+UObject* FindReplicatedObjectOnPIEServer(UObject* ClientObject)
+{
+	if (ClientObject == nullptr)
+	{
+		return nullptr;
+	}
+
+	UObject* ServerObject = nullptr;
+
+#if WITH_EDITOR
+	if (UWorld* World = ClientObject->GetWorld())
+	{
+		if (UNetDriver* NetDriver = World->GetNetDriver())
+		{
+			if (UNetConnection* NetConnection = NetDriver->ServerConnection)
+			{
+				FNetworkGUID NetGUID = NetConnection->PackageMap->GetNetGUIDFromObject(ClientObject);
+
+				// Find the PIE server world
+				for (TObjectIterator<UWorld> It; It; ++It)
+				{
+					if (It->WorldType == EWorldType::PIE && It->GetNetMode() != NM_Client)
+					{
+						if (UNetDriver* ServerNetDriver = It->GetNetDriver())
+						{
+							if (ServerNetDriver->ClientConnections.Num() > 0)
+							{
+								UPackageMap* ServerPackageMap = ServerNetDriver->ClientConnections[0]->PackageMap;
+								ServerObject = ServerPackageMap->GetObjectFromNetGUID(NetGUID, true);
+								break;
+							}
+						}
+					}
+				}
+			}
+		}
+	}
+#endif
+
+	return ServerObject;
+}
+
+// ------------------------------------------------------------------------------------------------------------------------
+//
+//
+// ------------------------------------------------------------------------------------------------------------------------
+
+FAutoConsoleCommandWithWorldAndArgs NetworkSimulationModelDebugCmd(TEXT("nms.Debug.LocallyControlledPawn"), TEXT(""),
+	FConsoleCommandWithWorldAndArgsDelegate::CreateLambda([](const TArray< FString >& InArgs, UWorld* World) 
+{
+
+	if (!World || !World->GetFirstLocalPlayerFromController())
+	{
+		return;
+	}
+	ULocalPlayer* Player = World->GetFirstLocalPlayerFromController();
+	if (!Player || !Player->GetPlayerController(World) || !Player->GetPlayerController(World)->GetPawn())
+	{
+		UE_LOG(LogNetworkSimDebug, Error, TEXT("Could not find valid locally controlled pawn."));
+		return;
+	}
+
+	APawn* Pawn = Player->GetPlayerController(World)->GetPawn();
+	FNetworkSimulationModelDebuggerManager::Get().ToggleDebuggerActive(Pawn);
+}));
+
+
+FAutoConsoleCommandWithWorldAndArgs NetworkSimulationModelDebugToggleContinousCmd(TEXT("nms.Debug.ToggleContinous"), TEXT(""),
+	FConsoleCommandWithWorldAndArgsDelegate::CreateLambda([](const TArray< FString >& InArgs, UWorld* World) 
+{
+
+	if (!World || !World->GetFirstLocalPlayerFromController())
+	{
+		return;
+	}
+	ULocalPlayer* Player = World->GetFirstLocalPlayerFromController();
+	if (!Player || !Player->GetPlayerController(World) || !Player->GetPlayerController(World)->GetPawn())
+	{
+		UE_LOG(LogNetworkSimDebug, Error, TEXT("Could not find valid locally controlled pawn. "));
+		return;
+	}
+
+	APawn* Pawn = Player->GetPlayerController(World)->GetPawn();
+	FNetworkSimulationModelDebuggerManager::Get().SetDebuggerActive(Pawn, true);
+	FNetworkSimulationModelDebuggerManager::Get().ToggleContinousGather();
+}));
\ No newline at end of file
diff --git a/Engine/Plugins/Runtime/NetworkPrediction/Source/NetworkPrediction/Public/Movement/FlyingMovement.h b/Engine/Plugins/Runtime/NetworkPrediction/Source/NetworkPrediction/Public/Movement/FlyingMovement.h
new file mode 100644
--- /dev/null
+++ b/Engine/Plugins/Runtime/NetworkPrediction/Source/NetworkPrediction/Public/Movement/FlyingMovement.h
@@ -0,0 +1,219 @@
+// Copyright 1998-2019 Epic Games, Inc. All Rights Reserved.
+
+#pragma once
+
+#include "CoreMinimal.h"
+#include "UObject/ObjectMacros.h"
+#include "UObject/UObjectGlobals.h"
+#include "Engine/NetSerialization.h"
+#include "Engine/EngineTypes.h"
+#include "Engine/EngineBaseTypes.h"
+#include "WorldCollision.h"
+#include "Components/ActorComponent.h"
+#include "Components/SceneComponent.h"
+#include "Misc/OutputDevice.h"
+#include "Misc/CoreDelegates.h"
+#include "NetworkSimulationModelTemplates.h"
+#include "NetworkPredictionComponent.h"
+
+#include "FlyingMovement.generated.h"
+
+// -------------------------------------------------------------------------------------------------------------------------------
+// FlyingMovement: simple flying movement that was based on UE4's FloatingPawnMovement
+// -------------------------------------------------------------------------------------------------------------------------------
+
+class IFlyingMovementDriver;
+class USceneComponent;
+
+namespace FlyingMovement
+{
+	// State the client generates
+	struct FInputCmd
+	{
+
+		// Client's FrameTime. This is essential for a non-fixed step simulation. The server will run the movement simulation for this client at this rate.
+		float FrameDeltaTime = 0.f;
+
+		// Input: "pure" input for this frame. At this level, frame time has not been accounted for. (E.g., "move straight" would be (1,0,0) regardless of frame time)
+		FRotator RotationInput;
+		FVector MovementInput;
+
+		void NetSerialize(const FNetSerializeParams& P)
+		{
+			P.Ar << FrameDeltaTime;
+			P.Ar << RotationInput;
+			P.Ar << MovementInput;
+		}
+
+		void Log(FStandardLoggingParameters& P) const
+		{
+			if (P.Context == EStandardLoggingContext::HeaderOnly)
+			{
+				P.Ar->Logf(TEXT(" %d "), P.Keyframe);
+			}
+			else if (P.Context == EStandardLoggingContext::Full)
+			{
+				P.Ar->Logf(TEXT("FrameDeltaTime: %.4f"), FrameDeltaTime);
+				P.Ar->Logf(TEXT("Movement: %s"), *MovementInput.ToString());
+				P.Ar->Logf(TEXT("ViewRot: %s"), *RotationInput.ToString());
+			}
+		}
+	};
+
+	// State we are evolving frame to frame and keeping in sync
+	struct FMoveState
+	{
+		FVector Location;
+		FVector Velocity;
+		FRotator Rotation;
+
+		void NetSerialize(const FNetSerializeParams& P)
+		{
+			P.Ar << Location;
+			P.Ar << Velocity;
+			P.Ar << Rotation;
+		}
+
+		// Compare this state with AuthorityState. return true if a reconcile (correction) should happen
+		bool ShouldReconcile(const FMoveState& AuthorityState)
+		{
+			const float ErrorTolerance = 1.f;
+			return !AuthorityState.Location.Equals(Location, ErrorTolerance);
+		}
+
+		void Log(FStandardLoggingParameters& Params) const
+		{
+			if (Params.Context == EStandardLoggingContext::HeaderOnly)
+			{
+				Params.Ar->Logf(TEXT(" %d "), Params.Keyframe);
+			}
+			else if (Params.Context == EStandardLoggingContext::Full)
+			{
+				Params.Ar->Logf(TEXT("Frame: %d"), Params.Keyframe);
+				Params.Ar->Logf(TEXT("Loc: %s"), *Location.ToString());
+				Params.Ar->Logf(TEXT("Vel: %s"), *Velocity.ToString());
+				Params.Ar->Logf(TEXT("Rot: %s"), *Rotation.ToString());
+			}
+		}
+
+		void VisualLog(const FVisualLoggingParameters& Parameters, IFlyingMovementDriver* Driver, IFlyingMovementDriver* LogDriver);
+	};
+
+	// Auxiliary state that is input into the simulation. Doesn't change during the simulation tick.
+	// (It can change and even be predicted but doing so will trigger more bookeeping, etc. Changes will happen "next tick")
+	struct FAuxState
+	{
+		float Multiplier=1;
+		void NetSerialize(const FNetSerializeParams& P)
+		{
+			P.Ar << Multiplier;
+		}
+	};
+
+	// Actual definition of our network simulation.
+	class FMovementSystem : public TNetworkedSimulationModel<FMovementSystem, FInputCmd, FMoveState, FAuxState>
+	{
+	public:
+
+		// Interface between the simulation and owning component driving it. Functions added here are available in ::Update and must be implemented by UMockNetworkSimulationComponent.
+		class IMovementDriver : public IDriver
+		{
+		public:
+
+			virtual bool SafeMoveUpdatedComponent(const FVector& Delta, const FQuat& NewRotation, bool bSweep, FHitResult& OutHit, ETeleportType Teleport) const = 0;
+			virtual bool MoveUpdatedComponent(const FVector& Delta, const FQuat& NewRotation, bool bSweep, FHitResult* OutHit, ETeleportType Teleport) const = 0;
+			virtual FTransform GetUpdateComponentTransform() const = 0;
+
+
+			virtual void GetCapsuleDimensions(float &Radius, float& HalfHeight) const = 0;
+			virtual UWorld* GetDriverWorld() const = 0;
+			virtual UObject* GetVLogOwner() const = 0;
+			virtual FTransform GetDebugWorldTransform() const = 0;
+		};
+
+		/** Main update function */
+		static void Update(IMovementDriver* Driver, const FInputCmd& InputCmd, const FMoveState& InputState, FMoveState& OutputState, const FAuxState& AuxState);
+
+		/** Dev tool to force simple mispredict */
+		static bool ForceMispredict;
+	};
+
+	// general tolerance value for rotation checks
+	static const float ROTATOR_TOLERANCE = (1e-3);
+
+} // End namespace
+
+// Needed to trick UHT into letting UMockNetworkSimulationComponent implement. UHT cannot parse the ::
+// Also needed for forward declaring. Can't just be a typedef/using =
+class IFlyingMovementDriver : public FlyingMovement::FMovementSystem::IMovementDriver { };
+
+// -------------------------------------------------------------------------------------------------------------------------------
+// ActorComponent for running FlyingMovement 
+// -------------------------------------------------------------------------------------------------------------------------------
+
+UCLASS(BlueprintType, meta=(BlueprintSpawnableComponent))
+class NETWORKPREDICTION_API UFlyingMovementComponent : public UNetworkPredictionComponent, public IFlyingMovementDriver
+{
+	GENERATED_BODY()
+
+public:
+
+	UFlyingMovementComponent();
+
+	virtual void InitializeComponent() override;
+	virtual void OnRegister() override;
+	virtual void RegisterComponentTickFunctions(bool bRegister) override;
+	virtual void TickComponent(float DeltaTime, enum ELevelTick TickType, FActorComponentTickFunction *ThisTickFunction) override;
+
+	// Gets writable copy of FClientInputCmd for this frame. Should only be called once per frame by the local player.
+	FlyingMovement::FInputCmd* GetNextClientInputCmdForWrite(float DeltaTime);
+
+	// IFlyingMovementDriver
+	bool SafeMoveUpdatedComponent(const FVector& Delta, const FQuat& NewRotation, bool bSweep, FHitResult& OutHit, ETeleportType Teleport) const;
+	bool MoveUpdatedComponent(const FVector& Delta, const FQuat& NewRotation, bool bSweep, FHitResult* OutHit, ETeleportType Teleport) const;
+	FTransform GetUpdateComponentTransform() const;
+
+	void InitSyncState(FlyingMovement::FMoveState& OutSyncState) const override;
+	void SyncTo(const FlyingMovement::FMoveState& SyncState) override;
+	virtual UWorld* GetDriverWorld() const override final { return GetWorld(); }
+	virtual UObject* GetVLogOwner() const override final;
+	virtual FTransform GetDebugWorldTransform() const override final;
+	virtual void GetCapsuleDimensions(float &Radius, float& HalfHeight) const;
+
+protected:
+
+	// Basic "Update Component/Ticking"
+	virtual void SetUpdatedComponent(USceneComponent* NewUpdatedComponent);
+	virtual void UpdateTickRegistration();
+
+	UFUNCTION()
+	virtual void PhysicsVolumeChanged(class APhysicsVolume* NewVolume);
+
+	// Network Prediction
+	virtual IReplicationProxy* InstantiateNetworkSimulation() override;	
+	virtual void InitializeForNetworkRole(ENetRole Role) override;
+	TUniquePtr<FlyingMovement::FMovementSystem> NetworkSim; // The Network sim that this component is managing. This is what is doing all the work.
+
+private:
+
+	UPROPERTY()
+	USceneComponent* UpdatedComponent = nullptr;
+
+	UPROPERTY()
+	UPrimitiveComponent* UpdatedPrimitive = nullptr;
+
+	/** Transient flag indicating whether we are executing OnRegister(). */
+	bool bInOnRegister = false;
+	
+	/** Transient flag indicating whether we are executing InitializeComponent(). */
+	bool bInInitializeComponent = false;
+
+	static FVector GetPenetrationAdjustment(const FHitResult& Hit);
+
+	bool OverlapTest(const FVector& Location, const FQuat& RotationQuat, const ECollisionChannel CollisionChannel, const FCollisionShape& CollisionShape, const AActor* IgnoreActor) const;
+	void InitCollisionParams(FCollisionQueryParams &OutParams, FCollisionResponseParams& OutResponseParam) const;
+	bool ResolvePenetration(const FVector& ProposedAdjustment, const FHitResult& Hit, const FQuat& NewRotationQuat) const;
+
+	/**  Flags that control the behavior of calls to MoveComponent() on our UpdatedComponent. */
+	mutable EMoveComponentFlags MoveComponentFlags = MOVECOMP_NoFlags; // Mutable because we sometimes need to disable these flags ::ResolvePenetration. Better way may be possible
+};
\ No newline at end of file
diff --git a/Engine/Plugins/Runtime/NetworkPrediction/Source/NetworkPrediction/Public/NetworkPredictionComponent.h b/Engine/Plugins/Runtime/NetworkPrediction/Source/NetworkPrediction/Public/NetworkPredictionComponent.h
new file mode 100644
--- /dev/null
+++ b/Engine/Plugins/Runtime/NetworkPrediction/Source/NetworkPrediction/Public/NetworkPredictionComponent.h
@@ -0,0 +1,79 @@
+// Copyright 1998-2019 Epic Games, Inc. All Rights Reserved.
+
+#pragma once
+#include "Components/ActorComponent.h"
+#include "NetworkSimulationModel.h"
+
+#include "NetworkPredictionComponent.generated.h"
+
+// ----------------------------------------------------------------------------------------------------------------------------------------------------------------------
+//	UNetworkPredictionComponent
+//	This is the base component for running a TNetworkedSimulationModel through an actor component. This contains the boiler plate hooks into getting the system
+//	initialized and plugged into the UE4 replication system.
+//
+//	This is an abstract component and cannot function on its own. It must be subclassed and InstantiateNetworkSimulation must be implemented.
+//	The subclass is responsible for running the network simulation. This at a bare minimum would mean ticking the simulation and supplying it with input.
+// ----------------------------------------------------------------------------------------------------------------------------------------------------------------------
+UCLASS(Abstract)
+class NETWORKPREDICTION_API UNetworkPredictionComponent : public UActorComponent
+{
+	GENERATED_BODY()
+
+public:
+
+	UNetworkPredictionComponent();
+
+	virtual void InitializeComponent() override;
+	virtual void PreReplication(IRepChangedPropertyTracker & ChangedPropertyTracker) override;
+	virtual void PreNetReceive() override;
+
+	// TickComponent calls CheckOwnerRoleChange and invokes PreTickLocallyControlledSim. You should always call it at the top of your child's tick implementation.
+	virtual void TickComponent(float DeltaTime, enum ELevelTick TickType, FActorComponentTickFunction *ThisTickFunction) override;
+
+	// PreSimulationTick callback.
+	// Called right in ::TickComponent for autonomous proxies. This is optional and intended to be a convencience for generating local input.
+	// If you don't use this function, you should manually be submitting client input prior to simulation tick (via GetNextClientInputCmdForWrite).
+	// (It is important to submit input prior to ticking the sim. If you get it backwards, you introduce local latency and cause problems with variable time steps).
+	DECLARE_DELEGATE_OneParam(FPreSimTickDelegate, float /* DeltaSeconds */)
+	void SetLocallyControlledPreTick(const FPreSimTickDelegate& PreSimTick) { PreTickLocallyControlledSim = PreSimTick; }
+
+protected:
+
+	// Child classes must allocate and manage lifetime of their own NetworkSim (E.g, TNetworkedSimulationModel). Recommend to just store in a TUniquePtr.
+	// Child classes must also register with Network Sim debugger here (FNetworkSimulationModelDebuggerManager RegisterNetworkSimulationModel)
+	virtual IReplicationProxy* InstantiateNetworkSimulation() { return nullptr; }
+
+	// Child classes should override this an initialize their NetworkSim here
+	virtual void InitializeForNetworkRole(ENetRole Role) { }
+
+	// Attempts to determine if the component is locally controlled or not, meaning we expect input cmds to be generated locally (and hence need to prepare the input buffer for that)
+	virtual bool IsLocallyControlled();
+
+	// Helper: Checks if the owner's role has changed and calls InitializeForNetworkRole if necessary.
+	// NOTE: You may need to call this in your TickComponent function
+	void CheckOwnerRoleChange();
+	ENetRole OwnerCachedNetRole = ROLE_None;
+
+	UFUNCTION(Server, unreliable, WithValidation)
+	void ServerRecieveClientInput(const FServerReplicationRPCParameter& ProxyParameter);
+
+	// ReplicationProxies are just pointers to the data/NetSerialize functions within the NetworkSim
+	UPROPERTY()
+	FReplicationProxy ReplicationProxy_ServerRPC;
+
+private:
+
+	FPreSimTickDelegate PreTickLocallyControlledSim;
+
+	UPROPERTY(Replicated, transient)
+	FReplicationProxy ReplicationProxy_Autonomous;
+
+	UPROPERTY(Replicated, transient)
+	FReplicationProxy ReplicationProxy_Simulated;
+
+	UPROPERTY(Replicated, transient)
+	FReplicationProxy ReplicationProxy_Replay;
+
+	UPROPERTY(Replicated, transient)
+	FReplicationProxy ReplicationProxy_Debug;
+};
\ No newline at end of file
diff --git a/Engine/Plugins/Runtime/NetworkPrediction/Source/NetworkPrediction/Public/NetworkPredictionModule.h b/Engine/Plugins/Runtime/NetworkPrediction/Source/NetworkPrediction/Public/NetworkPredictionModule.h
new file mode 100644
--- /dev/null
+++ b/Engine/Plugins/Runtime/NetworkPrediction/Source/NetworkPrediction/Public/NetworkPredictionModule.h
@@ -0,0 +1,55 @@
+// Copyright 1998-2018 Epic Games, Inc. All Rights Reserved.
+
+#pragma once
+
+#include "CoreMinimal.h"
+#include "Modules/ModuleInterface.h"
+#include "Modules/ModuleManager.h"
+
+/*=============================================================================
+
+Overview of the NetworkPredictionModule and glossary of high level terms:
+
+
+
+
+Replication Proxy:
+
+Actor Component:
+
+
+
+
+
+
+
+=============================================================================*/
+
+
+class INetworkPredictionModule : public IModuleInterface
+{
+
+public:
+
+	/**
+	 * Singleton-like access to this module's interface.  This is just for convenience!
+	 * Beware of calling this during the shutdown phase, though.  Your module might have been unloaded already.
+	 *
+	 * @return Returns singleton instance, loading the module on demand if needed
+	 */
+	static inline INetworkPredictionModule& Get()
+	{
+		return FModuleManager::LoadModuleChecked< INetworkPredictionModule >( "NetworkPrediction" );
+	}
+
+	/**
+	 * Checks to see if this module is loaded and ready.  It is only valid to call Get() if IsAvailable() returns true.
+	 *
+	 * @return True if the module is loaded and ready to use
+	 */
+	static inline bool IsAvailable()
+	{
+		return FModuleManager::Get().IsModuleLoaded( "NetworkPrediction" );
+	}
+};
+
diff --git a/Engine/Plugins/Runtime/NetworkPrediction/Source/NetworkPrediction/Public/NetworkSimulationModel.h b/Engine/Plugins/Runtime/NetworkPrediction/Source/NetworkPrediction/Public/NetworkSimulationModel.h
new file mode 100644
--- /dev/null
+++ b/Engine/Plugins/Runtime/NetworkPrediction/Source/NetworkPrediction/Public/NetworkSimulationModel.h
@@ -0,0 +1,213 @@
+// Copyright 1998-2019 Epic Games, Inc. All Rights Reserved.
+
+#pragma once
+
+#include "Engine/EngineTypes.h"
+#include "Net/UnrealNetwork.h" // For MakeRelative
+#include "NetworkSimulationModel.generated.h"
+
+// -------------------------------------------------------------------------------------------------------------------------------
+//	This file contains supporting types that are used by both the UE4 side (UNetworkPredictionComponent etc) and templated
+//	network simulation model side (TNetworkedSimulationModel).
+// -------------------------------------------------------------------------------------------------------------------------------
+
+#define NETSIM_MODEL_DEBUG !(UE_BUILD_SHIPPING || UE_BUILD_TEST)
+
+NETWORKPREDICTION_API DECLARE_LOG_CATEGORY_EXTERN(LogNetworkSim, Log, All);
+
+UENUM()
+enum class EReplicationProxyTarget: uint8
+{
+	ServerRPC,			// Client -> Server
+	AutonomousProxy,	// Owning/Controlling client
+	SimulatedProxy,		// Non owning client
+	Replay,				// Replay net driver
+	Debug,				// Replication target that is disabled in shipping
+};
+
+inline FString LexToString(EReplicationProxyTarget A)
+{
+	return *UEnum::GetValueAsString(TEXT("NetworkPrediction.EReplicationProxyTarget"), A);
+}
+
+// The parameters for NetSerialize that are passed around the system. Everything should use this, expecting to have to add more.
+struct FNetSerializeParams
+{
+	FNetSerializeParams(FArchive& InAr) : Ar(InAr) { }
+	FArchive& Ar;
+};
+
+struct FNetworkSimulationModelInitParameters
+{
+	int32 InputBufferSize=0;
+	int32 SyncedBufferSize=0;
+	int32 AuxBufferSize=0;
+	int32 DebugBufferSize=0;
+	int32 HistoricBufferSize=0;
+};
+
+struct FNetSimProcessedFrameDebugInfo
+{
+	int32 LocalGFrameNumber; // Local GFrame number
+	float LocalDeltaTimeSeconds; // Local frame time
+	TArray<int32> ProcessedKeyframes; // Which keyframes were processed this frame
+	int32 LastProcessedKeyframe; // What LastProcessedKeyframe was at the end of the frame. Does NOT mean we processed it this frame!
+	int32 HeadKeyframe; // Head keyframe of the inputbuffer when the frame ended
+
+	void NetSerialize(const FNetSerializeParams& P)
+	{	
+		P.Ar << LocalGFrameNumber;
+		P.Ar << LocalDeltaTimeSeconds;
+		P.Ar << LastProcessedKeyframe;
+		P.Ar << HeadKeyframe;
+
+		SafeNetSerializeTArray_Default<31>(P.Ar, ProcessedKeyframes);
+	}
+};
+
+enum class EStandardLoggingContext: uint8
+{
+	HeaderOnly,		// Minimal logging
+	Full			// "Logs everything"
+};
+
+struct FStandardLoggingParameters
+{
+	FStandardLoggingParameters(FOutputDevice* InAr, EStandardLoggingContext InContext, int32 InKeyframe)
+		: Ar(InAr), Context(InContext), Keyframe(InKeyframe) { }
+	FOutputDevice* Ar;
+	EStandardLoggingContext Context;
+	int32 Keyframe;
+};
+
+UENUM()
+enum class EVisualLoggingContext: uint8
+{
+	LastPredicted,		// The last predicted state. Where the character "is now"
+	OtherPredicted,		// "Middle" states between LastComfirmed and LastPredicted. Recommend to draw small crumbs (not full model/cylinder/etc) to avoid washing out the scene.
+	LastConfirmed,		// The last confirmed state from the server
+	FirstMispredicted,	// The first state that was mispredicted on the client. This correlates with LastConfirmed.
+	OtherMispredicted,	// subsequent mispredicted states on the client. Recommend to draw small crumbs (not full model/cylinder/etc) to avoid washing out the scene.
+	LastMispredicted,	// The latest predictive state when a misprediction was detected.
+	CurrentServerPIE,	// The current server position *right now*. Only available in PIE
+	MAX
+};
+
+inline FString LexToString(EVisualLoggingContext A)
+{
+	return *UEnum::GetValueAsString(TEXT("NetworkPrediction.EVisualLoggingContext"), A);
+}
+
+enum class EVisualLoggingLifetime : uint8
+{
+	Transient,			// This logging is transient and will (probably) be done every frame. Don't persist (in contexts where that would matter).
+	Persistent,			// This is more of a persistent/one-off event that should be drawn for some longer amount of time (probably configurable at whatever level is doing the logging)
+};
+
+struct FVisualLoggingParameters
+{
+	FVisualLoggingParameters(EVisualLoggingContext InContext, int32 InKeyframe, EVisualLoggingLifetime InLifetime) :
+		 Context(InContext), Keyframe(InKeyframe), Lifetime(InLifetime) { }
+
+	EVisualLoggingContext Context;
+	int32 Keyframe;
+	EVisualLoggingLifetime Lifetime;
+
+	FColor GetDebugColor() const { return DebugColors[(int32)Context]; }
+	static NETWORKPREDICTION_API FColor DebugColors[(int32)EVisualLoggingContext::MAX];
+};
+
+// -------------------------------------------------------------------------------------------------------------------------------
+// Interface that the proxy talks to. This is what will implement the replication.
+// -------------------------------------------------------------------------------------------------------------------------------
+class IReplicationProxy
+{
+public:
+
+	virtual ~IReplicationProxy() { }
+	
+	virtual void NetSerializeProxy(EReplicationProxyTarget Target, const FNetSerializeParams& Params) = 0;
+	virtual int32 GetProxyDirtyCount(EReplicationProxyTarget Target) = 0;
+};
+
+// -------------------------------------------------------------------------------------------------------------------------------
+//	FReplicationProxy
+//	Replicated USTRUCT that point to the IReplicationProxy to do the replication.
+// -------------------------------------------------------------------------------------------------------------------------------
+USTRUCT()
+struct FReplicationProxy
+{
+	GENERATED_BODY()
+
+	void Init(IReplicationProxy* InNetworkSimModel, EReplicationProxyTarget InReplicationTarget);
+	bool NetSerialize(FArchive& Ar, class UPackageMap* Map, bool& bOutSuccess);
+	void OnPreReplication();	
+	bool Identical(const FReplicationProxy* Other, uint32 PortFlags) const;
+	
+private:
+
+	IReplicationProxy* NetworkSimModel;
+	EReplicationProxyTarget ReplicationTarget;
+	int32 CachedDirtyCount = 0;
+};
+
+template<>
+struct TStructOpsTypeTraits<FReplicationProxy> : public TStructOpsTypeTraitsBase2<FReplicationProxy>
+{
+	enum
+	{
+		WithNetSerializer = true,
+		WithIdentical = true,
+	};
+};
+
+// -------------------------------------------------------------------------------------------------------------------------------
+//	FServerRPCProxyParameter
+//	Used for the client->Server RPC. Since this is instantiated on the stack by the replication system prior to net serializing,
+//	we have no opportunity to point the RPC parameter to the member variables we want. So we serialize into a generic temp byte buffer
+//	and move them into the real buffers on the component in the RPC body (via ::NetSerialzeToProxy).
+// -------------------------------------------------------------------------------------------------------------------------------
+USTRUCT()
+struct FServerReplicationRPCParameter
+{
+	GENERATED_BODY()
+
+	// Receive flow: ctor() -> NetSerializetoProxy
+	FServerReplicationRPCParameter() : Proxy(nullptr)	{ }
+	void NetSerializeToProxy(FReplicationProxy& InProxy);
+
+	// Send flow: ctor(Proxy) -> NetSerialize
+	FServerReplicationRPCParameter(FReplicationProxy& InProxy) : Proxy(&InProxy) { }
+	bool NetSerialize(FArchive& Ar, class UPackageMap* Map, bool& bOutSuccess);
+
+private:
+
+	static TArray<uint8> TempStorage;
+
+	FReplicationProxy* Proxy;
+	int64 CachedNumBits = -1;
+	class UPackageMap* CachedPackageMap = nullptr;
+};
+
+template<>
+struct TStructOpsTypeTraits<FServerReplicationRPCParameter> : public TStructOpsTypeTraitsBase2<FServerReplicationRPCParameter>
+{
+	enum
+	{
+		WithNetSerializer = true
+	};
+};
+
+// Helper struct to bypass the bandwidth limit imposed by the engine's NetDriver (QueuedBits, NetSpeed, etc).
+// This is really a temp measure to make the system easier to drop in/try in a project without messing with your engine settings.
+// (bandwidth optimizations have not been done yet and the system in general hasn't been stressed with packetloss / gaps in command streams)
+// So, you are free to use this in your own code but it may be removed one day. Hopefully in general bandwidth limiting will also become more robust.
+struct FScopedBandwidthLimitBypass
+{
+	FScopedBandwidthLimitBypass(AActor* OwnerActor);
+	~FScopedBandwidthLimitBypass();
+private:
+
+	int64 RestoreBits = 0;
+	class UNetConnection* CachedNetConnection = nullptr;
+};
\ No newline at end of file
diff --git a/Engine/Plugins/Runtime/NetworkPrediction/Source/NetworkPrediction/Public/NetworkSimulationModelDebugger.h b/Engine/Plugins/Runtime/NetworkPrediction/Source/NetworkPrediction/Public/NetworkSimulationModelDebugger.h
new file mode 100644
--- /dev/null
+++ b/Engine/Plugins/Runtime/NetworkPrediction/Source/NetworkPrediction/Public/NetworkSimulationModelDebugger.h
@@ -0,0 +1,564 @@
+// Copyright 1998-2019 Epic Games, Inc. All Rights Reserved.
+
+#pragma once
+
+#include "Math/Color.h"
+#include "DrawDebugHelpers.h"
+#include "Debug/DebugDrawService.h"
+#include "Engine/Canvas.h"
+#include "Debug/ReporterGraph.h"
+#include "Engine/Engine.h"
+#include "UObject/WeakObjectPtrTemplates.h"
+#include "Tickable.h"
+#include "NetworkSimulationModel.h"
+#include "CanvasItem.h"
+#include "Containers/Array.h"
+#include "Templates/UniquePtr.h"
+#include "NetworkSimulationModelTemplates.h"
+
+DEFINE_LOG_CATEGORY_STATIC(LogNetworkSimDebug, Log, All);
+
+namespace NetworkSimulationModelDebugCVars
+{
+static int32 DrawKeyframes = 1;
+static FAutoConsoleVariableRef CVarDrawKeyframes(TEXT("nsm.debug.DrawKeyFrames"),
+	DrawKeyframes, TEXT("Draws keyframe data (text) in debug graphs"), ECVF_Default);
+
+static int32 GatherServerSidePIE = 1;
+static FAutoConsoleVariableRef CVarGatherServerSidePIE(TEXT("nsm.debug.GatherServerSide"),
+	GatherServerSidePIE, TEXT("Whenever we gather debug info from a client side actor, also gather server side equivelent. Only works in PIE."), ECVF_Default);
+}
+
+struct FNetworkSimulationModelDebuggerManager;
+
+NETWORKPREDICTION_API UObject* FindReplicatedObjectOnPIEServer(UObject* ClientObject);
+
+// ------------------------------------------------------------------------------------------------------------------------
+//	Debugger support classes
+// ------------------------------------------------------------------------------------------------------------------------
+
+struct INetworkSimulationModelDebugger
+{
+	virtual ~INetworkSimulationModelDebugger() { }
+
+	bool IsActive() { return bActive; }
+	void SetActive(bool InActive) { bActive = InActive; }
+
+	virtual void GatherCurrent(FNetworkSimulationModelDebuggerManager& Out, UCanvas* C) = 0;
+	virtual void Tick( float DeltaTime ) = 0;
+
+protected:
+	bool bActive = false; // Whether you should draw every frame
+};
+
+struct NETWORKPREDICTION_API FNetworkSimulationModelDebuggerManager: public FTickableGameObject, FNoncopyable
+{
+	static FNetworkSimulationModelDebuggerManager& Get();
+
+	~FNetworkSimulationModelDebuggerManager()
+	{
+		if (Graph.IsValid())
+		{
+			Graph->RemoveFromRoot();
+		}
+	}
+
+	FNetworkSimulationModelDebuggerManager()
+	{
+		DrawDebugServicesHandle = UDebugDrawService::Register(TEXT("Game"), FDebugDrawDelegate::CreateRaw(this, &FNetworkSimulationModelDebuggerManager::DrawDebugService));
+		check(DrawDebugServicesHandle.IsValid());
+	}
+
+	// ---------------------------------------------------------------------------------------------------------------------------------------
+	//	Outside API (registration, console commands, draw services, etc)
+	// ---------------------------------------------------------------------------------------------------------------------------------------
+
+	template <typename T, typename TDriver>
+	void RegisterNetworkSimulationModel(T* NetworkSim, TDriver* Driver, AActor* OwningActor, FString DebugName);
+
+	void SetDebuggerActive(AActor* OwningActor, bool InActive)
+	{
+		if (INetworkSimulationModelDebugger* Debugger = Find(OwningActor))
+		{
+			Debugger->SetActive(InActive);
+		}
+		ResetCache();
+		Gather(LastCanvas.Get());
+	}
+
+	void ToggleDebuggerActive(AActor* OwningActor)
+	{
+		if (INetworkSimulationModelDebugger* Debugger = Find(OwningActor))
+		{
+			Debugger->SetActive(!Debugger->IsActive());
+		}
+		ResetCache();
+		Gather(LastCanvas.Get());
+	}
+
+	void SetContinousGather(bool InGather)
+	{
+		bContinousGather = InGather;
+		if (!bContinousGather)
+		{
+			Gather(LastCanvas.Get());
+		}
+	}
+
+	void ToggleContinousGather()
+	{
+		SetContinousGather(!bContinousGather);
+	}
+
+	void DrawDebugService(UCanvas* C, APlayerController* PC)
+	{
+		LastCanvas = C;
+		if (bContinousGather)
+		{
+			Gather(C);
+		}
+		
+		FDisplayDebugManager& DisplayDebugManager = C->DisplayDebugManager;
+		DisplayDebugManager.Initialize(C, GEngine->GetSmallFont(), FVector2D(4.0f, 150.0f));
+
+		if (Lines.Num() > 0)
+		{
+			const float TextScale = FMath::Max(C->SizeX / 1920.0f, 1.0f);
+			FCanvasTileItem TextBackgroundTile(FVector2D(0.0f, 120.0f), FVector2D(400.0f, 1800.0f) * TextScale, FColor(0, 0, 0, 100));
+			TextBackgroundTile.BlendMode = SE_BLEND_Translucent;
+			C->DrawItem(TextBackgroundTile);
+		}
+
+		// --------------------------------------------------------
+		//	Lines
+		// --------------------------------------------------------
+
+		for (FDebugLine& Line : Lines)
+		{
+			DisplayDebugManager.SetDrawColor(Line.Color);
+			DisplayDebugManager.DrawString(Line.Str);
+		}
+
+		// --------------------------------------------------------
+		//	Canvas Items (graphs+text)
+		// --------------------------------------------------------
+		
+		for (auto& Item : CanvasItems[0])
+		{
+			C->DrawItem(*Item.Get());
+		}
+
+		if (NetworkSimulationModelDebugCVars::DrawKeyframes > 0)
+		{
+			for (auto& Item : CanvasItems[1])
+			{
+				C->DrawItem(*Item.Get());
+			}
+		}
+	}
+
+	virtual void Tick( float DeltaTime )
+	{
+		for (auto It = DebuggerMap.CreateIterator(); It; ++It)
+		{
+			AActor* Owner = It.Key().Get();
+			if (!Owner)
+			{
+				It.RemoveCurrent();
+				continue;;
+			}
+			
+			if (It.Value()->IsActive())
+			{
+				It.Value()->Tick(DeltaTime);
+			}
+		}
+	}	
+
+	/** return the stat id to use for this tickable **/
+	virtual TStatId GetStatId() const override
+	{
+		RETURN_QUICK_DECLARE_CYCLE_STAT(FNetworkSimulationModelDebuggerManager, STATGROUP_TaskGraphTasks);
+	}
+
+	/** Gathers latest and Logs single frame */
+	void LogSingleFrame(FOutputDevice& Ar)
+	{
+		Gather(LastCanvas.Get());
+		
+		for (FDebugLine& Line : Lines)
+		{
+			Ar.Logf(TEXT("%s"), *Line.Str);
+		}
+	}
+
+	// ---------------------------------------------------------------------------------------------------------------------------------------
+	//	Debugging API used by TNetworkSimulationModelDebugger
+	// ---------------------------------------------------------------------------------------------------------------------------------------
+
+	void Emit(const FString& Str = FString(), FColor Color = FColor::White, float XOffset=0.f, float YOffset=0.f)
+	{
+		Lines.Emplace(FDebugLine{ Str, Color, XOffset, YOffset });
+	}
+
+	template <typename TBuffer>
+	void EmitElement(TBuffer& Buffer, const FStandardLoggingParameters& Parameters)
+	{
+		FStringOutputDevice StrOut;
+		StrOut.SetAutoEmitLineTerminator(true);
+
+		FStandardLoggingParameters LocalParameters = Parameters;
+		LocalParameters.Ar = &StrOut;
+
+		auto* Element = Buffer.FindElementByKeyframe(LocalParameters.Keyframe);
+		if (Element)
+		{
+			Element->Log(LocalParameters);
+	
+			TArray<FString> StrLines;
+			StrOut.ParseIntoArrayLines(StrLines, true);
+			for (FString& Str : StrLines)
+			{
+				Emit(Str);
+			}
+		}
+	}
+
+	void EmitQuad(FVector2D ScreenPosition, FVector2D ScreenSize, FColor Color)
+	{
+		FVector2D Quad[4];
+		
+		Quad[0].X = ScreenPosition.X;
+		Quad[0].Y = ScreenPosition.Y;
+
+		Quad[1].X = ScreenPosition.X;
+		Quad[1].Y = ScreenPosition.Y + ScreenSize.Y;
+
+		Quad[2].X = ScreenPosition.X + ScreenSize.X;
+		Quad[2].Y = ScreenPosition.Y + ScreenSize.Y;
+
+		Quad[3].X = ScreenPosition.X + ScreenSize.X;
+		Quad[3].Y = ScreenPosition.Y;
+		
+		CanvasItems[0].Emplace( MakeUnique<FCanvasTriangleItem>(Quad[0], Quad[1], Quad[2], GWhiteTexture) );
+		CanvasItems[0].Last()->SetColor(Color);
+
+		CanvasItems[0].Emplace( MakeUnique<FCanvasTriangleItem>(Quad[2], Quad[3], Quad[0], GWhiteTexture) );
+		CanvasItems[0].Last()->SetColor(Color);
+	}
+
+	void EmitText(FVector2D ScreenPosition, FColor Color, const FString& Str)
+	{
+		CanvasItems[1].Emplace( MakeUnique<FCanvasTextItem>(ScreenPosition, FText::FromString(Str), GEngine->GetTinyFont(), Color) );
+	}
+
+private:
+
+	INetworkSimulationModelDebugger* Find(AActor* Actor)
+	{
+		if (!Actor)
+		{
+			return nullptr;
+		}
+
+		INetworkSimulationModelDebugger* Debugger = DebuggerMap.FindRef(TWeakObjectPtr<AActor>(Actor));
+		if (!Debugger)
+		{
+			UE_LOG(LogNetworkSimDebug, Warning, TEXT("Could not find NetworkSimulationModel associated with %s"), *GetPathNameSafe(Actor));
+		}
+		return Debugger;
+	}
+
+	void Gather(UCanvas* C)
+	{
+		ResetCache();
+
+		for (auto It = DebuggerMap.CreateIterator(); It; ++It)
+		{
+			AActor* Owner = It.Key().Get();
+			if (!Owner)
+			{
+				It.RemoveCurrent();
+				continue;;
+			}
+			
+			if (It.Value()->IsActive())
+			{
+				It.Value()->GatherCurrent(*this, C);
+				if (NetworkSimulationModelDebugCVars::GatherServerSidePIE > 0)
+				{
+					if (AActor* ServerSideActor = Cast<AActor>(FindReplicatedObjectOnPIEServer(Owner)))
+					{
+						if (INetworkSimulationModelDebugger* ServerSideSim = Find(ServerSideActor))
+						{
+							Emit();
+							Emit();
+							ServerSideSim->GatherCurrent(*this, nullptr); // Dont do graphs for server side state
+						}
+					}
+				}
+			}
+		}
+	}
+
+	void ResetCache()
+	{
+		Lines.Reset();
+		
+		CanvasItems[0].Reset();
+		CanvasItems[1].Reset();
+	}
+
+	TMap<TWeakObjectPtr<AActor>, INetworkSimulationModelDebugger*>	DebuggerMap;
+	bool bContinousGather = true; // Whether you should gather new data every frame
+
+	FDelegateHandle DrawDebugServicesHandle;
+	struct FDebugLine
+	{
+		FString Str;
+		FColor Color;
+		float XOffset;
+		float YOffset;
+	};
+	TArray<FDebugLine> Lines;
+	TArray<TUniquePtr<FCanvasItem>> CanvasItems[2];
+	TWeakObjectPtr<UReporterGraph> Graph;
+	TWeakObjectPtr<UCanvas> LastCanvas;
+};
+
+template <typename TNetSimModel, typename TDriver>
+struct TNetworkSimulationModelDebugger : public INetworkSimulationModelDebugger
+{
+	TNetworkSimulationModelDebugger(TNetSimModel* InNetSim, TDriver* InDriver, AActor* OwningActor, FString InDebugName)
+	{
+		NetworkSim = InNetSim;
+		Driver = InDriver;
+		WeakOwningActor = OwningActor;
+		DebugName = InDebugName;
+	}
+
+	~TNetworkSimulationModelDebugger()
+	{
+		
+	}	
+
+	void GatherCurrent(FNetworkSimulationModelDebuggerManager& Out, UCanvas* Canvas) override
+	{
+		AActor* Owner = WeakOwningActor.Get();
+		if (!ensure(Owner))
+		{
+			return;
+		}
+
+		// ------------------------------------------------------------------------------------------------------------------------------------------------
+		//	Lines
+		// ------------------------------------------------------------------------------------------------------------------------------------------------
+
+		Out.Emit(FString::Printf(TEXT("%s - %s"), *Owner->GetName(), *UEnum::GetValueAsString(TEXT("Engine.ENetRole"), Owner->Role)), FColor::Yellow);
+		Out.Emit(FString::Printf(TEXT("LastProcessedInputKeyframe: %d (%d Buffered)"), NetworkSim->LastProcessedInputKeyframe, NetworkSim->InputBuffer.GetHeadKeyframe() - NetworkSim->LastProcessedInputKeyframe));
+
+		// Autorproxy
+		{
+			FColor Color = FColor::White;
+			const bool FaultDetected = NetworkSim->RepProxy_Autonomous.IsReconcileFaultDetected();
+
+			const int32 LastSerializedKeyframe = NetworkSim->RepProxy_Autonomous.GetLastSerializedKeyframe();
+
+			// Calc how much predicted time we have processed. Note that we use the motionstate buffer to iterate but the MS is on the input cmd. (if we are buffering cmds, don't want to count them)
+			float PredictedMS = 0.f;
+			for (int32 PredKeyrame = LastSerializedKeyframe+1; PredKeyrame <= NetworkSim->SyncBuffer.GetHeadKeyframe(); ++PredKeyrame)
+			{
+				if (auto* Cmd = NetworkSim->InputBuffer.FindElementByKeyframe(PredKeyrame))
+				{
+					PredictedMS += Cmd->FrameDeltaTime * 1000.f;
+				}
+			}
+
+			FString ConfirmedFrameStr = FString::Printf(TEXT("LastConfirmedFrame: %d. Prediction: %d Frames, %.2f MS"), LastSerializedKeyframe, NetworkSim->SyncBuffer.GetHeadKeyframe() - LastSerializedKeyframe, PredictedMS);
+			if (FaultDetected)
+			{
+				ConfirmedFrameStr += TEXT(" RECONCILE FAULT DETECTED!");
+				Color = FColor::Red;
+			}
+
+			Out.Emit(*ConfirmedFrameStr, Color);
+		}
+
+		auto EmitBuffer = [&Out](FString BufferName, auto& Buffer)
+		{
+			Out.Emit();
+			Out.Emit(FString::Printf(TEXT("//////////////// %s ///////////////"), *BufferName), FColor::Yellow);
+			Out.Emit(FString::Printf(TEXT("%s"), *Buffer.GetBasicDebugStr()));		
+			Out.Emit();
+			Out.EmitElement(Buffer, FStandardLoggingParameters(nullptr, EStandardLoggingContext::Full, Buffer.GetHeadKeyframe()));
+		};
+
+		EmitBuffer(TEXT("InputBuffer"), NetworkSim->InputBuffer);
+		EmitBuffer(TEXT("SyncBuffer"), NetworkSim->SyncBuffer);
+
+		// ------------------------------------------------------------------------------------------------------------------------------------------------
+		//	Canvas
+		// ------------------------------------------------------------------------------------------------------------------------------------------------
+
+		auto* DebugBuffer = NetworkSim->GetDebugBuffer();
+		if (Canvas && DebugBuffer && DebugBuffer->GetNumValidElements() > 0)
+		{
+			auto& InputBuffer = NetworkSim->GetHistoricBuffers() ? NetworkSim->GetHistoricBuffers()->InputBuffer : NetworkSim->InputBuffer;
+
+			static float StartPctX = 0.3f;
+			static float StartPctY = 0.6f;
+			static float LocalFrameHeightPct = 0.01f;
+
+			float ScreenX = StartPctX * (float)Canvas->SizeX;
+			float ScreenY = StartPctY * (float)Canvas->SizeY;
+			float LocalFrameHeight = LocalFrameHeightPct * (float)Canvas->SizeY;
+
+			static float LocalFrameTimeGreen = 1/30.f;
+			static float LocalFrameTimeRed = 1/10.f;
+
+			static float Thickness = 1.f;
+			static float ClientOffsetY = 2.f;
+
+			FTextSizingParameters TextSizing;
+			TextSizing.DrawFont = GEngine->GetTinyFont();
+			TextSizing.Scaling = FVector2D(1.f,1.f);
+			Canvas->CanvasStringSize(TextSizing, TEXT("00000"));
+
+			float MinWidth = TextSizing.DrawXL;
+			float MinHeight = TextSizing.DrawYL;
+			static float MinHeightMS = 1/60.f * 1000.f;	// 60hz frame is drawn at MinHeight.
+
+			auto CalcHeight = [&](float MS)
+			{
+				const float Ratio = MinHeight / MinHeightMS;
+				return MS * Ratio;
+			};
+
+			for (auto It = DebugBuffer->CreateIterator(); It; ++It)
+			{
+				auto* DebugState = It.Element();
+
+				
+				float ScreenWidth = MinWidth;
+				float FramePct = FMath::Clamp<float>((DebugState->LocalDeltaTimeSeconds - LocalFrameTimeRed) / ( LocalFrameTimeGreen - LocalFrameTimeRed ), 0.f, 1.f);
+				FColor Color = FColor::MakeRedToGreenColorFromScalar( FramePct );
+				float ServerHeight = CalcHeight(DebugState->LocalDeltaTimeSeconds * 1000.f);
+
+				Out.EmitQuad(FVector2D( ScreenX, ScreenY), FVector2D( ScreenWidth, ServerHeight), Color);
+				Out.EmitText(FVector2D( ScreenX, ScreenY), FColor::Black, FString::Printf(TEXT("%.2f"), (DebugState->LocalDeltaTimeSeconds * 1000.f)));
+
+				float ClientSimTime = 0.f;
+				float ClientX = ScreenX;
+				float ClientY = ScreenY - ClientOffsetY;
+				for (int32 Keyframe : DebugState->ProcessedKeyframes)
+				{
+					auto* Cmd = InputBuffer.FindElementByKeyframe(Keyframe);
+					if (Cmd)
+					{
+						ClientSimTime += Cmd->FrameDeltaTime * 1000.f;
+						
+						float ClientSizeX = MinWidth;
+						float ClientSizeY =  CalcHeight(Cmd->FrameDeltaTime * 1000.f);
+
+						FVector2D ScreenPos(ClientX, ClientY - ClientSizeY);
+						Out.EmitQuad(ScreenPos, FVector2D( ClientSizeX, ClientSizeY), FColor::Blue);
+						Out.EmitText(ScreenPos, FColor::White, LexToString(Keyframe));
+						ClientY -= (ClientSizeY + ClientOffsetY);
+					}
+				}
+
+				for (int32 Keyframe = DebugState->LastProcessedKeyframe+1; Keyframe <= DebugState->HeadKeyframe; ++Keyframe)
+				{
+					if (auto* Cmd = InputBuffer.FindElementByKeyframe(Keyframe))
+					{
+						float ClientSizeX = MinWidth;
+						float ClientSizeY = CalcHeight(Cmd->FrameDeltaTime * 1000.f);
+						FVector2D ScreenPos(ClientX, ClientY - ClientSizeY);
+						Out.EmitQuad(ScreenPos, FVector2D( ClientSizeX, ClientSizeY), FColor::Red);
+						Out.EmitText(ScreenPos, FColor::White, LexToString(Keyframe));
+						ClientY -= (ClientSizeY + ClientOffsetY);
+					}
+				}
+
+					
+				// -------------------------------------------
+				ScreenX += ScreenWidth + 2.f;
+			}
+		}
+	}
+
+	void Tick( float DeltaTime ) override
+	{
+		AActor* Owner = WeakOwningActor.Get();
+		if (!Owner)
+		{
+			return;
+		}
+
+		UWorld* World = Owner->GetWorld();
+
+		if (auto* LatestSync = NetworkSim->SyncBuffer.GetElementFromHead(0))
+		{
+			LatestSync->VisualLog( FVisualLoggingParameters(EVisualLoggingContext::LastPredicted, NetworkSim->SyncBuffer.GetHeadKeyframe(), EVisualLoggingLifetime::Transient), Driver, Driver);
+		}
+
+		FStuff ServerPIEStuff = GetServerPIEStuff();
+		if (ServerPIEStuff.Driver && ServerPIEStuff.NetworkSim)
+		{
+			if (auto* ServerLatestSync = ServerPIEStuff.NetworkSim->SyncBuffer.GetElementFromHead(0))
+			{
+				ServerLatestSync->VisualLog( FVisualLoggingParameters(EVisualLoggingContext::CurrentServerPIE, ServerPIEStuff.NetworkSim->SyncBuffer.GetHeadKeyframe(), EVisualLoggingLifetime::Transient), ServerPIEStuff.Driver, Driver);
+			}
+		}
+		
+		for (int32 Keyframe = NetworkSim->RepProxy_Autonomous.GetLastSerializedKeyframe(); Keyframe < NetworkSim->SyncBuffer.GetHeadKeyframe(); ++Keyframe)
+		{
+			if (auto* SyncState = NetworkSim->SyncBuffer.FindElementByKeyframe(Keyframe))
+			{
+				const EVisualLoggingContext Context = (Keyframe == NetworkSim->RepProxy_Autonomous.GetLastSerializedKeyframe()) ? EVisualLoggingContext::LastConfirmed : EVisualLoggingContext::OtherPredicted;
+				SyncState->VisualLog( FVisualLoggingParameters(Context, NetworkSim->SyncBuffer.GetHeadKeyframe(), EVisualLoggingLifetime::Transient), Driver, Driver);
+			}
+		}
+	}
+	
+
+	struct FStuff
+	{
+		TNetSimModel* NetworkSim = nullptr;
+		TDriver* Driver = nullptr;
+	};
+
+	FStuff GetStuff()
+	{
+		return {NetworkSim, Driver};
+	}
+
+	TFunction< FStuff() > GetServerPIEStuff;
+
+private:
+
+	
+	TWeakObjectPtr<AActor>	WeakOwningActor;
+	FString DebugName;
+
+	TNetSimModel* NetworkSim = nullptr;
+	TDriver* Driver = nullptr;
+};
+
+template <typename T, typename TDriver>
+void FNetworkSimulationModelDebuggerManager::RegisterNetworkSimulationModel(T* NetworkSim, TDriver* Driver, AActor* OwningActor, FString DebugName)
+{
+	TNetworkSimulationModelDebugger<T, TDriver>* Debugger = new TNetworkSimulationModelDebugger<T, TDriver>(NetworkSim, Driver, OwningActor, DebugName);
+	DebuggerMap.Add( TWeakObjectPtr<AActor>(OwningActor), Debugger );
+
+	// Gross stuff so that the debugger can find the ServerPIE equiv
+	TWeakObjectPtr<AActor> WeakOwner(OwningActor);
+	Debugger->GetServerPIEStuff = [WeakOwner, this]()
+	{
+		if (AActor* ServerOwner = Cast<AActor>(FindReplicatedObjectOnPIEServer(WeakOwner.Get())))
+		{
+			return ((TNetworkSimulationModelDebugger<T, TDriver>*)DebuggerMap.FindRef(TWeakObjectPtr<AActor>(ServerOwner)))->GetStuff();
+		}
+		return typename TNetworkSimulationModelDebugger<T, TDriver>::FStuff();
+	};
+}
\ No newline at end of file
diff --git a/Engine/Plugins/Runtime/NetworkPrediction/Source/NetworkPrediction/Public/NetworkSimulationModelTemplates.h b/Engine/Plugins/Runtime/NetworkPrediction/Source/NetworkPrediction/Public/NetworkSimulationModelTemplates.h
new file mode 100644
--- /dev/null
+++ b/Engine/Plugins/Runtime/NetworkPrediction/Source/NetworkPrediction/Public/NetworkSimulationModelTemplates.h
@@ -0,0 +1,942 @@
+
+// Copyright 1998-2018 Epic Games, Inc. All Rights Reserved.
+
+#pragma once
+#include "NetworkSimulationModel.h"
+
+// ----------------------------------------------------------------------------------------------------------------------------------------------------------------------
+//	CVars and compile time constants
+// ----------------------------------------------------------------------------------------------------------------------------------------------------------------------
+
+#ifndef NETSIM_NETCONSTANT_NUM_BITS_KEYFRAME
+	#define NETSIM_NETCONSTANT_NUM_BITS_KEYFRAME 8	// Allows you to override this setting via UBT, but access via cleaner FActorMotionNetworkingConstants::NUM_BITS_KEYFRAME
+#endif
+
+struct FNetworkSimulationSerialization
+{
+	// How many bits we use to encode the key frame number for buffers.
+	// Client Keyframes are stored locally as 32 bit integers, but we use a smaller # of bits to NetSerialize.
+	// Frames are only relatively relevant: the absolute value doesn't really matter. We just need to detect newer/older.
+	enum { NUM_BITS_KEYFRAME = NETSIM_NETCONSTANT_NUM_BITS_KEYFRAME };		
+
+	// Abs max value we encode into the bit writer
+	enum { MAX_KEYFRAME_WRITE = 1 << NUM_BITS_KEYFRAME };
+
+	// This is the threshold at which we would wrap around and incorrectly assign a keyframe on the receiving side.
+	// E.g, If there are KEYFRAME_ERROR_THRESHOLD keyframes that do not make it across from sender->receiver, the
+	// receiver will have incorrect local values. With 8 bits, this works out to be 128 frames or about 2 seconds at 60fps.
+	enum { KEYFRAME_ERROR_THRESHOLD = MAX_KEYFRAME_WRITE / 2};
+
+	// Helper to serialize the int32 HeadKeyframe. Returns the unpacked value (this will be same as input in the save path)
+	static int32 SerializeKeyframe(FArchive& Ar, int32 LocalHeadKeyframe)
+	{
+		if (Ar.IsSaving())
+		{
+			((FNetBitWriter&)Ar).WriteIntWrapped( LocalHeadKeyframe, MAX_KEYFRAME_WRITE );
+			return LocalHeadKeyframe;
+		}
+		
+		return MakeRelative(((FNetBitReader&)Ar).ReadInt( MAX_KEYFRAME_WRITE ), LocalHeadKeyframe, MAX_KEYFRAME_WRITE );
+	}
+};
+
+namespace NetworkSimulationModelCVars
+{
+static int32 EnableLocalPrediction = 1;
+static FAutoConsoleVariableRef CVarEnableLocalPrediction(TEXT("ns.EnableLocalPrediction"),
+	EnableLocalPrediction, TEXT("Toggle local prediction. Meant for debugging."), ECVF_Default);
+
+static int32 ForceReconcile = 0;
+static FAutoConsoleVariableRef CVarForceReconcile(TEXT("ns.ForceReconcile"),
+	ForceReconcile, TEXT("Forces reconcile even if state does not differ. E.g, force resimulation after every netupdate."), ECVF_Default);
+
+static int32 MaxInputCmdsFrame = 1;
+static FAutoConsoleVariableRef CVarMaxInputCmdsFrame(TEXT("ns.MaxInputCmdsFrame"),
+	MaxInputCmdsFrame, TEXT("Max cap on how many input cmds can be processed in a frame."), ECVF_Default);
+}
+
+// ----------------------------------------------------------------------------------------------------------------------------------------------------------------------
+//	TReplicationBuffer
+//
+//	Generic Cyclic buffer. Has canonical head position. This is the "Client Frame" or "Keyframe" identifier used by the rest of the system.
+//	Contract: Elements in buffer are contiguously valid. We don't allow "gaps" in the buffer!
+//	Use ::GetWriteNext<T> to write to the buffer.
+//	Use ::CreateIterator to iterate from Tail->Head
+//
+//	for (auto It = Buffer.CreateIterator(); It; ++It) {
+//		FMyType* Element = It.Element();
+//		int32 Keyframe = It.Keyframe();
+//	}
+//
+// ----------------------------------------------------------------------------------------------------------------------------------------------------------------------
+template<typename T>
+struct TReplicationBuffer
+{
+	// Set the buffer size. Resizing is supported but not optimized because we preserve head/tail iteration (not a simple memcopy since where buffer wraps around will change)
+	// In general, avoid resizing outside of startup/initialization.
+	void SetBufferSize(int32 NewMaxNumElements)
+	{
+		check(NewMaxNumElements >= 0);	
+		if (Data.Num() == NewMaxNumElements)
+		{
+			return;
+		}
+		
+		if (Data.Num() == 0)
+		{
+			Data.SetNum(NewMaxNumElements, false);
+		}
+		else
+		{
+			// Grow or shrink. This is far from optimal but this operation should be rare
+			TReplicationBuffer<T> Old(MoveTemp(*this)); // Move old data to a copy
+			Data.SetNum(NewMaxNumElements, false);
+			ResetNextHeadKeyframe(Old.GetTailKeyframe(), true); // Reset so our next write is at the old copies tail keyframe
+
+			for (auto It = Old.CreateIterator(); It; ++It)
+			{
+				T* WriteNext = GetWriteNext();
+				T* ReadNext = It.Element();
+				FMemory::Memcpy(WriteNext, ReadNext, sizeof(T));
+			}
+
+			// Blow away the temp copy so we don't run the dtor on it on scope end (because we memcpy'd it over, not copy assigned)
+			Old.Data.Reset();
+			Old.NumValidElements = 0;
+		}
+	}
+
+	// Gets the oldest element in the buffer (or null)
+	T* GetElementFromTail(int32 OffsetFromTail)
+	{
+		check(OffsetFromTail >= 0); // Must always pass in valid offset
+
+		// Out of range. This is not fatal, just don't return the wrong element.
+		if (NumValidElements <= OffsetFromTail)
+		{
+			return nullptr;
+		}
+		
+		const int32 Tail = GetTailKeyframe();
+		const int32 Position = Tail + OffsetFromTail;
+		const int32 Offset = (Position % Data.Num());
+		return &Data[Offset];
+	}
+
+	// Gets the newest element in the buffer (or null)
+	T* GetElementFromHead(int32 OffsetFromHead)
+	{
+		check(OffsetFromHead >= 0); // Must always pass in valid offset
+
+		// Out of range. This is not fatal, just don't return the wrong element.
+		if (NumValidElements <= OffsetFromHead)
+		{
+			return nullptr;
+		}
+		
+		const int32 Position = Head - OffsetFromHead;
+		check(Position >= 0);
+		const int32 Offset = (Position % Data.Num());
+		return &Data[Offset];
+	}
+
+	const T* FindElementByKeyframe(int32 Key) const
+	{
+		return FindElementByKeyframeImpl(Key);
+	}
+
+	T* FindElementByKeyframe(int32 Key)
+	{
+		const T* Element = const_cast<TReplicationBuffer<T>*>(this)->FindElementByKeyframeImpl(Key);
+		return const_cast<T*>(Element);
+	}
+
+	// Returns the next element for writing. The contents of this element are unknown (could be stale content!). This also advances the head. Note the element returned is immediately considered "valid"!
+	T* GetWriteNext()
+	{
+		check(Data.Num() > 0); // Buffer must be initialized first
+		
+		++Head;
+		++DirtyCount;
+		NumValidElements = FMath::Min<int32>(NumValidElements+1, Data.Num());
+
+		const int32 Offset = (Head % Data.Num());
+		return &Data[Offset];
+	}
+
+	// Resets the head keyframen, so that NextHeadKeyframe will be the keyframe of the next GetWriteNext<>(). E.g, call this before you want to write to NextHeadKeyframe.
+	// To remphasize: NextHeadKeyframe is where the NEXT write will go. NOT "where the current head is".
+	// Contents will be preserved if possible. E.g, if NextHeadKeyframe-1 is a valid keyframe, it will remain valid. If its not (if this creates "gaps" in the buffer) then
+	// the contents will effectively be cleared. bForceClearContents will force this behavior even if the previous keyframe was valid.
+	void ResetNextHeadKeyframe(int32 NextHeadKeyframe=0, bool bForceClearContents=false)
+	{
+		const int32 NewHeadKeyframe = NextHeadKeyframe - 1;
+		if (bForceClearContents || NewHeadKeyframe < GetTailKeyframe() || NewHeadKeyframe > GetHeadKeyframe())
+		{
+			NumValidElements = 0;
+		}
+		else
+		{
+			NumValidElements += (NewHeadKeyframe - GetHeadKeyframe());
+			check(NumValidElements >= 0 && NumValidElements <= Data.Num());
+		}
+
+		Head =  NewHeadKeyframe;
+		DirtyCount++;
+	}
+
+	// Copies the contents of the SourceBuffer into this buffer. This existing buffer's data will be preserved as much as possible, but the guarantee is that the SourceBuffer's
+	// data will all be added to this buffer.
+	
+	// Example 1: Target={1...5}, Source={3...9}  --> Target={1...9}
+	// Example 2: Target={1...5}, Source={7...9}  --> Target={7...9}
+	// Exmaple 3: Target={6...9}, Source={1...4}  --> Target={1...4}
+	void CopyAndMerge(const TReplicationBuffer<T>& SourceBuffer)
+	{
+		const int32 StartingHeadKeyframe = GetHeadKeyframe();
+
+		ResetNextHeadKeyframe(SourceBuffer.GetTailKeyframe());
+		for (int32 Keyframe = SourceBuffer.GetTailKeyframe(); Keyframe <= SourceBuffer.GetHeadKeyframe(); ++Keyframe)
+		{
+			T* TargetData = GetWriteNext();
+			check(GetHeadKeyframe() == Keyframe);
+
+			const T* SourceData = SourceBuffer.FindElementByKeyframe(Keyframe);
+			check(SourceData != nullptr);
+
+			*TargetData = *SourceData;
+		}
+	}
+
+	FString GetBasicDebugStr() const
+	{
+		return FString::Printf(TEXT("Elements: [%d/%d]. Keyframes: [%d-%d]"), NumValidElements, Data.Num(), GetTailKeyframe(), GetHeadKeyframe());
+	}
+
+	int32 GetNumValidElements() const { return NumValidElements; }
+	int32 GetMaxNumElements() const { return Data.Num(); }
+	int32 GetHeadKeyframe() const { return Head; }
+	int32 GetTailKeyframe() const { return Head - NumValidElements + 1; }
+	bool IsValidKeyframe(int32 Keyframe) const { return Keyframe >= GetTailKeyframe() && Keyframe <= GetHeadKeyframe(); }
+	int32 GetDirtyCount() const { return DirtyCount; }
+
+	// Create an iterator from tail->head. Note that no matter what template class is, this WILL iterate correctly across the element. The templated type
+	// is just for casting the return element type. E.g, it is fine to use <uint8> in generic code, this will still step by 'StructSize'.
+	struct TGenericIterator
+	{
+		TGenericIterator(TReplicationBuffer<T>& InBuffer) : Buffer(InBuffer)
+		{
+			CurrentKeyframe = Buffer.GetTailKeyframe();
+		}
+
+		int32 Keyframe() const { return CurrentKeyframe; }
+		T* Element() const { return Buffer.FindElementByKeyframe(CurrentKeyframe); }
+		void operator++() { CurrentKeyframe++; }
+		operator bool()	{ return CurrentKeyframe <= Buffer.GetHeadKeyframe(); }
+
+	private:
+		TReplicationBuffer<T>& Buffer;
+		int32 CurrentKeyframe = -1;
+	};
+		
+	TGenericIterator CreateIterator() { return TGenericIterator(*this); }
+
+private:
+
+	const T* FindElementByKeyframeImpl(int32 Key) const
+	{
+		const int32 RelativeToHead = Key - Head;
+		if (RelativeToHead > 0 || RelativeToHead <= -NumValidElements)
+		{
+			return nullptr;
+		}
+
+		return &Data[Key % Data.Num()];
+	}
+
+	int32 Head = INDEX_NONE;
+	int32 DirtyCount = 0;
+	int32 NumValidElements = 0;
+
+	TArray<T> Data;
+};
+
+// ----------------------------------------------------------------------------------------------------------------------------------------------------------------------
+//	Templated Replicators/Reconciliars
+// ----------------------------------------------------------------------------------------------------------------------------------------------------------------------
+
+/** Replicates only the latest element. Does not synchronize keyframe */
+template<typename TReplicated>
+struct TReplicator_LatestOnly
+{
+	TReplicator_LatestOnly(TReplicationBuffer<TReplicated>& InBuffer) : Buffer(InBuffer) { }
+	int32 GetProxyDirtyCount() const { return Buffer.GetDirtyCount(); }
+
+	void NetSerialize(const FNetSerializeParams& P)
+	{
+		FArchive& Ar = P.Ar;
+		TReplicated* State = nullptr;
+		if (Ar.IsSaving())
+		{
+			State = Buffer.GetElementFromHead(0);
+			check(State); // We should not be here if the buffer is empty. Want to avoid serializing an "empty" flag at the top here.
+		}
+		else
+		{
+			State = Buffer.GetWriteNext();
+		}
+
+		State->NetSerialize(Ar);
+	}
+
+private:
+	TReplicationBuffer<TReplicated>& Buffer;
+};
+
+/** Replicates the latest sequence of N elements. N is dynamic, not compiled in. */
+template<typename TReplicated>
+struct TReplicator_DynamicSequence
+{
+	TReplicator_DynamicSequence(TReplicationBuffer<TReplicated>& InBuffer) : Buffer(InBuffer) { }
+	int32 GetProxyDirtyCount() const { return Buffer.GetDirtyCount(); }
+
+	// --------------------------------------------------------------------
+	//	NetSerialize
+	// --------------------------------------------------------------------
+	void NetSerialize(const FNetSerializeParams& P)
+	{
+		FArchive& Ar = P.Ar;
+		uint8 SerializedNumElements = FMath::Min<uint8>(MaxNumElements, Buffer.GetNumValidElements());
+		Ar << SerializedNumElements;
+
+		const int32 HeadKeyframe = FNetworkSimulationSerialization::SerializeKeyframe(Ar, Buffer.GetHeadKeyframe());
+		const int32 StartingKeyframe = FMath::Max(0, HeadKeyframe - SerializedNumElements + 1);
+
+		if (Ar.IsLoading())
+		{
+			const int32 PrevHead = Buffer.GetHeadKeyframe();
+			if (PrevHead < StartingKeyframe && PrevHead >= 0)
+			{
+				// There is a gap in the stream. In some cases, we want this to be a "fault" and bubble up. We may want to synthesize state or maybe we just skip ahead.
+				UE_LOG(LogNetworkSim, Warning, TEXT("Fault: gap in received buffer. PrevHead: %d. Received: %d-%d. Reseting previous buffer contents"), PrevHead, StartingKeyframe, HeadKeyframe);
+			}
+
+			Buffer.ResetNextHeadKeyframe(StartingKeyframe);
+		}
+
+		for (int32 Keyframe = StartingKeyframe; Keyframe <= HeadKeyframe; ++Keyframe)
+		{
+			// This, as is, is bad. The intention is that these functions serialize multiple items in some delta compressed fashion.
+			// As is, we are just serializing the elements individually.
+			TReplicated* Cmd = Ar.IsLoading() ? Buffer.GetWriteNext() : Buffer.FindElementByKeyframe(Keyframe);
+			Cmd->NetSerialize(P);
+		}
+	}
+
+private:
+
+	TReplicationBuffer<TReplicated>& Buffer;
+	uint8 MaxNumElements = 3;
+};
+
+/** Replicates TSyncState and does basic reconciliation */
+template<typename TInputCmd, typename TSyncState, typename TAuxState>
+struct TReplicator_BasicReconciliar
+{
+	TReplicator_BasicReconciliar(TReplicationBuffer<TInputCmd>& InInputBuffer, TReplicationBuffer<TSyncState>& InSyncBuffer, TReplicationBuffer<TAuxState>& InAuxBuffer)
+		: InputBuffer(InInputBuffer), SyncBuffer(InSyncBuffer), AuxBuffer(InAuxBuffer) { }
+
+	int32 GetLastSerializedKeyframe() const { return LastSerializedKeyframe; }
+	bool IsReconcileFaultDetected() const { return bReconcileFaultDetected; }
+	int32 GetProxyDirtyCount() const { return SyncBuffer.GetDirtyCount(); }
+
+	// --------------------------------------------------------------------
+	//	NetSerialize
+	// --------------------------------------------------------------------
+	void NetSerialize(const FNetSerializeParams& P)
+	{
+		FArchive& Ar = P.Ar;
+		const int32 SerializedHeadKeyframe = FNetworkSimulationSerialization::SerializeKeyframe(Ar, SyncBuffer.GetHeadKeyframe());
+		TSyncState* SerializedState = nullptr;
+
+		if (Ar.IsSaving())
+		{
+			// Server serialize the latest state
+			SerializedState = SyncBuffer.GetElementFromHead(0);
+		}
+		else
+		{
+			// Its possible we process two packets in between ticks. This isn't really a problem but should be rare.
+			if (bPendingReconciliation)
+			{
+				UE_LOG(LogNetworkSim, Warning, TEXT("bPendingReconciliation while in ::NetSerialize. LastSerializedKeyframe: %d. New SerializedHeadKeyframe: %d."), LastSerializedKeyframe, SerializedHeadKeyframe);
+			}
+
+			// Lazy init the reconciliation buffer. We don't need this on the server/writing side. (Fixme, with templated approach we could easily spcialize this to an inline elment instead of buffer?)
+			if (ReconciliationBuffer.GetMaxNumElements() == 0)
+			{
+				ReconciliationBuffer.SetBufferSize(1);
+			}
+
+			// Reset the head to match the server. This is not critical, we just need a place to store this, but it makes the code in the reconciliation a bit nicer since the reconcilated frame number = head frame.
+			ReconciliationBuffer.ResetNextHeadKeyframe(SerializedHeadKeyframe);
+			SerializedState = ReconciliationBuffer.GetWriteNext();
+		}
+
+		SerializedState->NetSerialize(Ar);
+
+		if (Ar.IsLoading())
+		{
+			bReconcileFaultDetected = false;
+			bPendingReconciliation = false;
+
+			// The state the client predicted that corresponds to the state the server just serialized to us
+			if (TSyncState* ClientExistingState = SyncBuffer.FindElementByKeyframe(SerializedHeadKeyframe))
+			{
+				if (ClientExistingState->ShouldReconcile(*SerializedState) || (NetworkSimulationModelCVars::ForceReconcile > 0))
+				{
+					UE_CLOG(!InputBuffer.IsValidKeyframe(SerializedHeadKeyframe), LogNetworkSim, Error, TEXT("::NetSerialize: Client InputBuffer does not contain data for frame %d. {%s} {%s}"), SerializedHeadKeyframe, *InputBuffer.GetBasicDebugStr(), *SyncBuffer.GetBasicDebugStr());
+					bPendingReconciliation =  true;
+				}
+			}
+			else
+			{
+				// We don't have corresponding local state. There are two cases:
+				if (NetworkSimulationModelCVars::EnableLocalPrediction) // Fixme: this is awkward. Expected if player prediction is disabled but coupling like this feels bad.
+				{
+					UE_LOG(LogNetworkSim, Warning, TEXT("::NetSerialize Fault: MotionStateBuffer does not contain data for frame %d. [%d-%d]"), SerializedHeadKeyframe, SyncBuffer.GetTailKeyframe(), SyncBuffer.GetHeadKeyframe());
+				}
+
+				// Case 1: the serialized state is older than what we've kept in our buffer. A bigger buffer would solve this! (at the price of more resimulated frames to recover when this happens)
+				if (SerializedHeadKeyframe < SyncBuffer.GetTailKeyframe())
+				{
+					// This is a reconcile fault and we just need to chill. We'll stop sending user commands until the cmds in flight flush through the system and we catch back up.
+					bReconcileFaultDetected = true;
+				}
+				else
+				{
+					// Case 2: We've received a newer frame than what we've processed locally. This could happen if we are buffering our inputs locally (but still sending to the server).
+					// Since this doesn't require resimulating, we can just set the authoritative state here and continue on. Ticking logic will need to detect this gap and skip
+					// our LastProcessedInputKeyframe ahead.
+
+					// However, this keyframe better be a valid input. Otherwise how did the server generate it?
+					checkf(SerializedHeadKeyframe <= InputBuffer.GetHeadKeyframe(), TEXT("Received newer motionstate that doesn't correspond to valid input cmd. SerializedHeadKeyframe: %d. {%s} {%s}"),
+						SerializedHeadKeyframe, *InputBuffer.GetBasicDebugStr(), *SyncBuffer.GetBasicDebugStr());
+
+					SyncBuffer.ResetNextHeadKeyframe(SerializedHeadKeyframe);
+					TSyncState* ClientMotionState = SyncBuffer.GetWriteNext();
+					*ClientMotionState = *SerializedState;
+				}
+			}
+		}
+
+		LastSerializedKeyframe = SerializedHeadKeyframe;
+	}
+
+	// --------------------------------------------------------------------
+	//	Reconcile
+	// --------------------------------------------------------------------
+	template<typename T, typename TDriver>
+	void Reconcile(TDriver* Driver)
+	{
+		if (bPendingReconciliation == false)
+		{
+			return;
+		}
+		bPendingReconciliation = false;
+		check(ReconciliationBuffer.GetNumValidElements() == 1); // this function assumes only 1 element
+
+		const int32 ReconciliationKeyframe = ReconciliationBuffer.GetHeadKeyframe();
+		TSyncState* ServerState = ReconciliationBuffer.GetElementFromHead(0);
+
+		TInputCmd* ClientInputCmd = InputBuffer.FindElementByKeyframe( ReconciliationKeyframe );
+		if (ClientInputCmd == nullptr)
+		{
+			// Fault: no historic command for this frame to reconcile with.
+			UE_LOG(LogNetworkSim, Error, TEXT("Client InputBuffer does not contain data for frame %d. {%s} {%s}"), ReconciliationKeyframe, *InputBuffer.GetBasicDebugStr(), *SyncBuffer.GetBasicDebugStr());
+			return;
+		}
+
+		TSyncState* ClientSyncState = SyncBuffer.FindElementByKeyframe( ReconciliationKeyframe );
+
+		// Should not be here if the client doesn't have a valid state. See ::NetSerialize
+		checkf(ClientSyncState != nullptr, TEXT("SyncBuffer does not contain data for frame %d. %s"), ReconciliationKeyframe, *SyncBuffer.GetBasicDebugStr());
+
+		// -------------------------------------------------------------------------------------------------------------------------
+		// Resimulate
+		// -------------------------------------------------------------------------------------------------------------------------
+
+		ServerState->VisualLog( FVisualLoggingParameters(EVisualLoggingContext::LastConfirmed, ReconciliationKeyframe, EVisualLoggingLifetime::Persistent), Driver, Driver);
+		ClientSyncState->VisualLog( FVisualLoggingParameters(EVisualLoggingContext::FirstMispredicted, ReconciliationKeyframe, EVisualLoggingLifetime::Persistent), Driver, Driver);
+
+		TSyncState* LatestClient = SyncBuffer.GetElementFromHead(0);
+
+		
+		UE_LOG(LogNetworkSim, Warning, TEXT("Client misprediction. Frame: %d"), ReconciliationKeyframe);
+		/*
+		UE_LOG(LogNetworkSim, Warning, TEXT("  Client: %s"), *ClientMotionState->Location.ToString());
+		UE_LOG(LogNetworkSim, Warning, TEXT("  Server: %s"), *ServerState->Location.ToString());
+		UE_LOG(LogNetworkSim, Warning, TEXT("  LatestClient: %s"), *LatestClient->Location.ToString());
+		*/
+
+		// Copy authoritative state over client state (FIXME: we may want to store this off historically somewhere? Or will VLog be enough for debugging?)
+		*ClientSyncState = *ServerState;
+		
+		// Resimulate up to our latest MotionStateBuffer frame. Note that we may have unprocessed user commands in the command buffer at this point.
+		const int32 LatestKeyframe = SyncBuffer.GetHeadKeyframe();		
+		for (int32 Keyframe = ReconciliationKeyframe+1; Keyframe <= LatestKeyframe; ++Keyframe)
+		{
+			// Keyframe is the frame we are resimulating right now.
+			TInputCmd* ResimulateCmd	= InputBuffer.FindElementByKeyframe(Keyframe);
+			TSyncState* PrevMotionState = SyncBuffer.FindElementByKeyframe(Keyframe - 1);
+			TSyncState* NextMotionState = SyncBuffer.FindElementByKeyframe(Keyframe);
+			
+			check(ResimulateCmd);
+			check(PrevMotionState);
+			check(NextMotionState);
+
+			// TEMP (Aux buffer not fully plumbed through the system yet)
+			// ------------------------------------------------------
+			TAuxState TempAuxState;
+			TAuxState* AuxState = AuxBuffer.FindElementByKeyframe(Keyframe);
+			if (AuxState == nullptr)
+			{
+				AuxState = &TempAuxState;
+			}
+			// ------------------------------------------------------
+
+			// Log out the Mispredicted state that we are about to overwrite.
+			NextMotionState->VisualLog( FVisualLoggingParameters(Keyframe == LatestKeyframe ? EVisualLoggingContext::LastMispredicted : EVisualLoggingContext::OtherMispredicted, Keyframe, EVisualLoggingLifetime::Persistent), Driver, Driver);
+
+			// Do the actual update
+			T::Update(Driver, *ResimulateCmd, *PrevMotionState, *NextMotionState, *AuxState);
+
+			// Log out the newly predicted state that we got.
+			NextMotionState->VisualLog( FVisualLoggingParameters(Keyframe == LatestKeyframe ? EVisualLoggingContext::LastPredicted : EVisualLoggingContext::OtherPredicted, Keyframe, EVisualLoggingLifetime::Persistent), Driver, Driver);
+		}
+	}
+
+private:
+
+	TReplicationBuffer<TInputCmd>& InputBuffer;
+	TReplicationBuffer<TSyncState>& SyncBuffer;
+	TReplicationBuffer<TAuxState>& AuxBuffer;
+	TReplicationBuffer<TSyncState> ReconciliationBuffer;
+
+	int32 LastSerializedKeyframe = -1;
+	bool bPendingReconciliation = false;	// Reconciliation is pending: we need to reconcile state from the server that differs from the locally predicted state
+	bool bReconcileFaultDetected = false;	// A fault was detected: we received state from the server that we are unable to reconcile with locally predicted state
+};
+
+
+
+// ----------------------------------------------------------------------------------------------------------------------------------------------------------------------
+//	TNetworkedSimulationModel
+//	
+//	* Has all logic for "ticking, advancing buffers, calling Update, calling ServerRPC etc
+//	* Doesn't have anything about update component, movesweep, etc
+//	* Concept of "IDriver" which is the owning object that is driving the network sim. This is the interface to the outside UE4 world.
+//	* Has 4 buffers:
+//		-Input: Generated by a client / not the authority.
+//		-Sync: What we are trying to keep in sync. The state that evolves frame to frame with an Update function.
+//		-Aux: State that is also an input into the simulation but does not intrinsically evolve from to frame. Changes to this state can be trapped/tracked/predicted.
+//		-Debug: Replicated buffer from server->client with server-frame centered debug information. Compiled out of shipping builds.
+//
+// ----------------------------------------------------------------------------------------------------------------------------------------------------------------------
+template <
+	typename T,
+	typename TInputCmd,
+	typename TSyncedState,
+	typename TAuxState,
+	typename TDebugState = FNetSimProcessedFrameDebugInfo,
+	typename TRepProxyServerRPC = TReplicator_DynamicSequence<TInputCmd>,
+	typename TRepProxyAutonomous = TReplicator_BasicReconciliar<TInputCmd, TSyncedState, TAuxState>,
+	typename TRepProxySimulated = TReplicator_LatestOnly<TSyncedState>,
+	typename TRepProxyReplay = TReplicator_DynamicSequence<TSyncedState>,
+	typename TRepProxyDebug = TReplicator_DynamicSequence<TDebugState>
+>
+class TNetworkedSimulationModel : public IReplicationProxy
+{
+public:
+
+	class IDriver
+	{
+	public:
+		virtual void InitSyncState(TSyncedState& OutSyncState) const = 0;
+		virtual void SyncTo(const TSyncedState& SyncState) = 0;
+	};
+
+	struct FTickParameters
+	{
+		ENetRole Role;
+		float LocalDeltaTimeSeconds;
+	};	
+
+	TNetworkedSimulationModel()
+		: RepProxy_ServerRPC(InputBuffer), RepProxy_Autonomous(InputBuffer, SyncBuffer, AuxBuffer), RepProxy_Simulated(SyncBuffer), RepProxy_Replay(SyncBuffer)
+#if NETSIM_MODEL_DEBUG
+		, RepProxy_Debug(DebugBuffer)
+#endif
+	{
+	}
+
+	template<typename TDriver>
+	void Tick(TDriver* Driver, const FTickParameters& Parameters)
+	{
+		TDebugState* const DebugState = GetNextDebugStateWrite();
+		if (DebugState)
+		{
+			DebugState->LocalDeltaTimeSeconds = Parameters.LocalDeltaTimeSeconds;
+			DebugState->LocalGFrameNumber = GFrameNumber;
+			DebugState->ProcessedKeyframes.Reset();
+		}
+
+		// How many commands we are allowed to process right now. This will need to be built out a bit more to handle shifting network conditions etc. For now, just a cvar.
+		int32 AllowCmds = NetworkSimulationModelCVars::MaxInputCmdsFrame;
+
+		// -------------------------------------------------------------------------------------------------------------------------------------------------
+		//														Reconciliation
+		// -------------------------------------------------------------------------------------------------------------------------------------------------
+		if (Parameters.Role == ROLE_AutonomousProxy)
+		{
+			// Client: don't allow buffering and ensure continuous stream of inputs. This will need to be built out a bit more. The main points are:
+			// -Don't buffer input locally (unless we explicitly opt into it as a lag hiding measurement). Don't want to accidentally introduce application latency.
+			// -Don't send server input cmds that we didn't actually process. Current Input buffer replication is rather hard coded to send continuous frames and doesn't handle gaps.
+
+			// Recovering from this is possible but messy. We could immediately set things up so that we process the latest user cmd.
+			// The tricky part is making sure we the server does the same thing. We would need to set the tail of the input cmd stream so
+			// that we don't replicate any cmds that were skipped. Then server side it would have to detect this and know to jump ahead.
+			// At this point in development, the complexity this adds to the overall system is not worth the benefit. The better way to
+			// fix this for now is to closely guard what we allow into the client's InputCmdBuffer
+
+			// You can be at head (no new cmds to process) or right before the head (1 new cmd to process)
+			checkf(InputBuffer.GetHeadKeyframe() - LastProcessedInputKeyframe <= 1, TEXT("Client Input Processing is out of sync. LastProcessedInputKeyframe: %d. InputBuffer: [%d - %d]"), LastProcessedInputKeyframe, InputBuffer.GetTailKeyframe(), InputBuffer.GetHeadKeyframe());
+
+			// We should be processing input from this GFrameNumber. If we aren't, then our DeltaTime is probably off (input cmd's delta time wont match the actual delta time of the frame we sim on which will cause artifacts).
+			// We also want to avoid accidentally introducing input latency as a matter of principle (input latency bad)
+			// We *can* support local input buffering in two ways (neither currently implemented):
+			//	-Fixed time steps (delta times will be constant so buffering wont matter)
+			//	-Variable time steps with "interpolation layer": we will need to be able to absorb fluctuations in frame rate so there will be render frames that doesn't process any new input and
+			//		render frames that have to process multiple inputs. (E.g, the goal would be to always "see" an update every render frame, even if we can't actually tick a new sim frame every render frame. Interpolation would be needed to accomplish this)
+			// If either of the above ideas are implemented, it will be done as an optional/templated parameter to the network sim class.
+			checkf(LastLocalInputGFrameNumber == 0 || GFrameNumber == LastLocalInputGFrameNumber, TEXT("TNetworkedSimulationModel running on stale input. You may have an ordering issue between TNetworkedSimulationModel::GetNextInputForWrite and TNetworkedSimulationModel::Tick. "
+				"GFrameNumber: %d LastLocalInputGFrameNumber: %d"), GFrameNumber, LastLocalInputGFrameNumber);
+
+		
+			// Note that is is important for the NetworkSimModel to control *when* the reconcile happens. The replicator objects themselves should never do
+			// reconciliation inside a NetSerialze call. Tick may not be the permanent place for this, but it can't be in NetSerialize.				
+			RepProxy_Autonomous.template Reconcile<T, TDriver>(Driver);
+
+			// Don't process new user commands if we are dealing with a reconciliation fault. That means we weren't able to correlate our predicted position with what the
+			// server last sent us. If we are hitting this, it means we are waiting for the in flight commands to flush through the system. The proxy rep object will handle recovering.
+			if (RepProxy_Autonomous.IsReconcileFaultDetected())
+			{
+				AllowCmds = 0;
+			}
+
+			// Check SyncBuffer being ahead of processed Keyframes. This would happen in cases where we are either not predicting or are buffering our input locally
+			// while sending latest cmds to the server. Essentially, we got the authoritative motion state from the server before we ran the simulation locally.
+			if (SyncBuffer.GetHeadKeyframe() > LastProcessedInputKeyframe)
+			{
+				checkf(InputBuffer.IsValidKeyframe(SyncBuffer.GetHeadKeyframe()), TEXT("MotionState and InputCmd buffers are out of system. LastProcessedInputKeyframe: %d {%s} vs {%s}"),
+					LastProcessedInputKeyframe, *SyncBuffer.GetBasicDebugStr(), *InputBuffer.GetBasicDebugStr());
+
+				UE_LOG(LogNetworkSim, Warning, TEXT("Skipping local input frames because we have newer data in MotionStateBuffer. LastProcessedInputKeyframe: %d. {%s} {%s}"),
+					LastProcessedInputKeyframe, *SyncBuffer.GetBasicDebugStr(), *InputBuffer.GetBasicDebugStr());
+
+				LastProcessedInputKeyframe = SyncBuffer.GetHeadKeyframe();
+			}
+
+			if (NetworkSimulationModelCVars::EnableLocalPrediction == 0)
+			{
+				AllowCmds = 0; // Don't process any commands this frame
+				LastProcessedInputKeyframe = InputBuffer.GetHeadKeyframe(); // Increment so we can accept a new command next frame
+			}
+		}
+
+		if (Parameters.Role == ROLE_Authority)
+		{
+			if ( LastProcessedInputKeyframe+1 < InputBuffer.GetTailKeyframe() )
+			{
+				// We've missed commands
+				UE_LOG(LogNetworkSim, Warning, TEXT("::Tick missing inputcmds. LastProcessedInputKeyframe: %d. %s"), LastProcessedInputKeyframe, *InputBuffer.GetBasicDebugStr());
+				LastProcessedInputKeyframe = InputBuffer.GetTailKeyframe()+1;
+			}
+		}
+
+		// -------------------------------------------------------------------------------------------------------------------------------------------------
+		//														Input Processing
+		// -------------------------------------------------------------------------------------------------------------------------------------------------
+
+		int32 NumProcessed = 0;
+		while (AllowCmds-- > 0)
+		{
+			const int32 Keyframe = LastProcessedInputKeyframe+1;
+			if (TInputCmd* NextCmd = InputBuffer.FindElementByKeyframe(Keyframe))
+			{
+				// The SyncedState buffer needs to be in sync here:
+				//	-We want it to have a SyncedState, but it may not on the first frame through (thats ok).
+				//  -Its HeadKeyframe should be one behind the Keyframe we are about to process.
+				//
+				// Note, InputCmds start @ Keyframe=1. The first SyncedState that Update produces will go in KeyFrame=1.
+				// (E.g, InputCmd @ keyframe=X is used to generate MotionState @ keyframe=X)
+				// This means that SyncedState @ keyframe=0 is always created here via InitSyncState.
+				// This also means that we never actually process InputCmd @ keyframe=0. Which is why LastProcessedInputKeyframe is initialized to 0 ("already processed")
+				// and the buffer has an empty element inserted in InitLocalInputBuffer.
+
+				if (SyncBuffer.GetHeadKeyframe() != LastProcessedInputKeyframe)
+				{
+					if (LastProcessedInputKeyframe != 0)
+					{
+						// This shouldn't happen, but is not fatal. We are reseting the motion state buffer.
+						UE_LOG(LogNetworkSim, Warning, TEXT("Break in SyncState continuity. LastProcessedInputKeyframe: %d. SyncBuffer.GetHeadKeyframe(): %d. Role=%d"), LastProcessedInputKeyframe, SyncBuffer.GetHeadKeyframe(), (int32)Parameters.Role);
+					}
+
+					// We need an initial/current state. Get this from the sim driver
+					SyncBuffer.ResetNextHeadKeyframe(LastProcessedInputKeyframe);
+					TSyncedState* StartingState = SyncBuffer.GetWriteNext();
+					Driver->InitSyncState(*StartingState);
+					
+				}
+
+				TSyncedState* PrevSyncState = SyncBuffer.FindElementByKeyframe(LastProcessedInputKeyframe);
+				TSyncedState* NextSyncState = SyncBuffer.GetWriteNext();
+
+				check(PrevSyncState != nullptr);
+				check(NextSyncState != nullptr);
+				check(SyncBuffer.GetHeadKeyframe() == Keyframe);
+				
+				if (DebugState)
+				{
+					DebugState->ProcessedKeyframes.Add(Keyframe);
+				}
+
+				TAuxState AuxState; // Temp: aux buffer not implemented yet
+
+				T::Update(Driver, *NextCmd, *PrevSyncState, *NextSyncState, AuxState);
+				NumProcessed++;
+
+				LastProcessedInputKeyframe = Keyframe;
+			}
+			else
+			{
+				// Not a warning in itself anymore. Though the concept of "haven't processed cmds in x MS" is something we want to define.
+				break;
+			}
+		}
+
+		// Call into the driver to sync to the latest state if we processed any input
+		if (NumProcessed > 0)
+		{
+			check(SyncBuffer.GetNumValidElements() > 0);
+			Driver->SyncTo(*SyncBuffer.GetElementFromHead(0));
+		}
+
+		// -------------------------------------------------------------------------------------------------------------------------------------------------
+		//														Debug
+		// -------------------------------------------------------------------------------------------------------------------------------------------------
+
+		// Finish debug state buffer recording (what the server processed each frame)
+		if (DebugState)
+		{
+			DebugState->LastProcessedKeyframe = LastProcessedInputKeyframe;
+			DebugState->HeadKeyframe = InputBuffer.GetHeadKeyframe();
+		}
+
+		// Historical data recording (longer buffers for historical reference)
+		if (FHistoricBuffers* HistoricData = GetHistoricBuffers())
+		{
+			HistoricData->InputBuffer.CopyAndMerge(InputBuffer);
+			HistoricData->SyncBuffer.CopyAndMerge(SyncBuffer);
+			HistoricData->AuxBuffer.CopyAndMerge(AuxBuffer);
+		}
+	}	
+	
+	void InitializeForNetworkRole(const ENetRole Role, const bool IsLocallyControleld, const FNetworkSimulationModelInitParameters& Parameters)
+	{
+		InputBuffer.SetBufferSize(Parameters.InputBufferSize);
+		SyncBuffer.SetBufferSize(Parameters.SyncedBufferSize);
+		AuxBuffer.SetBufferSize(Parameters.AuxBufferSize);
+
+		if (GetDebugBuffer())
+		{
+			GetDebugBuffer()->SetBufferSize(Parameters.DebugBufferSize);
+		}
+
+		if (FHistoricBuffers* MyHistoricBuffers = GetHistoricBuffers(true))
+		{
+			MyHistoricBuffers->InputBuffer.SetBufferSize(Parameters.HistoricBufferSize);
+			MyHistoricBuffers->SyncBuffer.SetBufferSize(Parameters.HistoricBufferSize);
+			MyHistoricBuffers->AuxBuffer.SetBufferSize(Parameters.HistoricBufferSize);
+		}
+
+		if (IsLocallyControleld)
+		{
+			InitLocalInputBuffer();
+		}
+	}
+
+	void NetSerializeProxy(EReplicationProxyTarget Target, const FNetSerializeParams& Params)
+	{
+		switch(Target)
+		{
+		case EReplicationProxyTarget::ServerRPC:
+			RepProxy_ServerRPC.NetSerialize(Params);
+			break;
+		case EReplicationProxyTarget::AutonomousProxy:
+			RepProxy_Autonomous.NetSerialize(Params);
+			break;
+		case EReplicationProxyTarget::SimulatedProxy:
+			RepProxy_Simulated.NetSerialize(Params);
+			break;
+		case EReplicationProxyTarget::Replay:
+			RepProxy_Replay.NetSerialize(Params);
+			break;
+		case EReplicationProxyTarget::Debug:
+#if NETSIM_MODEL_DEBUG
+			RepProxy_Debug.NetSerialize(Params);
+			break;
+#endif
+		default:
+			checkf(false, TEXT("Unknown: %d"), (int32)Target);
+		};
+	}
+
+	int32 GetProxyDirtyCount(EReplicationProxyTarget Target)
+	{
+		switch(Target)
+		{
+		case EReplicationProxyTarget::ServerRPC:
+			return RepProxy_ServerRPC.GetProxyDirtyCount();
+		case EReplicationProxyTarget::AutonomousProxy:
+			return RepProxy_Autonomous.GetProxyDirtyCount();
+		case EReplicationProxyTarget::SimulatedProxy:
+			return RepProxy_Simulated.GetProxyDirtyCount();
+		case EReplicationProxyTarget::Replay:
+			return RepProxy_Replay.GetProxyDirtyCount();
+		case EReplicationProxyTarget::Debug:
+#if NETSIM_MODEL_DEBUG
+			return RepProxy_Debug.GetProxyDirtyCount();
+#endif
+		default:
+			checkf(false, TEXT("Unknown: %d"), (int32)Target);
+			return 0;
+		};
+	}
+	
+	TInputCmd* GetNextInputForWrite(float DeltaTime)
+	{
+		ensure(LastLocalInputGFrameNumber != GFrameNumber);
+		LastLocalInputGFrameNumber = GFrameNumber;
+
+		if (InputBuffer.GetHeadKeyframe() == LastProcessedInputKeyframe)
+		{
+			// Only return a cmd if we have processed the last one. This is a bit heavy handed but is a good practice to start with. We want buffering of input to be very explicit, never something that accidently happens.
+			TInputCmd* Next = InputBuffer.GetWriteNext();
+			Next->FrameDeltaTime = DeltaTime;
+			return Next;
+		}
+		return nullptr;
+	}
+
+	void InitLocalInputBuffer()
+	{
+		// Buffer should also be empty before calling this
+		check(InputBuffer.GetHeadKeyframe() == INDEX_NONE);
+
+		// We want to start with an empty command in the input buffer. See notes in input buffer processing function.
+		*InputBuffer.GetWriteNext() = TInputCmd();
+
+		LastLocalInputGFrameNumber = 0;
+	}
+
+	int32 LastProcessedInputKeyframe = 0;
+	uint32 LastLocalInputGFrameNumber = 0;			// Tracks the last time we accepted local input via GetNextInputForWrite. Used to guard against accidental input latency due to ordering of input/sim ticking.
+
+	TReplicationBuffer<TInputCmd> InputBuffer;
+	TReplicationBuffer<TSyncedState> SyncBuffer;
+	TReplicationBuffer<TAuxState> AuxBuffer;
+
+	TRepProxyServerRPC RepProxy_ServerRPC;
+	TRepProxyAutonomous RepProxy_Autonomous;
+	TRepProxySimulated RepProxy_Simulated;
+	TRepProxyReplay RepProxy_Replay;
+
+	// ------------------------------------------------------------------
+	// RPC Sending helper: provides basic send frequency settings for tracking when the Server RPC can be invoked.
+	// Note that the Driver is the one that must call the RPC, that cannot be rolled into this templated structure.
+	// More flexbile/dynamic send rates may be desireable. There is not reason this *has* to be done here, it could
+	// completely be tracked at the driver level, but that will also push more boilerplate to that layer for users.
+	// ------------------------------------------------------------------
+
+	void SetDesiredServerRPCSendFrequency(float DesiredHz) { ServerRPCThresholdTimeSeconds = 1.f / DesiredHz; }
+	bool ShouldSendServerRPC(ENetRole OwnerRole, float DeltaTimeSeconds)
+	{
+		// Don't allow a large delta time to pollute the accumulator
+		const float CappedDeltaTimeSeconds = FMath::Min<float>(DeltaTimeSeconds, ServerRPCThresholdTimeSeconds);
+		if (OwnerRole == ROLE_AutonomousProxy)
+		{
+			ServerRPCAccumulatedTimeSeconds += DeltaTimeSeconds;
+			if (ServerRPCAccumulatedTimeSeconds >= ServerRPCThresholdTimeSeconds)
+			{
+				ServerRPCAccumulatedTimeSeconds -= ServerRPCThresholdTimeSeconds;
+				return true;
+			}
+		}
+		return false;
+	}
+private:
+	float ServerRPCAccumulatedTimeSeconds = 0.f;
+	float ServerRPCThresholdTimeSeconds = 1.f / 999.f; // Default is to send at a max of 999hz. This part of the system needs to be build out more (better handling of super high FPS clients and fixed rate servers)
+
+	// ------------------------------------------------------------------
+	//	Debugging
+	// ------------------------------------------------------------------
+public:
+	struct FHistoricBuffers
+	{
+		TReplicationBuffer<TInputCmd> InputBuffer;
+		TReplicationBuffer<TSyncedState> SyncBuffer;
+		TReplicationBuffer<TAuxState> AuxBuffer;
+	};
+
+#if NETSIM_MODEL_DEBUG
+	TReplicationBuffer<TDebugState>* GetDebugBuffer() {	return &DebugBuffer; }
+	TDebugState* GetNextDebugStateWrite() { return DebugBuffer.GetWriteNext(); }
+	FHistoricBuffers* GetHistoricBuffers(bool bCreate=false)
+	{
+		if (HistoricBuffers.IsValid() == false && bCreate) { HistoricBuffers.Reset(new FHistoricBuffers()); }
+		return HistoricBuffers.Get();
+	}
+#else
+	TReplicationBuffer<TDebugState>* GetDebugBuffer(bool bCreate=false) {	return nullptr; }
+	TDebugState* GetNextDebugStateWrite() { return nullptr; }
+	FHistoricBuffers* GetHistoricBuffers() { return nullptr; }
+#endif
+
+private:
+
+#if NETSIM_MODEL_DEBUG
+	TReplicationBuffer<TDebugState> DebugBuffer;
+	TRepProxyDebug RepProxy_Debug;
+
+	TUniquePtr<FHistoricBuffers> HistoricBuffers;
+#endif
+};
+
+class INetworkedMotionOwnerProxy
+{
+
+};
+
+/** This is the layer that is tied in with mothing a component around a world */
+template <typename T, typename TInputState, typename TSyncedState, typename TAuxState, typename TDebugState>
+class TNetworkedMotionModel : public TNetworkedSimulationModel<T, TInputState, TSyncedState, TAuxState, TDebugState>
+{
+public:
+	TNetworkedMotionModel(INetworkedMotionOwnerProxy* OwnerProxy)
+	{
+	
+	}
+};
+
diff --git a/Engine/Plugins/Runtime/NetworkPrediction/readme.txt b/Engine/Plugins/Runtime/NetworkPrediction/readme.txt
new file mode 100644
--- /dev/null
+++ b/Engine/Plugins/Runtime/NetworkPrediction/readme.txt
@@ -0,0 +1,73 @@
+
+// -------------------------------------------------------------------------------------------------------------------------------------------------
+//	Network Prediction Plugin
+// -------------------------------------------------------------------------------------------------------------------------------------------------
+
+High level overview:
+-This is a *WIP* plugin that aims to generalize and improve our player prediction systems.
+-Long term we hope this replaces UCharacterMovementComponent and its networking system (ServerMove, ClientAck*, ClientAdjust*, etc).
+-The GameplayAbility system will be updated to interop with this as well.
+-NetworkPrediction plugin is the core generalized system that you need to enable to use anything.
+	-For now at least, new movement classes build on top of this will live here but could one day be moved out into a dedicated movement system plugin.
+-NetworkPredictionExtras plugin contains some examples that may be helpful in learning the system. It is not inteded to be used directly in shipping games.
+
+High level goals:
+-Make working with predicted gameplay systems such as character movement easier.
+-Give predictived gameplay systems the ability to interop with each other (E.g, CharacterMovement and Gameplay Abilities).
+-Provide better tooling to understand what is going on in these complex predicted gameplay systems.
+-This all leads to: ability to develop richer and more interactive gameplay systems that work in multiplayer games.
+
+
+Current State (7-25-19)
+-Initial public check in. Core simulation model structure is in place. Lots of pieces missing but overall system is roughed in.
+-Two main examples right now: Flying Movement and 'MockNetworkSimulation' (in NetworkPredictionExtras)
+-Flying Movement: a port of FloatingPawnMovement into the new system. Essentially just a basic flying movement system.
+-MockNetworkSimulation: a very basic minimal "simulation" that demonstrates how the system works. Not tied to movement or anything physical: just a running counter/sum.
+
+Getting Started:
+-Both NetworkPrediction and NetworkPredictionExtras are disabled by default. You will need to enable them for your project (manually edit .uproject or do so through editor plugin screen).
+-Once NetworkPredictionExtras is loaded, /NetworkPredictionExtras/Content/TestMap.umap can be loaded (must enable 'Show Engine Content' AND 'Show Plugin Content' in content browser!).
+-The "MockNetworkSimulation" can be tested anywhere just by typing "mns.Spawn".
+
+Code to look at:
+-MockNetworkSimulation.h: Good place to start to see a bare minimum example of a predicted network simulation. Both the simulation and actor component live in this header.
+-NetworkPredictionExtrasFlyingPawn.h: Starting point to see how a pawn using flying movement is setup. In NetworkPredictionExtras.
+-FlyingMovement.h: The actual implementation of the flying movement system. In NetworkPrediction.
+-NetworkSimulationModelTemplates.h: this is the lowest level guts of the system. This is what every simulation ultimately uses to stay in sync.
+
+Extremely high level technical overview:
+-Concept of "NetworkedSimulationModel": defining a tightly constrained (gameplay) simulation that can be run predictively on clients and authoratatively on server. E.g., a "Movement System" is a NetworkedSimulationModel.
+-The simulation has tightly defined input and output. At the lowest level we have generic buffers of structs that are synchronized:
+	Input Buffer: the data that the controlling client generates
+	Sync Buffer: the data that evolves frame to frame and is produced by ticking the simulation
+	Aux Buffer: additional data that can change, can be predicted, etc but does not necessarily get updated frame to frame (e.g, somethign else usually updates it).
+-Implicit correction: all prediction and correction logic is now client side. Server just tells clients what the authority state was and it is up to them to reconcile.
+-Single code path: the core network simulation has one "Update" function that does not branch based on authority/prediction/simulated etc. "Write once" is the goal for gamplay code here.
+
+UE4 Networking info (How this plugs into UE4 networking)
+-See UNetworkPredictionComponent for the glue.
+-"Replication Proxy" (FReplicationProxy): struct that points to the network simulation and has policies about how it replicates and to who. Custom NetSerialize.
+-E.g, there is a replication proxy for owner, non owners, replays, and debug. All point to the same underlying data.
+-Still using actor replication to determine when to replicate.
+
+
+
+TODO: Major missing elements
+-Aux Buffer is not implemented, but referenced in a few places.
+-"Events" into and out of the system are not in place. That is, something happens in your movement sim that needs to call outside code/generate an event. A system is needed to track this stuff rather than just calling directly into the handler.
+-Interpolation/smoothing layer (visual, non simulation affecting smoothing to help with corrections and simulated proxies)
+-No optimizations (bandwith, cpu) have been done.
+-[Movement specific]: must lock down SetActorLocation/Rotation API so that movement state can be changed out from underneath the system.
+
+Short term road map:
+-The focus right now is on the generalized prediction system, rather than actual movement code.
+-Next will be interopping the Flying Movement example with Ability System (E.g., movement abilities, stamina/fuel attributes, etc).
+-Once everything feels solid we will go deeper on movement and build out a more extendable movement system.
+-E.g, We don't intend all users of the engine to write their own network simulation model movement system.
+-We hope to provide a generalized movement system that is easy to extend in itself without knowing all the details of the NetwokrPrediction plugin.
+-We intend to support RootMotion (Animation) / RootMotionSources (non anim) in the new movement system.
+
+
+
+
+
diff --git a/Engine/Plugins/Runtime/NetworkPredictionExtras/NetworkPredictionExtras.uplugin b/Engine/Plugins/Runtime/NetworkPredictionExtras/NetworkPredictionExtras.uplugin
new file mode 100644
--- /dev/null
+++ b/Engine/Plugins/Runtime/NetworkPredictionExtras/NetworkPredictionExtras.uplugin
@@ -0,0 +1,30 @@
+{
+	"FileVersion": 3,
+	"Version": 1,
+	"VersionName": "0.1",
+	"FriendlyName": "Network Prediction Extras",
+	"Description": "Non essential classes for Network Prediction. Samples, test maps, etc intended to help developers start using teh system. Not intended to be usesd directly in a shipping product.",
+	"Category": "Gameplay",
+	"CreatedBy": "Epic Games, Inc.",
+	"CreatedByURL": "http://epicgames.com",
+	"DocsURL": "",
+	"MarketplaceURL": "",
+	"SupportURL": "",
+	"EnabledByDefault": false,
+	"CanContainContent": true,
+	"IsBetaVersion": true,
+	"Installed": false,
+	"Modules": [
+		{
+			"Name": "NetworkPredictionExtras",
+			"Type": "Runtime",
+			"LoadingPhase": "Default"
+		}
+	],
+	"Plugins": [
+		{
+			"Name": "NetworkPrediction",
+			"Enabled": true
+		}
+	]
+}
\ No newline at end of file
diff --git a/Engine/Plugins/Runtime/NetworkPredictionExtras/Source/NetworkPredictionExtras/NetworkPredictionExtras.Build.cs b/Engine/Plugins/Runtime/NetworkPredictionExtras/Source/NetworkPredictionExtras/NetworkPredictionExtras.Build.cs
new file mode 100644
--- /dev/null
+++ b/Engine/Plugins/Runtime/NetworkPredictionExtras/Source/NetworkPredictionExtras/NetworkPredictionExtras.Build.cs
@@ -0,0 +1,47 @@
+// Copyright 1998-2018 Epic Games, Inc. All Rights Reserved.
+
+namespace UnrealBuildTool.Rules
+{
+	public class NetworkPredictionExtras : ModuleRules
+	{
+		public NetworkPredictionExtras(ReadOnlyTargetRules Target) : base(Target)
+		{
+			PublicIncludePaths.AddRange(
+				new string[] {
+                    ModuleDirectory + "/Public",
+				}
+				);
+
+			PrivateIncludePaths.AddRange(
+				new string[] {
+                    "NetworkPredictionExtras/Private",
+				}
+				);
+
+			PublicDependencyModuleNames.AddRange(
+				new string[]
+				{
+					"NetworkPrediction",
+					"Core",
+                    "CoreUObject",
+                    "Engine",
+                    "RenderCore",
+					"InputCore",
+				}
+				);
+
+			PrivateDependencyModuleNames.AddRange(
+				new string[]
+				{
+				}
+				);
+
+			DynamicallyLoadedModuleNames.AddRange(
+				new string[]
+				{
+				}
+				);
+
+		}
+	}
+}
diff --git a/Engine/Plugins/Runtime/NetworkPredictionExtras/Source/NetworkPredictionExtras/Private/MockNetworkSimulation.cpp b/Engine/Plugins/Runtime/NetworkPredictionExtras/Source/NetworkPredictionExtras/Private/MockNetworkSimulation.cpp
new file mode 100644
--- /dev/null
+++ b/Engine/Plugins/Runtime/NetworkPredictionExtras/Source/NetworkPredictionExtras/Private/MockNetworkSimulation.cpp
@@ -0,0 +1,255 @@
+// Copyright 1998-2019 Epic Games, Inc. All Rights Reserved.
+
+#include "MockNetworkSimulation.h"
+#include "NetworkSimulationModelDebugger.h"
+#include "GameFramework/Actor.h"
+#include "Engine/World.h"
+#include "Components/PrimitiveComponent.h"
+#include "GameFramework/CharacterMovementComponent.h" // Temp
+#include "GameFramework/PlayerController.h"
+#include "Net/UnrealNetwork.h"
+#include "Engine/LocalPlayer.h"
+#include "Misc/CoreDelegates.h"
+#include "UObject/CoreNet.h"
+#include "Engine/NetDriver.h"
+#include "Engine/NetConnection.h"
+#include "UObject/UObjectIterator.h"
+#include "Components/CapsuleComponent.h"
+#include "VisualLogger/VisualLogger.h"
+#include "Math/Color.h"
+#include "DrawDebugHelpers.h"
+#include "Debug/DebugDrawService.h"
+#include "Engine/Canvas.h"
+#include "Engine/Engine.h"
+#include "Debug/ReporterGraph.h"
+#include "EngineUtils.h"
+
+
+DEFINE_LOG_CATEGORY_STATIC(LogMockNetworkSim, Log, All);
+
+namespace MockNetworkSimCVars
+{
+static int32 UseVLogger = 1;
+static FAutoConsoleVariableRef CVarUseVLogger(TEXT("mns.Debug.UseUnrealVLogger"),
+	UseVLogger,	TEXT("Use Unreal Visual Logger\n"),	ECVF_Default);
+
+static int32 UseDrawDebug = 1;
+static FAutoConsoleVariableRef CVarUseDrawDebug(TEXT("mns.Debug.UseDrawDebug"),
+	UseVLogger,	TEXT("Use built in DrawDebug* functions for visual logging\n"), ECVF_Default);
+
+static float DrawDebugDefaultLifeTime = 30.f;
+static FAutoConsoleVariableRef CVarDrawDebugDefaultLifeTime(TEXT("mns.Debug.UseDrawDebug.DefaultLifeTime"),
+	DrawDebugDefaultLifeTime, TEXT("Use built in DrawDebug* functions for visual logging"), ECVF_Default);
+
+// -------------------------------
+
+static int32 DoLocalInput = 0;
+static FAutoConsoleVariableRef CVarDoLocalInput(TEXT("mns.DoLocalInput"),
+	DoLocalInput, TEXT("Submit non 0 imput into the mock simulation"), ECVF_Default);
+
+static int32 RequestMispredict = 0;
+static FAutoConsoleVariableRef CVarRequestMispredict(TEXT("mns.RequestMispredict"),
+	RequestMispredict, TEXT("Causes a misprediction by inserting random value into stream on authority side"), ECVF_Default);
+
+static int32 BindAutomatically = 1;
+static FAutoConsoleVariableRef CVarBindAutomatically(TEXT("mns.BindAutomatically"),
+	BindAutomatically, TEXT("Binds local input and mispredict commands to 5 and 6 respectively"), ECVF_Default);
+}
+
+// ============================================================================================
+
+static bool ForceMispredict = false;
+void FMockNetworkSimulation::Update(IMockNetworkSimulationDriver* Driver, const FMockInputCmd& InputCmd, const FMockSyncState& InputState, FMockSyncState& OutputState, const FMockAuxState& AuxState)
+{
+	OutputState.Total = InputState.Total + (InputCmd.InputValue * AuxState.Multiplier * InputCmd.FrameDeltaTime);
+
+	// Dev hack to force mispredict
+	if (ForceMispredict)
+	{
+		const float Fudge = FMath::FRand() * 100.f;
+		OutputState.Total += Fudge;
+		UE_LOG(LogMockNetworkSim, Warning, TEXT("ForcingMispredict via CVar. Fudge=%.2f. NewTotal: %.2f"), Fudge, OutputState.Total);
+		
+		ForceMispredict = false;
+	}
+}
+
+void FMockSyncState::VisualLog(const FVisualLoggingParameters& Parameters, IMockNetworkSimulationDriver* Driver, IMockNetworkSimulationDriver* LogDriver)
+{
+	// This function is awkward because the mock example doesn't include any positional data, so we added an explicit ::GetDebugWorldTransform to the driver interface.
+	// still, this doesn't give us historical location info. So we will just use Z offset based on the EvisualLoggingContext to space things out.
+	const bool DrawStyleMinor = (Parameters.Context == EVisualLoggingContext::OtherMispredicted || Parameters.Context == EVisualLoggingContext::OtherPredicted);
+	const FTransform WorldTransform = LogDriver->GetDebugWorldTransform(); // No sense using the Driver's world transform (this would be the server's position in PIE, not useful in mock sim)
+	UWorld* DebugWorld = LogDriver->GetDriverWorld();
+
+	FString DebugStr = FString::Printf(TEXT("[%d][%s] Total: %.4f"), Parameters.Keyframe, *LexToString(Parameters.Context), this->Total);
+	if (MockNetworkSimCVars::UseDrawDebug)
+	{
+		const float Lifetime = Parameters.Lifetime == EVisualLoggingLifetime::Transient ? 0.0001f : MockNetworkSimCVars::DrawDebugDefaultLifeTime;
+		const float YOffset = 0.f; //DrawStyleMinor ? ((100.f * (Parameters.Keyframe % 32)) - 50.f) : 0.f;
+		const float ZOffset = 20.f * (int32)Parameters.Context;
+		
+		DrawDebugString( DebugWorld, WorldTransform.GetLocation() + WorldTransform.GetRotation().RotateVector( FVector(0.f, YOffset, ZOffset) ), *DebugStr, nullptr, Parameters.GetDebugColor(), Lifetime );
+	}
+
+	if (MockNetworkSimCVars::UseVLogger)
+	{
+		// I don't see a "Draw string in th world" function in VLOG so we are just doing this
+		UE_VLOG(LogDriver->GetVLogOwner(), LogMockNetworkSim, Log, TEXT("%s"), *DebugStr);
+	}
+}
+
+// ===============================================================================================
+//
+// ===============================================================================================
+
+UMockNetworkSimulationComponent::UMockNetworkSimulationComponent()
+{
+	PrimaryComponentTick.TickGroup = TG_PrePhysics;
+	PrimaryComponentTick.bCanEverTick = true;
+
+	bWantsInitializeComponent = true;
+	bAutoActivate = true;
+	bReplicates = true;
+
+	bWantsInitializeComponent = true;
+
+	// Set default bindings. This is not something you would ever want in a shipping game,
+	// but the optional, latent, non evasive nature of the NetworkPredictionExtras module makes this ideal approach
+	if (HasAnyFlags(RF_ClassDefaultObject) == false && MockNetworkSimCVars::BindAutomatically > 0)
+	{
+		if (ULocalPlayer* LocalPlayer = GetWorld()->GetFirstLocalPlayerFromController())
+		{
+			LocalPlayer->Exec(GetWorld(), TEXT("setbind Five \"mns.DoLocalInput 1 | OnRelease mns.DoLocalInput 0\""), *GLog);
+			LocalPlayer->Exec(GetWorld(), TEXT("setbind Six \"mns.RequestMispredict 1\""), *GLog);
+
+			LocalPlayer->Exec(GetWorld(), TEXT("setbind Nine nms.Debug.LocallyControlledPawn"), *GLog);
+			LocalPlayer->Exec(GetWorld(), TEXT("setbind Zero nms.Debug.ToggleContinous"), *GLog);
+		}
+	}
+}
+
+IReplicationProxy* UMockNetworkSimulationComponent::InstantiateNetworkSimulation()
+{
+	NetworkSim.Reset(new FMockNetworkSimulation());
+
+#if NETSIM_MODEL_DEBUG
+	FNetworkSimulationModelDebuggerManager::Get().RegisterNetworkSimulationModel(NetworkSim.Get(), (IMockNetworkSimulationDriver*)this, GetOwner(), TEXT("MockNetworkSimulation"));
+#endif
+	return NetworkSim.Get();
+}
+
+	// Child classes should override this an initialize their NetworkSim here
+void UMockNetworkSimulationComponent::InitializeForNetworkRole(ENetRole Role)
+{
+	FNetworkSimulationModelInitParameters InitParams;
+	InitParams.InputBufferSize = Role != ROLE_SimulatedProxy ? 32 : 0;
+	InitParams.SyncedBufferSize = Role != ROLE_AutonomousProxy ? 2 : 32;
+	InitParams.AuxBufferSize = Role != ROLE_AutonomousProxy ? 2 : 32;
+	InitParams.DebugBufferSize = 32;
+	InitParams.HistoricBufferSize = 128;
+	
+	NetworkSim->InitializeForNetworkRole(Role, IsLocallyControlled(), InitParams);
+}
+
+void UMockNetworkSimulationComponent::TickComponent(float DeltaTime, enum ELevelTick TickType, FActorComponentTickFunction *ThisTickFunction)
+{
+	CheckOwnerRoleChange();
+
+	const ENetRole OwnerRole = GetOwnerRole();
+
+	// -------------------------------------------------------
+	//	Local Input
+	//		-For the mock sim, just create the input here.
+	//		-Real implementation, the input will probably be set outside of tick, from a PC, InputComponent, etc.
+	// -------------------------------------------------------
+	if (OwnerRole == ROLE_AutonomousProxy)
+	{
+		if (FMockInputCmd* Input = NetworkSim->GetNextInputForWrite(DeltaTime))
+		{
+			if (MockNetworkSimCVars::DoLocalInput)
+			{
+				Input->InputValue = FMath::FRand() * 10.f;
+			}
+			else
+			{
+				Input->InputValue = 0.f;
+			}
+		}
+	}
+
+	// -------------------------------------
+	// Tick the network sim
+	// -------------------------------------
+	if (NetworkSim)
+	{
+		// Check if we should trip a mispredict. (Note how its not possible to do this inside the Update function!)
+		if (OwnerRole == ROLE_Authority && MockNetworkSimCVars::RequestMispredict)
+		{
+			ForceMispredict = true;
+			MockNetworkSimCVars::RequestMispredict = 0;
+		}
+
+		FMockNetworkSimulation::FTickParameters Parameters;
+		Parameters.Role = OwnerRole;
+		Parameters.LocalDeltaTimeSeconds = DeltaTime;
+
+		// Tick the core network sim, this will consume input and generate new sync state
+		NetworkSim->Tick((IMockNetworkSimulationDriver*)this, Parameters);
+
+		// Client->Server communication
+		if (NetworkSim->ShouldSendServerRPC(OwnerRole, DeltaTime))
+		{
+			FServerReplicationRPCParameter ProxyParameter(ReplicationProxy_ServerRPC);
+			ServerRecieveClientInput(ProxyParameter);
+		}
+	}
+	
+	// Mock example of displaying
+	DrawDebugString( GetWorld(), GetOwner()->GetActorLocation() + FVector(0.f,0.f,100.f), *LexToString(MockValue), nullptr, FColor::White, 0.00001f );
+}
+
+void UMockNetworkSimulationComponent::InitSyncState(FMockSyncState& OutSyncState) const
+{
+	OutSyncState.Total = MockValue;
+}
+
+void UMockNetworkSimulationComponent::SyncTo(const FMockSyncState& SyncState)
+{
+	MockValue = SyncState.Total;
+}
+
+FTransform UMockNetworkSimulationComponent::GetDebugWorldTransform() const
+{
+	return GetOwner()->GetActorTransform();
+}
+
+UObject* UMockNetworkSimulationComponent::GetVLogOwner() const
+{
+	return GetOwner();
+}
+
+// ------------------------------------------------------------------------------------------------------------------------
+//
+// ------------------------------------------------------------------------------------------------------------------------
+
+FAutoConsoleCommandWithWorldAndArgs MockNetworkSimulationSpawnCmd(TEXT("mns.Spawn"), TEXT(""),
+	FConsoleCommandWithWorldAndArgsDelegate::CreateLambda([](const TArray< FString >& InArgs, UWorld* World) 
+{
+	bool bFoundWorld = false;
+	for (TObjectIterator<UWorld> It; It; ++It)
+	{
+		if (It->WorldType == EWorldType::PIE && It->GetNetMode() != NM_Client)
+		{
+			bFoundWorld = true;
+			for (TActorIterator<APawn> ActorIt(*It); ActorIt; ++ActorIt)
+			{
+				if (APawn* Pawn = *ActorIt)
+				{
+					UMockNetworkSimulationComponent* NewComponent = NewObject<UMockNetworkSimulationComponent>(Pawn);
+					NewComponent->RegisterComponent();
+				}
+			}
+		}
+	}
+}));
\ No newline at end of file
diff --git a/Engine/Plugins/Runtime/NetworkPredictionExtras/Source/NetworkPredictionExtras/Private/NetworkPredictionExtrasFlyingPawn.cpp b/Engine/Plugins/Runtime/NetworkPredictionExtras/Source/NetworkPredictionExtras/Private/NetworkPredictionExtrasFlyingPawn.cpp
new file mode 100644
--- /dev/null
+++ b/Engine/Plugins/Runtime/NetworkPredictionExtras/Source/NetworkPredictionExtras/Private/NetworkPredictionExtrasFlyingPawn.cpp
@@ -0,0 +1,317 @@
+#include "NetworkPredictionExtrasFlyingPawn.h"
+#include "Components/InputComponent.h"
+#include "Movement/FlyingMovement.h"
+#include "GameFramework/Controller.h"
+#include "GameFramework/PlayerController.h"
+#include "GameFramework/PlayerInput.h"
+#include "GameFramework/SpringArmComponent.h"
+#include "DrawDebugHelpers.h"
+#include "Camera/PlayerCameraManager.h"
+#include "Engine/World.h"
+#include "Engine/LocalPlayer.h"
+
+
+namespace FlyingPawnCVars
+{
+
+int32 CameraStyle = 0;
+FAutoConsoleVariableRef CVarPenetrationPullbackDistance(TEXT("NetworkPredictionExtras.FlyingPawn.CameraSyle"),
+	CameraStyle,
+	TEXT("Sets camera mode style in ANetworkPredictionExtrasFlyingPawn \n")
+	TEXT("0=camera fixed behind pawn. 1-3 are variations of a free camera system (Gamepad recommended)."),
+	ECVF_Default);
+
+static int32 BindAutomatically = 1;
+static FAutoConsoleVariableRef CVarBindAutomatically(TEXT("NetworkPredictionExtras.FlyingPawn.BindAutomatically"),
+	BindAutomatically, TEXT("Binds local input and mispredict commands to 5 and 6 respectively"), ECVF_Default);
+}
+
+ANetworkPredictionExtrasFlyingPawn::ANetworkPredictionExtrasFlyingPawn()
+{
+	FlyingMovementComponent = CreateDefaultSubobject<UFlyingMovementComponent>(TEXT("FlyingMovementComponent"));
+	
+	if (HasAnyFlags(RF_ClassDefaultObject) == false)
+	{
+		// Setup callback for generating local input. There is some nuance here. This is the order we want things to happen for a locally controlled actor:
+		//	1. Create local input (the lambda below)
+		//	2. Tick the network sim (called by UFlyingMovementComponent::TickComponent). This wants the latest input (important the user cmd has the same DeltaTime as the current frame!)
+		//  3. Tick this actor (ANetworkPredictionExtrasFlyingPawn::Tick). This wants the latest movement state (computed in step #2).
+		//
+		// The ::SetLocallyControlledPreTick isn't strictly necessary. A player controller or some other object could tick first, submit the input, then let the sim tick and this actor tick.
+		// But for this simple example, its clearer to contain everything here. The lambda is just a helper to save you from dealing with tick graph prereq headaches.
+
+		FlyingMovementComponent->SetLocallyControlledPreTick(UNetworkPredictionComponent::FPreSimTickDelegate::CreateUObject(this, &ThisClass::GenerateLocalInput));
+
+		// Binds 0 and 9 to the debug hud commands. This is just a convenience for the extras plugin. Real projects should bind this themselves
+		if (FlyingPawnCVars::BindAutomatically > 0)
+		{
+			if (ULocalPlayer* LocalPlayer = GetWorld()->GetFirstLocalPlayerFromController())
+			{
+				LocalPlayer->Exec(GetWorld(), TEXT("setbind Nine nms.Debug.LocallyControlledPawn"), *GLog);
+				LocalPlayer->Exec(GetWorld(), TEXT("setbind Zero nms.Debug.ToggleContinous"), *GLog);
+			}
+		}
+	}
+
+
+}
+
+void ANetworkPredictionExtrasFlyingPawn::SetupPlayerInputComponent(UInputComponent* PlayerInputComponent)
+{
+	// Setup some bindings.
+	//
+	// This also is not necessary. This would usually be defined in DefaultInput.ini.
+	// But the nature of NetworkPredictionExtras is to not assume or require anything of the project its being run in.
+	// So in order to make this nice and self contained, manually do input bindings here.
+	if (APlayerController* PC = Cast<APlayerController>(Controller))
+	{
+		if (PC->PlayerInput)
+		{
+			//Gamepad
+			PC->PlayerInput->AddAxisMapping(FInputAxisKeyMapping(TEXT("MoveRight"), EKeys::Gamepad_LeftX));
+			PC->PlayerInput->AddAxisMapping(FInputAxisKeyMapping(TEXT("MoveForward"), EKeys::Gamepad_LeftY));
+			PC->PlayerInput->AddAxisMapping(FInputAxisKeyMapping(TEXT("LookYaw"), EKeys::Gamepad_RightX));
+			PC->PlayerInput->AddAxisMapping(FInputAxisKeyMapping(TEXT("LookPitch"), EKeys::Gamepad_RightY));
+			PC->PlayerInput->AddAxisMapping(FInputAxisKeyMapping(TEXT("LeftTriggerAxis"), EKeys::Gamepad_LeftTriggerAxis));
+			PC->PlayerInput->AddAxisMapping(FInputAxisKeyMapping(TEXT("RightTriggerAxis"), EKeys::Gamepad_RightTriggerAxis));
+			PC->PlayerInput->AddActionMapping(FInputActionKeyMapping(TEXT("LeftShoulder"), EKeys::Gamepad_LeftShoulder));
+			PC->PlayerInput->AddActionMapping(FInputActionKeyMapping(TEXT("RightShoulder"), EKeys::Gamepad_RightShoulder));
+
+			// Keyboard
+			PC->PlayerInput->AddAxisMapping(FInputAxisKeyMapping(TEXT("MoveRight"), EKeys::D, 1.f));
+			PC->PlayerInput->AddAxisMapping(FInputAxisKeyMapping(TEXT("MoveRight"), EKeys::A, -1.f));
+			PC->PlayerInput->AddAxisMapping(FInputAxisKeyMapping(TEXT("MoveForward"), EKeys::W, 1.f));
+			PC->PlayerInput->AddAxisMapping(FInputAxisKeyMapping(TEXT("MoveForward"), EKeys::S, -1.f));
+			PC->PlayerInput->AddAxisMapping(FInputAxisKeyMapping(TEXT("LeftTriggerAxis"), EKeys::LeftControl, 1.f));
+			PC->PlayerInput->AddAxisMapping(FInputAxisKeyMapping(TEXT("RightTriggerAxis"), EKeys::LeftShift, 1.f));
+
+			// Mouse
+			PC->PlayerInput->AddAxisMapping(FInputAxisKeyMapping(TEXT("LookYaw"), EKeys::MouseX));
+			PC->PlayerInput->AddAxisMapping(FInputAxisKeyMapping(TEXT("LookPitch"), EKeys::MouseY));
+		}
+	}
+
+	PlayerInputComponent->BindAxis(TEXT("MoveForward"), this,	&ThisClass::InputAxis_MoveForward);
+	PlayerInputComponent->BindAxis(TEXT("MoveRight"), this,		&ThisClass::InputAxis_MoveRight);
+	PlayerInputComponent->BindAxis(TEXT("LookYaw"), this,		&ThisClass::InputAxis_LookYaw);
+	PlayerInputComponent->BindAxis(TEXT("LookPitch"), this,		&ThisClass::InputAxis_LookPitch);
+	PlayerInputComponent->BindAxis(TEXT("RightTriggerAxis"), this,		&ThisClass::InputAxis_MoveUp);
+	PlayerInputComponent->BindAxis(TEXT("LeftTriggerAxis"), this,		&ThisClass::InputAxis_MoveDown);
+
+	PlayerInputComponent->BindAction(TEXT("LeftShoulder"), IE_Pressed, this, &ThisClass::Action_LeftShoulder_Pressed);
+	PlayerInputComponent->BindAction(TEXT("LeftShoulder"), IE_Released, this, &ThisClass::Action_LeftShoulder_Released);
+	PlayerInputComponent->BindAction(TEXT("RightShoulder"), IE_Pressed, this, &ThisClass::Action_RightShoulder_Pressed);
+	PlayerInputComponent->BindAction(TEXT("RightShoulder"), IE_Released, this, &ThisClass::Action_RightShoulder_Released);
+}
+
+void ANetworkPredictionExtrasFlyingPawn::InputAxis_MoveForward(float Value)
+{
+	CachedMoveInput.X = FMath::Clamp(Value, -1.0f, 1.0f);
+}
+
+void ANetworkPredictionExtrasFlyingPawn::InputAxis_MoveRight(float Value)
+{
+	CachedMoveInput.Y = FMath::Clamp(Value, -1.0f, 1.0f);
+}
+
+void ANetworkPredictionExtrasFlyingPawn::InputAxis_LookYaw(float Value)
+{
+	CachedLookInput.X = FMath::Clamp(Value, -1.0f, 1.0f);
+}
+
+void ANetworkPredictionExtrasFlyingPawn::InputAxis_LookPitch(float Value)
+{
+	CachedLookInput.Y = FMath::Clamp(Value, -1.0f, 1.0f);
+}
+
+void ANetworkPredictionExtrasFlyingPawn::InputAxis_MoveUp(float Value)
+{
+	CachedMoveInput.Z = FMath::Clamp(CachedMoveInput.Z + Value, -1.0f, 1.0f);
+}
+
+void ANetworkPredictionExtrasFlyingPawn::InputAxis_MoveDown(float Value)
+{
+	CachedMoveInput.Z = FMath::Clamp(CachedMoveInput.Z - Value, -1.0f, 1.0f);
+}
+
+void ANetworkPredictionExtrasFlyingPawn::Tick( float DeltaSeconds)
+{
+	Super::Tick(DeltaSeconds);
+
+	// Do whatever you want here. By now we have the latest movement state and latest input processed.
+}
+
+void ANetworkPredictionExtrasFlyingPawn::GenerateLocalInput(float DeltaSeconds)
+{
+	// Generate user commands. Called right before the flying movement simulation will tick (for a locally controlled pawn)
+	// This isn't meant to be the best way of doing a camera system. It is just meant to show a couple of ways it may be done
+	// and to make sure we can keep distinct the movement, rotation, and view angles.
+	// Change with CVar NetworkPredictionExtras.FlyingPawn.CameraSyle. Styles 1-3 are really meant to be used with a gamepad.
+	//
+	// Its worth calling out: the code that happens here is happening *outside* of the flying movement simulation. All we are doing
+	// is generating the input being fed into that simulation. That said, this means that A) the code below does not run on the server
+	// (and non controlling clients) and B) the code is not rerun during reconcile/resimulates. This this information guide any
+	// decisions about where something should go (such as aim assist, lock on targeting systems, etc).
+
+	check(Controller);
+	if (FlyingMovement::FInputCmd* NextCmd = FlyingMovementComponent->GetNextClientInputCmdForWrite(DeltaSeconds))
+	{
+		if (USpringArmComponent* SpringComp = FindComponentByClass<USpringArmComponent>())
+		{
+			// This is not best practice: do not search for component every frame
+			SpringComp->bUsePawnControlRotation = true;
+		}
+
+		// Simple input scaling. A real game will probably map this to an acceleration curve
+		static float LookRateYaw = 150.f;
+		static float LookRatePitch = 150.f;
+
+		static float ControllerLookRateYaw = 1.5f;
+		static float ControllerLookRatePitch = 1.5f;
+
+		// Zero out input structs in case each path doesnt set each member. This is really all we are filling out here.
+		NextCmd->MovementInput = FVector::ZeroVector;
+		NextCmd->RotationInput = FRotator::ZeroRotator;
+
+		switch (FlyingPawnCVars::CameraStyle)
+		{
+			case 0:
+			{
+				// Fixed camera
+				if (USpringArmComponent* SpringComp = FindComponentByClass<USpringArmComponent>())
+				{
+					// Only this camera mode has to set this to false
+					SpringComp->bUsePawnControlRotation = false;
+				}
+
+				NextCmd->RotationInput.Yaw = CachedLookInput.X * LookRateYaw;
+				NextCmd->RotationInput.Pitch = CachedLookInput.Y * LookRatePitch;
+				NextCmd->RotationInput.Roll = 0;
+					
+				NextCmd->MovementInput = CachedMoveInput;
+				break;
+			}
+			case 1:
+			{
+				// Free camera Restricted 2D movement on XY plane.
+				APlayerController* PC = Cast<APlayerController>(Controller);
+				if (ensure(PC)) // Requires player controller for now
+				{
+					// Camera yaw rotation
+					PC->AddYawInput(CachedLookInput.X * ControllerLookRateYaw );
+					PC->AddPitchInput(CachedLookInput.Y * ControllerLookRatePitch );
+
+					static float RotationMagMin = (1e-3);
+						
+					float RotationInputYaw = 0.f;
+					if (CachedMoveInput.Size() >= RotationMagMin)
+					{
+						FVector DesiredMovementDir = PC->GetControlRotation().RotateVector(CachedMoveInput);
+
+						// 2D xy movement, relative to camera
+						{
+							FVector DesiredMoveDir2D = FVector(DesiredMovementDir.X, DesiredMovementDir.Y, 0.f);
+							DesiredMoveDir2D.Normalize();
+
+							const float DesiredYaw = DesiredMoveDir2D.Rotation().Yaw;
+							const float DeltaYaw = DesiredYaw - GetActorRotation().Yaw;
+							NextCmd->RotationInput.Yaw = DeltaYaw / NextCmd->FrameDeltaTime;
+						}
+					}
+
+					NextCmd->MovementInput = FVector(FMath::Clamp<float>( CachedMoveInput.Size2D(), -1.f, 1.f ), 0.0f, CachedMoveInput.Z);
+				}
+				break;
+			}
+			case 2:
+			{
+				// Free camera on yaw and pitch, camera-relative movement.
+				APlayerController* PC = Cast<APlayerController>(Controller);
+				if (ensure(PC)) // Requires player controller for now
+				{
+					// Camera yaw rotation
+					PC->AddYawInput(CachedLookInput.X * ControllerLookRateYaw );
+					PC->AddPitchInput(CachedLookInput.Y * ControllerLookRatePitch );
+
+					// Rotational movement: orientate us towards our camera-relative desired velocity (unless we are upside down, then flip it)
+					static float RotationMagMin = (1e-3);
+					const float MoveInputMag = CachedMoveInput.Size();
+						
+					float RotationInputYaw = 0.f;
+					float RotationInputPitch = 0.f;
+					float RotationInputRoll = 0.f;
+
+					if (MoveInputMag >= RotationMagMin)
+					{								
+						FVector DesiredMovementDir = PC->GetControlRotation().RotateVector(CachedMoveInput);
+
+						const float DesiredYaw = DesiredMovementDir.Rotation().Yaw;
+						const float DeltaYaw = DesiredYaw - GetActorRotation().Yaw;
+
+						const float DesiredPitch = DesiredMovementDir.Rotation().Pitch;
+						const float DeltaPitch = DesiredPitch - GetActorRotation().Pitch;
+
+						const float DesiredRoll = DesiredMovementDir.Rotation().Roll;
+						const float DeltaRoll = DesiredRoll - GetActorRotation().Roll;
+
+						// Kind of gross but because we want "instant" turning we must factor in delta time so that it gets factored out inside the simulation
+						RotationInputYaw = DeltaYaw / NextCmd->FrameDeltaTime;
+						RotationInputPitch = DeltaPitch / NextCmd->FrameDeltaTime;
+						RotationInputRoll = DeltaRoll / NextCmd->FrameDeltaTime;
+								
+					}
+
+					NextCmd->RotationInput.Yaw = RotationInputYaw;
+					NextCmd->RotationInput.Pitch = RotationInputPitch;
+					NextCmd->RotationInput.Roll = RotationInputRoll;
+						
+					NextCmd->MovementInput = FVector(FMath::Clamp<float>( MoveInputMag, -1.f, 1.f ), 0.0f, 0.0f); // Not the best way but simple
+				}
+				break;
+			}
+			case 3:
+			{
+				// Free camera on the yaw, camera-relative motion
+				APlayerController* PC = Cast<APlayerController>(Controller);
+				if (ensure(PC)) // Requires player controller for now
+				{
+					// Camera yaw rotation
+					PC->AddYawInput(CachedLookInput.X * ControllerLookRateYaw );
+
+					// Rotational movement: orientate us towards our camera-relative desired velocity (unless we are upside down, then flip it)
+					static float RotationMagMin = (1e-3);
+					const float MoveInputMag = CachedMoveInput.Size2D();
+						
+					float RotationInputYaw = 0.f;
+					if (MoveInputMag >= RotationMagMin)
+					{
+						FVector DesiredMovementDir = PC->GetControlRotation().RotateVector(CachedMoveInput);
+
+						const bool bIsUpsideDown = FVector::DotProduct(FVector(0.f, 0.f, 1.f), GetActorQuat().GetUpVector() ) < 0.f;
+						if (bIsUpsideDown)
+						{
+							DesiredMovementDir *= -1.f;
+						}							
+
+						const float DesiredYaw = DesiredMovementDir.Rotation().Yaw;
+						const float DeltaYaw = DesiredYaw - GetActorRotation().Yaw;
+
+						// Kind of gross but because we want "instant" turning we must factor in delta time so that it gets factored out inside the simulation
+						RotationInputYaw = DeltaYaw / NextCmd->FrameDeltaTime;
+					}
+
+					NextCmd->RotationInput.Yaw = RotationInputYaw;
+					NextCmd->RotationInput.Pitch = CachedLookInput.Y * LookRatePitch; // Just pitch like normal
+					NextCmd->RotationInput.Roll = 0;
+						
+					NextCmd->MovementInput = FVector(FMath::Clamp<float>( MoveInputMag, -1.f, 1.f ), 0.0f, CachedMoveInput.Z); // Not the best way but simple
+				}
+				break;
+			}
+		}
+				
+
+		CachedMoveInput = FVector::ZeroVector;
+		CachedLookInput = FVector2D::ZeroVector;
+	}
+}
\ No newline at end of file
diff --git a/Engine/Plugins/Runtime/NetworkPredictionExtras/Source/NetworkPredictionExtras/Private/NetworkPredictionExtrasModule.cpp b/Engine/Plugins/Runtime/NetworkPredictionExtras/Source/NetworkPredictionExtras/Private/NetworkPredictionExtrasModule.cpp
new file mode 100644
--- /dev/null
+++ b/Engine/Plugins/Runtime/NetworkPredictionExtras/Source/NetworkPredictionExtras/Private/NetworkPredictionExtrasModule.cpp
@@ -0,0 +1,28 @@
+// Copyright 1998-2019 Epic Games, Inc. All Rights Reserved.
+
+#include "NetworkPredictionExtrasModule.h"
+#include "CoreMinimal.h"
+#include "Modules/ModuleManager.h"
+#include "Internationalization/Internationalization.h"
+
+#define LOCTEXT_NAMESPACE "FNetworkPredictionModule"
+
+class FNetworkPredictionExtrasModule : public INetworkPredictionExtrasModule
+{
+	/** IModuleInterface implementation */
+	virtual void StartupModule() override;
+	virtual void ShutdownModule() override;
+};
+
+void FNetworkPredictionExtrasModule::StartupModule()
+{
+}
+
+
+void FNetworkPredictionExtrasModule::ShutdownModule()
+{
+}
+
+IMPLEMENT_MODULE( FNetworkPredictionExtrasModule, NetworkPredictionExtras )
+#undef LOCTEXT_NAMESPACE
+
diff --git a/Engine/Plugins/Runtime/NetworkPredictionExtras/Source/NetworkPredictionExtras/Public/MockNetworkSimulation.h b/Engine/Plugins/Runtime/NetworkPredictionExtras/Source/NetworkPredictionExtras/Public/MockNetworkSimulation.h
new file mode 100644
--- /dev/null
+++ b/Engine/Plugins/Runtime/NetworkPredictionExtras/Source/NetworkPredictionExtras/Public/MockNetworkSimulation.h
@@ -0,0 +1,173 @@
+// Copyright 1998-2019 Epic Games, Inc. All Rights Reserved.
+
+#pragma once
+
+#include "CoreMinimal.h"
+#include "UObject/ObjectMacros.h"
+#include "UObject/UObjectGlobals.h"
+#include "Engine/NetSerialization.h"
+#include "Engine/EngineTypes.h"
+#include "Engine/EngineBaseTypes.h"
+#include "WorldCollision.h"
+#include "Components/ActorComponent.h"
+#include "Components/SceneComponent.h"
+#include "Misc/OutputDevice.h"
+#include "Misc/CoreDelegates.h"
+
+#include "NetworkSimulationModelTemplates.h"
+#include "NetworkPredictionComponent.h"
+
+#include "MockNetworkSimulation.generated.h"
+
+class IMockNetworkSimulationDriver;
+
+// -------------------------------------------------------------------------------------------------------------------------------
+//	Mock Network Simulation
+//
+//	This provides a minimal "mock" example of using the Network Prediction system. The simulation being run by these classes
+//	is a simple accumulator that takes random numbers (FMockInputCmd::InputValue) as input. There is no movement related functionality
+//	in this example. This is just meant to show the bare minimum hook ups into the system to make it easier to understand.
+//
+//	Highlights:
+//		FMockNetworkSimulation::Update			The "core update" function of the simulation.
+//		UMockNetworkSimulationComponent			The UE4 Actor Component that anchors the system to an actor.
+//
+//	Usage:
+//		You can just add a UMockNetworkSimulationComponent to any ROLE_AutonomousProxy actor yourself (Default Subobject, through blueprints, manually, etc)
+//		The console command "mns.Spawn" can be used to dynamically spawn the component on every pawn. Must be run on the server or in single process PIE.
+//
+//	Once spawned, there are some useful console commands that can be used. These bind to number keys by default (toggleable via mns.BindAutomatically).
+//		[Five] "mns.DoLocalInput 1"			can be used to submit random local input into the accumulator. This is how you advance the simulation.
+//		[Six]  "mns.RequestMispredict 1"	can be used to force a mispredict (random value added to accumulator server side). Useful for tracing through the correction/resimulate code path.
+//		[Nine] "nms.Debug.LocallyControlledPawn"	toggle debug hud for the locally controlled player.
+//		[Zero] "nms.Debug.ToggleContinous"			Toggles continuous vs snapshotted display of the debug hud
+//
+//	Notes:
+//		Everything is crammed into MockNetworkSimulation.h/.cpp. It may make sense to break the simulation and component code into
+//		separate files for more complex simulations.
+//
+// -------------------------------------------------------------------------------------------------------------------------------
+
+// State the client generates
+struct FMockInputCmd
+{
+	float FrameDeltaTime=0.f;
+	float InputValue=0;
+
+	void NetSerialize(const FNetSerializeParams& P)
+	{
+		P.Ar << FrameDeltaTime;
+		P.Ar << InputValue;
+	}
+
+	void Log(FStandardLoggingParameters& P) const
+	{
+		if (P.Context == EStandardLoggingContext::Full)
+		{
+			P.Ar->Logf(TEXT("FrameDeltaTime: %4f"), FrameDeltaTime);
+			P.Ar->Logf(TEXT("InputValue: %4f"), InputValue);
+		}
+		else
+		{
+			P.Ar->Logf(TEXT("%2f"), InputValue);
+		}
+	}
+};
+
+// State we are evolving frame to frame and keeping in sync
+struct FMockSyncState
+{
+	float Total=0;
+	void NetSerialize(const FNetSerializeParams& P)
+	{
+		P.Ar << Total;
+	}
+
+	// Compare this state with AuthorityState. return true if a reconcile (correction) should happen
+	bool ShouldReconcile(const FMockSyncState& AuthorityState)
+	{
+		return FMath::Abs<float>(Total - AuthorityState.Total) > SMALL_NUMBER;
+	}
+
+	void Log(FStandardLoggingParameters& P) const
+	{
+		if (P.Context == EStandardLoggingContext::Full)
+		{
+			P.Ar->Logf(TEXT("Total: %4f"), Total);
+		}
+		else
+		{
+			P.Ar->Logf(TEXT("%2f"), Total);
+		}
+	}
+
+	void VisualLog(const FVisualLoggingParameters& Parameters, IMockNetworkSimulationDriver* Driver, IMockNetworkSimulationDriver* LogDriver);
+};
+
+// Auxiliary state that is input into the simulation. Doesn't change during the simulation tick.
+// (It can change and even be predicted but doing so will trigger more bookeeping, etc. Changes will happen "next tick")
+// NOTE: Currently incomplete!
+struct FMockAuxState
+{
+	float Multiplier=1;
+	void NetSerialize(const FNetSerializeParams& P)
+	{
+		P.Ar << Multiplier;
+	}
+};
+
+// Actual definition of our network simulation.
+class FMockNetworkSimulation : public TNetworkedSimulationModel<FMockNetworkSimulation, FMockInputCmd, FMockSyncState, FMockAuxState>
+{
+public:
+
+	// Interface between the simulation and owning component driving it. Functions added here are available in ::Update and must be implemented by UMockNetworkSimulationComponent.
+	class IMockDriver : public IDriver
+	{
+	public:
+		virtual UWorld* GetDriverWorld() const = 0;
+		virtual UObject* GetVLogOwner() const = 0;
+		virtual FTransform GetDebugWorldTransform() const = 0;
+	};
+
+	/** Main update function */
+	static void Update(IMockNetworkSimulationDriver* Driver, const FMockInputCmd& InputCmd, const FMockSyncState& InputState, FMockSyncState& OutputState, const FMockAuxState& AuxState);
+};
+
+// Needed to trick UHT into letting UMockNetworkSimulationComponent implement. UHT cannot parse the ::
+// Also needed for forward declaring. Can't just be a typedef/using =
+class IMockNetworkSimulationDriver : public FMockNetworkSimulation::IMockDriver { };
+
+// -------------------------------------------------------------------------------------------------------------------------------
+// ActorComponent for running a MockNetworkSimulation
+// -------------------------------------------------------------------------------------------------------------------------------
+
+UCLASS(BlueprintType, meta=(BlueprintSpawnableComponent))
+class NETWORKPREDICTIONEXTRAS_API UMockNetworkSimulationComponent : public UNetworkPredictionComponent, public IMockNetworkSimulationDriver
+{
+	GENERATED_BODY()
+
+public:
+
+	UMockNetworkSimulationComponent();
+
+	virtual void TickComponent(float DeltaTime, enum ELevelTick TickType, FActorComponentTickFunction *ThisTickFunction) override;
+	
+	virtual IReplicationProxy* InstantiateNetworkSimulation() override;
+	virtual void InitializeForNetworkRole(ENetRole Role) override;
+
+private:
+
+	// The Network sim that this component is managing. This is what is doing all the work.
+	TUniquePtr<FMockNetworkSimulation> NetworkSim;
+	
+public:
+	void InitSyncState(FMockSyncState& OutSyncState) const override;
+	void SyncTo(const FMockSyncState& SyncState) override;
+	virtual UWorld* GetDriverWorld() const override final { return GetWorld(); }
+	virtual UObject* GetVLogOwner() const override final;
+	virtual FTransform GetDebugWorldTransform() const override final;
+
+	// Mock representation of "syncing' to the sync state in the network sim.
+	float MockValue = 1000.f;
+};
\ No newline at end of file
diff --git a/Engine/Plugins/Runtime/NetworkPredictionExtras/Source/NetworkPredictionExtras/Public/NetworkPredictionExtrasFlyingPawn.h b/Engine/Plugins/Runtime/NetworkPredictionExtras/Source/NetworkPredictionExtras/Public/NetworkPredictionExtrasFlyingPawn.h
new file mode 100644
--- /dev/null
+++ b/Engine/Plugins/Runtime/NetworkPredictionExtras/Source/NetworkPredictionExtras/Public/NetworkPredictionExtrasFlyingPawn.h
@@ -0,0 +1,66 @@
+// Copyright 1998-2019 Epic Games, Inc. All Rights Reserved.
+#pragma once
+
+#include "GameFramework/Pawn.h"
+#include "NetworkPredictionExtrasFlyingPawn.generated.h"
+
+class UInputComponent;
+class UFlyingMovementComponent;
+
+// -------------------------------------------------------------------------------------------------------------------------------
+//	ANetworkPredictionExtrasFlyingPawn
+//
+//	This provides a minimal pawn class that uses UFlyingMovementCompnent. This isn't really intended to be used in shipping games,
+//	rather just to serve as standalone example of using the system, contained completely in the NetworkPredictionExtras plugin.
+//	Apart from the most basic glue/setup, this class provides an example of turning UE4 input event callbacks into the input commands
+//	that are used by the flying movement simulation. This includes some basic camera/aiming code.
+//
+//	Highlights:
+//		FlyingMovement::FMovementSystem::Update						The "core update" function of the flying movement simulation.
+//		ANetworkPredictionExtrasFlyingPawn::GenerateLocalInput		Function that generates local input commands that are fed into the movement system.
+//
+//	Usage:
+//		You should be able to just use this pawn like you would any other pawn. You can specify it as your pawn class in your game mode, or manually override in world settings, etc.
+//		Alternatively, you can just load the NetworkPredictionExtras/Content/TestMap.umap which will have everything setup.
+//
+//	Once spawned, there are some useful console commands:
+//		NetworkPredictionExtras.FlyingPawn.CameraSyle [0-3]			Changes camera mode style.
+//		nms.Debug.LocallyControlledPawn 1							Enables debug hud. binds to '9' by default, see ANetworkPredictionExtrasFlyingPawn()
+//		nms.Debug.ToggleContinous 1									Toggles continuous updates of the debug hud. binds to '0' by default, see ANetworkPredictionExtrasFlyingPawn()
+//
+// -------------------------------------------------------------------------------------------------------------------------------
+
+
+/** Sample pawn that uses UFlyingMovementComponent. The main thing this provides is actually producing user input for the component/simulation to consume. */
+UCLASS(config = Game)
+class NETWORKPREDICTIONEXTRAS_API ANetworkPredictionExtrasFlyingPawn : public APawn
+{
+	GENERATED_BODY()
+
+	ANetworkPredictionExtrasFlyingPawn();
+
+	virtual void SetupPlayerInputComponent(UInputComponent* PlayerInputComponent) override;
+	virtual void Tick( float DeltaSeconds) override;
+
+private:
+
+	void GenerateLocalInput(float DeltaSeconds);
+
+	FVector CachedMoveInput;
+	FVector2D CachedLookInput;
+
+	void InputAxis_MoveForward(float Value);
+	void InputAxis_MoveRight(float Value);
+	void InputAxis_LookYaw(float Value);
+	void InputAxis_LookPitch(float Value);
+	void InputAxis_MoveUp(float Value);
+	void InputAxis_MoveDown(float Value);
+
+	void Action_LeftShoulder_Pressed() { }
+	void Action_LeftShoulder_Released() { }
+	void Action_RightShoulder_Pressed() { }
+	void Action_RightShoulder_Released() { }
+
+	UPROPERTY()
+	UFlyingMovementComponent* FlyingMovementComponent;
+};
\ No newline at end of file
diff --git a/Engine/Plugins/Runtime/NetworkPredictionExtras/Source/NetworkPredictionExtras/Public/NetworkPredictionExtrasModule.h b/Engine/Plugins/Runtime/NetworkPredictionExtras/Source/NetworkPredictionExtras/Public/NetworkPredictionExtrasModule.h
new file mode 100644
--- /dev/null
+++ b/Engine/Plugins/Runtime/NetworkPredictionExtras/Source/NetworkPredictionExtras/Public/NetworkPredictionExtrasModule.h
@@ -0,0 +1,26 @@
+// Copyright 1998-2018 Epic Games, Inc. All Rights Reserved.
+
+#pragma once
+
+#include "CoreMinimal.h"
+#include "Modules/ModuleInterface.h"
+#include "Modules/ModuleManager.h"
+
+
+
+class INetworkPredictionExtrasModule : public IModuleInterface
+{
+
+public:
+	
+	static inline INetworkPredictionExtrasModule& Get()
+	{
+		return FModuleManager::LoadModuleChecked< INetworkPredictionExtrasModule >( "NetworkExtrasPrediction" );
+	}
+	
+	static inline bool IsAvailable()
+	{
+		return FModuleManager::Get().IsModuleLoaded( "NetworkExtrasPrediction" );
+	}
+};
+
