NetworkPrediction WIP
* Added BaseMovementComponent to facilitate sharing common code between movement simulations. This is mostly the component management and moving code.

* Added ParametricMovementComponent/System. This is another basic movement system for simple parametric motion. We are using this to begin looking at forward predicting (a deliberate, subset of) simulated proxy actors along with an autonomous proxy controlled simulation. This is still a large WIP and is not fully functioning the way we want yet.

#rnx

diff --git a/Engine/Plugins/Runtime/NetworkPrediction/Source/NetworkPrediction/Private/Movement/BaseMovementComponent.cpp b/Engine/Plugins/Runtime/NetworkPrediction/Source/NetworkPrediction/Private/Movement/BaseMovementComponent.cpp
new file mode 100644
--- /dev/null
+++ b/Engine/Plugins/Runtime/NetworkPrediction/Source/NetworkPrediction/Private/Movement/BaseMovementComponent.cpp
@@ -0,0 +1,434 @@
+// Copyright 1998-2019 Epic Games, Inc. All Rights Reserved.
+
+#include "Movement/BaseMovementComponent.h"
+#include "GameFramework/Actor.h"
+#include "Engine/World.h"
+#include "Components/PrimitiveComponent.h"
+#include "Components/CapsuleComponent.h"
+#include "DrawDebugHelpers.h"
+#include "VisualLogger/VisualLogger.h"
+#include "NetworkSimulationModel.h"
+
+
+DEFINE_LOG_CATEGORY_STATIC(LogBaseMovement, Log, All);
+
+
+namespace BaseMovementCVars
+{
+
+static float PenetrationPullbackDistance = 0.125f;
+static FAutoConsoleVariableRef CVarPenetrationPullbackDistance(TEXT("bm.PenetrationPullbackDistance"),
+	PenetrationPullbackDistance,
+	TEXT("Pull out from penetration of an object by this extra distance.\n")
+	TEXT("Distance added to penetration fix-ups."),
+	ECVF_Default);
+
+static float PenetrationOverlapCheckInflation = 0.100f;
+static FAutoConsoleVariableRef CVarPenetrationOverlapCheckInflation(TEXT("bm.PenetrationOverlapCheckInflation"),
+	PenetrationOverlapCheckInflation,
+	TEXT("Inflation added to object when checking if a location is free of blocking collision.\n")
+	TEXT("Distance added to inflation in penetration overlap check."),
+	ECVF_Default);
+
+static int32 RequestMispredict = 0;
+static FAutoConsoleVariableRef CVarRequestMispredict(TEXT("bm.RequestMispredict"),
+	RequestMispredict, TEXT("Causes a misprediction by inserting random value into stream on authority side"), ECVF_Default);
+
+
+static int32 UseVLogger = 1;
+static FAutoConsoleVariableRef CVarUseVLogger(TEXT("bm.Debug.UseUnrealVLogger"),
+	UseVLogger,	TEXT("Use Unreal Visual Logger\n"),	ECVF_Default);
+
+static int32 UseDrawDebug = 1;
+static FAutoConsoleVariableRef CVarUseDrawDebug(TEXT("bm.Debug.UseDrawDebug"),
+	UseVLogger,	TEXT("Use built in DrawDebug* functions for visual logging\n"), ECVF_Default);
+
+static float DrawDebugDefaultLifeTime = 30.f;
+static FAutoConsoleVariableRef CVarDrawDebugDefaultLifeTime(TEXT("bm.Debug.UseDrawDebug.DefaultLifeTime"),
+	DrawDebugDefaultLifeTime, TEXT("Use built in DrawDebug* functions for visual logging"), ECVF_Default);
+}
+
+// ----------------------------------------------------------------------------------------------------------
+//	UBaseMovementComponent setup/init
+// ----------------------------------------------------------------------------------------------------------
+
+UBaseMovementComponent::UBaseMovementComponent()
+{
+	PrimaryComponentTick.TickGroup = TG_PrePhysics;
+	PrimaryComponentTick.bCanEverTick = true;
+	
+	bWantsInitializeComponent = true;
+	bAutoActivate = true;
+	bReplicates = true;
+	bWantsInitializeComponent = true;
+}
+
+void UBaseMovementComponent::RegisterComponentTickFunctions(bool bRegister)
+{
+	Super::RegisterComponentTickFunctions(bRegister);
+
+	// Super may start up the tick function when we don't want to.
+	UpdateTickRegistration();
+
+	// If the owner ticks, make sure we tick first. This is to ensure the owner's location will be up to date when it ticks.
+	AActor* Owner = GetOwner();
+	
+	if (bRegister && PrimaryComponentTick.bCanEverTick && Owner && Owner->CanEverTick())
+	{
+		Owner->PrimaryActorTick.AddPrerequisite(this, PrimaryComponentTick);
+	}
+}
+
+void UBaseMovementComponent::UpdateTickRegistration()
+{
+	const bool bHasUpdatedComponent = (UpdatedComponent != NULL);
+	SetComponentTickEnabled(bHasUpdatedComponent && bAutoActivate);
+}
+
+void UBaseMovementComponent::InitializeComponent()
+{
+	TGuardValue<bool> InInitializeComponentGuard(bInInitializeComponent, true);
+	Super::InitializeComponent();
+
+	// RootComponent is null in OnRegister for blueprint (non-native) root components.
+	if (!UpdatedComponent)
+	{
+		// Auto-register owner's root component if found.
+		if (AActor* MyActor = GetOwner())
+		{
+			if (USceneComponent* NewUpdatedComponent = MyActor->GetRootComponent())
+			{
+				SetUpdatedComponent(NewUpdatedComponent);
+			}
+		}
+	}
+}
+
+void UBaseMovementComponent::OnRegister()
+{
+	TGuardValue<bool> InOnRegisterGuard(bInOnRegister, true);
+
+	UpdatedPrimitive = Cast<UPrimitiveComponent>(UpdatedComponent);
+	Super::OnRegister();
+
+	const UWorld* MyWorld = GetWorld();
+	if (MyWorld && MyWorld->IsGameWorld())
+	{
+		USceneComponent* NewUpdatedComponent = UpdatedComponent;
+		if (!UpdatedComponent)
+		{
+			// Auto-register owner's root component if found.
+			AActor* MyActor = GetOwner();
+			if (MyActor)
+			{
+				NewUpdatedComponent = MyActor->GetRootComponent();
+			}
+		}
+
+		SetUpdatedComponent(NewUpdatedComponent);
+	}
+}
+
+void UBaseMovementComponent::SetUpdatedComponent(USceneComponent* NewUpdatedComponent)
+{
+	if (UpdatedComponent && UpdatedComponent != NewUpdatedComponent)
+	{
+		UpdatedComponent->SetShouldUpdatePhysicsVolume(false);
+		if (!UpdatedComponent->IsPendingKill())
+		{
+			UpdatedComponent->SetPhysicsVolume(NULL, true);
+			UpdatedComponent->PhysicsVolumeChangedDelegate.RemoveDynamic(this, &UBaseMovementComponent::PhysicsVolumeChanged);
+		}
+
+		// remove from tick prerequisite
+		UpdatedComponent->PrimaryComponentTick.RemovePrerequisite(this, PrimaryComponentTick); 
+	}
+
+	// Don't assign pending kill components, but allow those to null out previous UpdatedComponent.
+	UpdatedComponent = IsValid(NewUpdatedComponent) ? NewUpdatedComponent : NULL;
+	UpdatedPrimitive = Cast<UPrimitiveComponent>(UpdatedComponent);
+
+	// Assign delegates
+	if (UpdatedComponent && !UpdatedComponent->IsPendingKill())
+	{
+		UpdatedComponent->SetShouldUpdatePhysicsVolume(true);
+		UpdatedComponent->PhysicsVolumeChangedDelegate.AddUniqueDynamic(this, &UBaseMovementComponent::PhysicsVolumeChanged);
+
+		if (!bInOnRegister && !bInInitializeComponent)
+		{
+			// UpdateOverlaps() in component registration will take care of this.
+			UpdatedComponent->UpdatePhysicsVolume(true);
+		}
+		
+		// force ticks after movement component updates
+		UpdatedComponent->PrimaryComponentTick.AddPrerequisite(this, PrimaryComponentTick); 
+	}
+	UpdateTickRegistration();
+}
+
+void UBaseMovementComponent::PhysicsVolumeChanged(APhysicsVolume* NewVolume)
+{
+	// This itself feels bad. When will this be called? Its impossible to know what is allowed and not allowed to be done in this callback.
+	// Callbacks instead should be trapped within the simulation update function. This isn't really possible though since the UpdateComponent
+	// is the one that will call this.
+}
+
+// ----------------------------------------------------------------------------------------------------------
+//	Movement System Driver
+//
+//	NOTE: Most of the Movement Driver is not ideal! We are at the mercy of the UpdateComponent since it is the
+//	the object that owns its collision data and its MoveComponent function. Ideally we would have everything within
+//	the movement simulation code and it do its own collision queries. But instead we have to come back to the Driver/Component
+//	layer to do this kind of stuff.
+//
+// ----------------------------------------------------------------------------------------------------------
+
+FVector UBaseMovementComponent::GetPenetrationAdjustment(const FHitResult& Hit)
+{
+	if (!Hit.bStartPenetrating)
+	{
+		return FVector::ZeroVector;
+	}
+
+	FVector Result;
+	const float PullBackDistance = FMath::Abs(BaseMovementCVars::PenetrationPullbackDistance);
+	const float PenetrationDepth = (Hit.PenetrationDepth > 0.f ? Hit.PenetrationDepth : 0.125f);
+
+	Result = Hit.Normal * (PenetrationDepth + PullBackDistance);
+	
+	return Result;
+}
+
+bool UBaseMovementComponent::OverlapTest(const FVector& Location, const FQuat& RotationQuat, const ECollisionChannel CollisionChannel, const FCollisionShape& CollisionShape, const AActor* IgnoreActor) const
+{
+	FCollisionQueryParams QueryParams(SCENE_QUERY_STAT(MovementOverlapTest), false, IgnoreActor);
+	FCollisionResponseParams ResponseParam;
+	InitCollisionParams(QueryParams, ResponseParam);
+	return GetWorld()->OverlapBlockingTestByChannel(Location, RotationQuat, CollisionChannel, CollisionShape, QueryParams, ResponseParam);
+}
+
+void UBaseMovementComponent::InitCollisionParams(FCollisionQueryParams &OutParams, FCollisionResponseParams& OutResponseParam) const
+{
+	if (UpdatedPrimitive)
+	{
+		UpdatedPrimitive->InitSweepCollisionParams(OutParams, OutResponseParam);
+	}
+}
+
+bool UBaseMovementComponent::ResolvePenetration(const FVector& ProposedAdjustment, const FHitResult& Hit, const FQuat& NewRotationQuat) const
+{
+	// SceneComponent can't be in penetration, so this function really only applies to PrimitiveComponent.
+	const FVector Adjustment = ProposedAdjustment; //ConstrainDirectionToPlane(ProposedAdjustment);
+	if (!Adjustment.IsZero() && UpdatedPrimitive)
+	{
+		QUICK_SCOPE_CYCLE_COUNTER(STAT_BaseMovementComponent_ResolvePenetration);
+		// See if we can fit at the adjusted location without overlapping anything.
+		AActor* ActorOwner = UpdatedComponent->GetOwner();
+		if (!ActorOwner)
+		{
+			return false;
+		}
+
+		UE_LOG(LogBaseMovement, Verbose, TEXT("ResolvePenetration: %s.%s at location %s inside %s.%s at location %s by %.3f (netmode: %d)"),
+			   *ActorOwner->GetName(),
+			   *UpdatedComponent->GetName(),
+			   *UpdatedComponent->GetComponentLocation().ToString(),
+			   *GetNameSafe(Hit.GetActor()),
+			   *GetNameSafe(Hit.GetComponent()),
+			   Hit.Component.IsValid() ? *Hit.GetComponent()->GetComponentLocation().ToString() : TEXT("<unknown>"),
+			   Hit.PenetrationDepth,
+			   (uint32)GetNetMode());
+
+		// We really want to make sure that precision differences or differences between the overlap test and sweep tests don't put us into another overlap,
+		// so make the overlap test a bit more restrictive.
+		const float OverlapInflation = BaseMovementCVars::PenetrationOverlapCheckInflation;
+		bool bEncroached = OverlapTest(Hit.TraceStart + Adjustment, NewRotationQuat, UpdatedPrimitive->GetCollisionObjectType(), UpdatedPrimitive->GetCollisionShape(OverlapInflation), ActorOwner);
+		if (!bEncroached)
+		{
+			// Move without sweeping.
+			MoveUpdatedComponent(Adjustment, NewRotationQuat, false, nullptr, ETeleportType::TeleportPhysics);
+			UE_LOG(LogBaseMovement, Verbose, TEXT("ResolvePenetration:   teleport by %s"), *Adjustment.ToString());
+			return true;
+		}
+		else
+		{
+			// Disable MOVECOMP_NeverIgnoreBlockingOverlaps if it is enabled, otherwise we wouldn't be able to sweep out of the object to fix the penetration.
+			TGuardValue<EMoveComponentFlags> ScopedFlagRestore(MoveComponentFlags, EMoveComponentFlags(MoveComponentFlags & (~MOVECOMP_NeverIgnoreBlockingOverlaps)));
+
+			// Try sweeping as far as possible...
+			FHitResult SweepOutHit(1.f);
+			bool bMoved = MoveUpdatedComponent(Adjustment, NewRotationQuat, true, &SweepOutHit, ETeleportType::TeleportPhysics);
+			UE_LOG(LogBaseMovement, Verbose, TEXT("ResolvePenetration:   sweep by %s (success = %d)"), *Adjustment.ToString(), bMoved);
+			
+			// Still stuck?
+			if (!bMoved && SweepOutHit.bStartPenetrating)
+			{
+				// Combine two MTD results to get a new direction that gets out of multiple surfaces.
+				const FVector SecondMTD = GetPenetrationAdjustment(SweepOutHit);
+				const FVector CombinedMTD = Adjustment + SecondMTD;
+				if (SecondMTD != Adjustment && !CombinedMTD.IsZero())
+				{
+					bMoved = MoveUpdatedComponent(CombinedMTD, NewRotationQuat, true, nullptr, ETeleportType::TeleportPhysics);
+					UE_LOG(LogBaseMovement, Verbose, TEXT("ResolvePenetration:   sweep by %s (MTD combo success = %d)"), *CombinedMTD.ToString(), bMoved);
+				}
+			}
+
+			// Still stuck?
+			if (!bMoved)
+			{
+				// Try moving the proposed adjustment plus the attempted move direction. This can sometimes get out of penetrations with multiple objects
+				const FVector MoveDelta = (Hit.TraceEnd - Hit.TraceStart); //ConstrainDirectionToPlane(Hit.TraceEnd - Hit.TraceStart);
+				if (!MoveDelta.IsZero())
+				{
+					bMoved = MoveUpdatedComponent(Adjustment + MoveDelta, NewRotationQuat, true, nullptr, ETeleportType::TeleportPhysics);
+					UE_LOG(LogBaseMovement, Verbose, TEXT("ResolvePenetration:   sweep by %s (adjusted attempt success = %d)"), *(Adjustment + MoveDelta).ToString(), bMoved);
+				}
+			}	
+
+			return bMoved;
+		}
+	}
+
+	return false;
+}
+
+bool UBaseMovementComponent::SafeMoveUpdatedComponent(const FVector& Delta, const FQuat& NewRotation, bool bSweep, FHitResult& OutHit, ETeleportType Teleport) const
+{
+	if (UpdatedComponent == NULL)
+	{
+		OutHit.Reset(1.f);
+		return false;
+	}
+
+	bool bMoveResult = false;
+
+	// Scope for move flags
+	{
+		bMoveResult = MoveUpdatedComponent(Delta, NewRotation, bSweep, &OutHit, Teleport);
+	}
+
+	// Handle initial penetrations
+	if (OutHit.bStartPenetrating && UpdatedComponent)
+	{
+		const FVector RequestedAdjustment = GetPenetrationAdjustment(OutHit);
+		if (ResolvePenetration(RequestedAdjustment, OutHit, NewRotation))
+		{
+			// Retry original move
+			bMoveResult = MoveUpdatedComponent(Delta, NewRotation, bSweep, &OutHit, Teleport);
+		}
+	}
+
+	return bMoveResult;
+}
+
+bool UBaseMovementComponent::MoveUpdatedComponent(const FVector& Delta, const FQuat& NewRotation, bool bSweep, FHitResult* OutHit, ETeleportType Teleport) const
+{
+	if (UpdatedComponent)
+	{
+		const FVector NewDelta = Delta;
+		return UpdatedComponent->MoveComponent(NewDelta, NewRotation, bSweep, OutHit, MoveComponentFlags, Teleport);
+	}
+
+	return false;
+}
+
+FTransform UBaseMovementComponent::GetUpdateComponentTransform() const
+{
+	if (ensure(UpdatedComponent))
+	{
+		return UpdatedComponent->GetComponentTransform();		
+	}
+	return FTransform::Identity;
+}
+
+UObject* UBaseMovementComponent::GetVLogOwner() const
+{
+	return GetOwner();
+}
+
+void UBaseMovementComponent::DrawDebug(const IBaseMovementDriver::FDrawDebugParams& Params) const
+{
+	if (!UpdatedComponent)
+	{
+		return;
+	}
+
+	AActor* Owner = GetOwner();
+
+	const bool PersistentLines = (Params.Lifetime == EVisualLoggingLifetime::Persistent);
+	const float LifetimeSeconds = PersistentLines ? 20.f : -1.f;
+
+	if (Params.DrawType == EVisualLoggingDrawType::Crumb)
+	{
+		if (BaseMovementCVars::UseDrawDebug)
+		{
+			static float PointSize = 3.f;
+			DrawDebugPoint(Params.DebugWorld, Params.Transform.GetLocation(), PointSize, Params.DrawColor, PersistentLines, LifetimeSeconds);
+		}
+
+		if (BaseMovementCVars::UseVLogger)
+		{
+			static float CrumbRadius = 3.f;
+			UE_VLOG_LOCATION(Params.DebugLogOwner, LogBaseMovement, Log, Params.Transform.GetLocation(), CrumbRadius, Params.DrawColor, TEXT("%s"), *Params.InWorldText);
+		}
+	}
+	else
+	{
+		// Full drawing. Trying to provide a generic implementation that will give the most accurate representation.
+		// Still, subclasses may want to customize this
+		if (UCapsuleComponent* CapsuleComponent = Cast<UCapsuleComponent>(UpdatedComponent))
+		{
+			float Radius=0.f;
+			float HalfHeight=0.f;
+			CapsuleComponent->GetScaledCapsuleSize(Radius, HalfHeight);
+
+			if (BaseMovementCVars::UseDrawDebug)
+			{
+				static const float Thickness = 2.f;
+				DrawDebugCapsule(Params.DebugWorld, Params.Transform.GetLocation(), HalfHeight, Radius, Params.Transform.GetRotation(), Params.DrawColor, PersistentLines, LifetimeSeconds, 0.f, Thickness);
+			}
+
+			if (BaseMovementCVars::UseVLogger)
+			{
+				FVector VLogPosition = Params.Transform.GetLocation();
+				VLogPosition.Z -= HalfHeight;
+				UE_VLOG_CAPSULE(Params.DebugLogOwner, LogBaseMovement, Log, VLogPosition, HalfHeight, Radius, Params.Transform.GetRotation(), Params.DrawColor, TEXT("%s"), *Params.InWorldText);
+			}
+		}
+		else
+		{
+			// Generic Actor Bounds drawing
+			FBox LocalSpaceBox = Owner->CalculateComponentsBoundingBoxInLocalSpace();
+
+			if (BaseMovementCVars::UseDrawDebug)
+			{
+				static const float Thickness = 2.f;
+
+				FVector ActorOrigin;
+				FVector ActorExtent;
+				LocalSpaceBox.GetCenterAndExtents(ActorOrigin, ActorExtent);
+				ActorExtent *= Params.Transform.GetScale3D();
+				DrawDebugBox(Params.DebugWorld, Params.Transform.GetLocation(), ActorExtent, Params.Transform.GetRotation(), Params.DrawColor, PersistentLines, LifetimeSeconds, 0, Thickness);
+			}
+
+			if (BaseMovementCVars::UseVLogger)
+			{
+				UE_VLOG_OBOX(Params.DebugLogOwner, LogBaseMovement, Log, LocalSpaceBox, Params.Transform.ToMatrixWithScale(), Params.DrawColor, TEXT("%s"), *Params.InWorldText);
+			}
+		}
+	}
+
+	if (BaseMovementCVars::UseVLogger)
+	{
+		UE_VLOG(Params.DebugLogOwner, LogBaseMovement, Log, TEXT("%s"), *Params.LogText);
+	}
+}
+
+IBaseMovementDriver::FDrawDebugParams::FDrawDebugParams(const FVisualLoggingParameters& Parameters, IBaseMovementDriver* LogDriver)
+{
+	DebugWorld = LogDriver->GetDriverWorld();
+	DebugLogOwner = LogDriver->GetVLogOwner();
+	DrawType = (Parameters.Context == EVisualLoggingContext::OtherMispredicted || Parameters.Context == EVisualLoggingContext::OtherPredicted) ? EVisualLoggingDrawType::Crumb : EVisualLoggingDrawType::Full;
+	Lifetime = Parameters.Lifetime;
+	DrawColor = Parameters.GetDebugColor();	
+}
\ No newline at end of file
diff --git a/Engine/Plugins/Runtime/NetworkPrediction/Source/NetworkPrediction/Private/Movement/FlyingMovementComponent.cpp b/Engine/Plugins/Runtime/NetworkPrediction/Source/NetworkPrediction/Private/Movement/FlyingMovementComponent.cpp
--- a/Engine/Plugins/Runtime/NetworkPrediction/Source/NetworkPrediction/Private/Movement/FlyingMovementComponent.cpp
+++ b/Engine/Plugins/Runtime/NetworkPrediction/Source/NetworkPrediction/Private/Movement/FlyingMovementComponent.cpp
@@ -55,124 +55,7 @@ static FAutoConsoleVariableRef CVarRequestMispredict(TEXT("fp.RequestMispredict"
 
 UFlyingMovementComponent::UFlyingMovementComponent()
 {
-	PrimaryComponentTick.TickGroup = TG_PrePhysics;
-	PrimaryComponentTick.bCanEverTick = true;
-	
-	bWantsInitializeComponent = true;
-	bAutoActivate = true;
-	bReplicates = true;
-
-	bWantsInitializeComponent = true;
-}
-
-void UFlyingMovementComponent::RegisterComponentTickFunctions(bool bRegister)
-{
-	Super::RegisterComponentTickFunctions(bRegister);
-
-	// Super may start up the tick function when we don't want to.
-	UpdateTickRegistration();
-
-	// If the owner ticks, make sure we tick first. This is to ensure the owner's location will be up to date when it ticks.
-	AActor* Owner = GetOwner();
-	
-	if (bRegister && PrimaryComponentTick.bCanEverTick && Owner && Owner->CanEverTick())
-	{
-		Owner->PrimaryActorTick.AddPrerequisite(this, PrimaryComponentTick);
-	}
-}
-
-void UFlyingMovementComponent::UpdateTickRegistration()
-{
-	const bool bHasUpdatedComponent = (UpdatedComponent != NULL);
-	SetComponentTickEnabled(bHasUpdatedComponent && bAutoActivate);
-}
-
-void UFlyingMovementComponent::InitializeComponent()
-{
-	TGuardValue<bool> InInitializeComponentGuard(bInInitializeComponent, true);
-	Super::InitializeComponent();
-
-	// RootComponent is null in OnRegister for blueprint (non-native) root components.
-	if (!UpdatedComponent)
-	{
-		// Auto-register owner's root component if found.
-		if (AActor* MyActor = GetOwner())
-		{
-			if (USceneComponent* NewUpdatedComponent = MyActor->GetRootComponent())
-			{
-				SetUpdatedComponent(NewUpdatedComponent);
-			}
-		}
-	}
-}
-
-void UFlyingMovementComponent::OnRegister()
-{
-	TGuardValue<bool> InOnRegisterGuard(bInOnRegister, true);
-
-	UpdatedPrimitive = Cast<UPrimitiveComponent>(UpdatedComponent);
-	Super::OnRegister();
-
-	const UWorld* MyWorld = GetWorld();
-	if (MyWorld && MyWorld->IsGameWorld())
-	{
-		USceneComponent* NewUpdatedComponent = UpdatedComponent;
-		if (!UpdatedComponent)
-		{
-			// Auto-register owner's root component if found.
-			AActor* MyActor = GetOwner();
-			if (MyActor)
-			{
-				NewUpdatedComponent = MyActor->GetRootComponent();
-			}
-		}
 
-		SetUpdatedComponent(NewUpdatedComponent);
-	}
-}
-
-void UFlyingMovementComponent::SetUpdatedComponent(USceneComponent* NewUpdatedComponent)
-{
-	if (UpdatedComponent && UpdatedComponent != NewUpdatedComponent)
-	{
-		UpdatedComponent->SetShouldUpdatePhysicsVolume(false);
-		if (!UpdatedComponent->IsPendingKill())
-		{
-			UpdatedComponent->SetPhysicsVolume(NULL, true);
-			UpdatedComponent->PhysicsVolumeChangedDelegate.RemoveDynamic(this, &UFlyingMovementComponent::PhysicsVolumeChanged);
-		}
-
-		// remove from tick prerequisite
-		UpdatedComponent->PrimaryComponentTick.RemovePrerequisite(this, PrimaryComponentTick); 
-	}
-
-	// Don't assign pending kill components, but allow those to null out previous UpdatedComponent.
-	UpdatedComponent = IsValid(NewUpdatedComponent) ? NewUpdatedComponent : NULL;
-	UpdatedPrimitive = Cast<UPrimitiveComponent>(UpdatedComponent);
-
-	// Assign delegates
-	if (UpdatedComponent && !UpdatedComponent->IsPendingKill())
-	{
-		UpdatedComponent->SetShouldUpdatePhysicsVolume(true);
-		UpdatedComponent->PhysicsVolumeChangedDelegate.AddUniqueDynamic(this, &UFlyingMovementComponent::PhysicsVolumeChanged);
-
-		if (!bInOnRegister && !bInInitializeComponent)
-		{
-			// UpdateOverlaps() in component registration will take care of this.
-			UpdatedComponent->UpdatePhysicsVolume(true);
-		}
-		
-		// force ticks after movement component updates
-		UpdatedComponent->PrimaryComponentTick.AddPrerequisite(this, PrimaryComponentTick); 
-	}
-	UpdateTickRegistration();
-}
-
-void UFlyingMovementComponent::PhysicsVolumeChanged(APhysicsVolume* NewVolume)
-{
-	// This itself feels bad. When will this be called? Its impossible to know what is allowed and not allowed to be done in this callback.
-	// Callbacks instead should be trapped within the simulation update function. This isn't really possible though since the UpdateComponent
-	// is the one that will call this.
 }
 
 // ----------------------------------------------------------------------------------------------------------
@@ -189,18 +72,10 @@ IReplicationProxy* UFlyingMovementComponent::InstantiateNetworkSimulation()
 	return NetworkSim.Get();
 }
 
-	// Child classes should override this an initialize their NetworkSim here
 void UFlyingMovementComponent::InitializeForNetworkRole(ENetRole Role)
-{
+{	
 	check(NetworkSim);
-
-	FNetworkSimulationModelInitParameters InitParams;
-	InitParams.InputBufferSize = Role != ROLE_SimulatedProxy ? 32 : 0;
-	InitParams.SyncedBufferSize = Role != ROLE_AutonomousProxy ? 2 : 32;
-	InitParams.AuxBufferSize = Role != ROLE_AutonomousProxy ? 2 : 32;
-	InitParams.DebugBufferSize = 32;
-	InitParams.HistoricBufferSize = 128;
-	NetworkSim->InitializeForNetworkRole(Role, IsLocallyControlled(), InitParams);
+	NetworkSim->InitializeForNetworkRole(Role, IsLocallyControlled(), GetSimulationInitParameters(Role));
 }
 
 void UFlyingMovementComponent::TickComponent(float DeltaTime, enum ELevelTick TickType, FActorComponentTickFunction* ThisTickFunction)
@@ -213,9 +88,10 @@ void UFlyingMovementComponent::TickComponent(float DeltaTime, enum ELevelTick Ti
 			if (OldComp->IsActive())
 			{
 				OldComp->Deactivate();
-				Owner->bReplicateMovement = false;
 			}
 		}
+
+		Owner->bReplicateMovement = false;
 	}
 
 
@@ -283,163 +159,6 @@ FlyingMovement::FInputCmd* UFlyingMovementComponent::GetNextClientInputCmdForWri
 //
 // ----------------------------------------------------------------------------------------------------------
 
-FVector UFlyingMovementComponent::GetPenetrationAdjustment(const FHitResult& Hit)
-{
-	if (!Hit.bStartPenetrating)
-	{
-		return FVector::ZeroVector;
-	}
-
-	FVector Result;
-	const float PullBackDistance = FMath::Abs(FlyingMovementCVars::PenetrationPullbackDistance);
-	const float PenetrationDepth = (Hit.PenetrationDepth > 0.f ? Hit.PenetrationDepth : 0.125f);
-
-	Result = Hit.Normal * (PenetrationDepth + PullBackDistance);
-	
-	return Result;
-}
-
-bool UFlyingMovementComponent::OverlapTest(const FVector& Location, const FQuat& RotationQuat, const ECollisionChannel CollisionChannel, const FCollisionShape& CollisionShape, const AActor* IgnoreActor) const
-{
-	FCollisionQueryParams QueryParams(SCENE_QUERY_STAT(MovementOverlapTest), false, IgnoreActor);
-	FCollisionResponseParams ResponseParam;
-	InitCollisionParams(QueryParams, ResponseParam);
-	return GetWorld()->OverlapBlockingTestByChannel(Location, RotationQuat, CollisionChannel, CollisionShape, QueryParams, ResponseParam);
-}
-
-void UFlyingMovementComponent::InitCollisionParams(FCollisionQueryParams &OutParams, FCollisionResponseParams& OutResponseParam) const
-{
-	if (UpdatedPrimitive)
-	{
-		UpdatedPrimitive->InitSweepCollisionParams(OutParams, OutResponseParam);
-	}
-}
-
-bool UFlyingMovementComponent::ResolvePenetration(const FVector& ProposedAdjustment, const FHitResult& Hit, const FQuat& NewRotationQuat) const
-{
-	// SceneComponent can't be in penetration, so this function really only applies to PrimitiveComponent.
-	const FVector Adjustment = ProposedAdjustment; //ConstrainDirectionToPlane(ProposedAdjustment);
-	if (!Adjustment.IsZero() && UpdatedPrimitive)
-	{
-		QUICK_SCOPE_CYCLE_COUNTER(STAT_FlyingMovementComponent_ResolvePenetration);
-		// See if we can fit at the adjusted location without overlapping anything.
-		AActor* ActorOwner = UpdatedComponent->GetOwner();
-		if (!ActorOwner)
-		{
-			return false;
-		}
-
-		UE_LOG(LogFlyingMovement, Verbose, TEXT("ResolvePenetration: %s.%s at location %s inside %s.%s at location %s by %.3f (netmode: %d)"),
-			   *ActorOwner->GetName(),
-			   *UpdatedComponent->GetName(),
-			   *UpdatedComponent->GetComponentLocation().ToString(),
-			   *GetNameSafe(Hit.GetActor()),
-			   *GetNameSafe(Hit.GetComponent()),
-			   Hit.Component.IsValid() ? *Hit.GetComponent()->GetComponentLocation().ToString() : TEXT("<unknown>"),
-			   Hit.PenetrationDepth,
-			   (uint32)GetNetMode());
-
-		// We really want to make sure that precision differences or differences between the overlap test and sweep tests don't put us into another overlap,
-		// so make the overlap test a bit more restrictive.
-		const float OverlapInflation = FlyingMovementCVars::PenetrationOverlapCheckInflation;
-		bool bEncroached = OverlapTest(Hit.TraceStart + Adjustment, NewRotationQuat, UpdatedPrimitive->GetCollisionObjectType(), UpdatedPrimitive->GetCollisionShape(OverlapInflation), ActorOwner);
-		if (!bEncroached)
-		{
-			// Move without sweeping.
-			MoveUpdatedComponent(Adjustment, NewRotationQuat, false, nullptr, ETeleportType::TeleportPhysics);
-			UE_LOG(LogFlyingMovement, Verbose, TEXT("ResolvePenetration:   teleport by %s"), *Adjustment.ToString());
-			return true;
-		}
-		else
-		{
-			// Disable MOVECOMP_NeverIgnoreBlockingOverlaps if it is enabled, otherwise we wouldn't be able to sweep out of the object to fix the penetration.
-			TGuardValue<EMoveComponentFlags> ScopedFlagRestore(MoveComponentFlags, EMoveComponentFlags(MoveComponentFlags & (~MOVECOMP_NeverIgnoreBlockingOverlaps)));
-
-			// Try sweeping as far as possible...
-			FHitResult SweepOutHit(1.f);
-			bool bMoved = MoveUpdatedComponent(Adjustment, NewRotationQuat, true, &SweepOutHit, ETeleportType::TeleportPhysics);
-			UE_LOG(LogFlyingMovement, Verbose, TEXT("ResolvePenetration:   sweep by %s (success = %d)"), *Adjustment.ToString(), bMoved);
-			
-			// Still stuck?
-			if (!bMoved && SweepOutHit.bStartPenetrating)
-			{
-				// Combine two MTD results to get a new direction that gets out of multiple surfaces.
-				const FVector SecondMTD = GetPenetrationAdjustment(SweepOutHit);
-				const FVector CombinedMTD = Adjustment + SecondMTD;
-				if (SecondMTD != Adjustment && !CombinedMTD.IsZero())
-				{
-					bMoved = MoveUpdatedComponent(CombinedMTD, NewRotationQuat, true, nullptr, ETeleportType::TeleportPhysics);
-					UE_LOG(LogFlyingMovement, Verbose, TEXT("ResolvePenetration:   sweep by %s (MTD combo success = %d)"), *CombinedMTD.ToString(), bMoved);
-				}
-			}
-
-			// Still stuck?
-			if (!bMoved)
-			{
-				// Try moving the proposed adjustment plus the attempted move direction. This can sometimes get out of penetrations with multiple objects
-				const FVector MoveDelta = (Hit.TraceEnd - Hit.TraceStart); //ConstrainDirectionToPlane(Hit.TraceEnd - Hit.TraceStart);
-				if (!MoveDelta.IsZero())
-				{
-					bMoved = MoveUpdatedComponent(Adjustment + MoveDelta, NewRotationQuat, true, nullptr, ETeleportType::TeleportPhysics);
-					UE_LOG(LogFlyingMovement, Verbose, TEXT("ResolvePenetration:   sweep by %s (adjusted attempt success = %d)"), *(Adjustment + MoveDelta).ToString(), bMoved);
-				}
-			}	
-
-			return bMoved;
-		}
-	}
-
-	return false;
-}
-
-bool UFlyingMovementComponent::SafeMoveUpdatedComponent(const FVector& Delta, const FQuat& NewRotation, bool bSweep, FHitResult& OutHit, ETeleportType Teleport) const
-{
-	if (UpdatedComponent == NULL)
-	{
-		OutHit.Reset(1.f);
-		return false;
-	}
-
-	bool bMoveResult = false;
-
-	// Scope for move flags
-	{
-		bMoveResult = MoveUpdatedComponent(Delta, NewRotation, bSweep, &OutHit, Teleport);
-	}
-
-	// Handle initial penetrations
-	if (OutHit.bStartPenetrating && UpdatedComponent)
-	{
-		const FVector RequestedAdjustment = GetPenetrationAdjustment(OutHit);
-		if (ResolvePenetration(RequestedAdjustment, OutHit, NewRotation))
-		{
-			// Retry original move
-			bMoveResult = MoveUpdatedComponent(Delta, NewRotation, bSweep, &OutHit, Teleport);
-		}
-	}
-
-	return bMoveResult;
-}
-
-bool UFlyingMovementComponent::MoveUpdatedComponent(const FVector& Delta, const FQuat& NewRotation, bool bSweep, FHitResult* OutHit, ETeleportType Teleport) const
-{
-	if (UpdatedComponent)
-	{
-		const FVector NewDelta = Delta; //ConstrainDirectionToPlane(Delta);
-		return UpdatedComponent->MoveComponent(NewDelta, NewRotation, bSweep, OutHit, MoveComponentFlags, Teleport);
-	}
-
-	return false;
-}
-
-FTransform UFlyingMovementComponent::GetUpdateComponentTransform() const
-{
-	if (ensure(UpdatedComponent))
-	{
-		return UpdatedComponent->GetComponentTransform();		
-	}
-	return FTransform::Identity;
-}
 
 void UFlyingMovementComponent::InitSyncState(FlyingMovement::FMoveState& OutSyncState) const
 {
@@ -458,21 +177,3 @@ void UFlyingMovementComponent::SyncTo(const FlyingMovement::FMoveState& SyncStat
 		UpdatedComponent->ComponentVelocity = SyncState.Velocity;
 	}
 }
-
-void UFlyingMovementComponent::GetCapsuleDimensions(float &Radius, float& HalfHeight) const
-{
-	if (UCapsuleComponent const* CapsuleComponent = Cast<UCapsuleComponent>(UpdatedComponent))
-	{
-		CapsuleComponent->GetScaledCapsuleSize(Radius, HalfHeight);
-	}
-}
-
-FTransform UFlyingMovementComponent::GetDebugWorldTransform() const
-{
-	return GetOwner()->GetActorTransform();
-}
-
-UObject* UFlyingMovementComponent::GetVLogOwner() const
-{
-	return GetOwner();
-}
\ No newline at end of file
diff --git a/Engine/Plugins/Runtime/NetworkPrediction/Source/NetworkPrediction/Private/Movement/FlyingMovementSimulation.cpp b/Engine/Plugins/Runtime/NetworkPrediction/Source/NetworkPrediction/Private/Movement/FlyingMovementSimulation.cpp
--- a/Engine/Plugins/Runtime/NetworkPrediction/Source/NetworkPrediction/Private/Movement/FlyingMovementSimulation.cpp
+++ b/Engine/Plugins/Runtime/NetworkPrediction/Source/NetworkPrediction/Private/Movement/FlyingMovementSimulation.cpp
@@ -113,7 +113,9 @@ namespace FlyingMovement
 
 		if ((SlideDelta | Delta) > 0.f)
 		{
-			Driver->SafeMoveUpdatedComponent(SlideDelta, Rotation, true, Hit, ETeleportType::None);
+			IBaseMovementDriver& BaseMovementDriver = Driver->GetBaseMovementDriver();
+
+			BaseMovementDriver.SafeMoveUpdatedComponent(SlideDelta, Rotation, true, Hit, ETeleportType::None);
 
 			const float FirstHitPercent = Hit.Time;
 			PercentTimeApplied = FirstHitPercent;
@@ -132,7 +134,7 @@ namespace FlyingMovement
 				if (!SlideDelta.IsNearlyZero(1e-3f) && (SlideDelta | Delta) > 0.f)
 				{
 					// Perform second move
-					Driver->SafeMoveUpdatedComponent(SlideDelta, Rotation, true, Hit, ETeleportType::None);
+					BaseMovementDriver.SafeMoveUpdatedComponent(SlideDelta, Rotation, true, Hit, ETeleportType::None);
 					const float SecondHitPercent = Hit.Time * (1.f - FirstHitPercent);
 					PercentTimeApplied += SecondHitPercent;
 
@@ -164,6 +166,8 @@ namespace FlyingMovement
 		const float DeltaSeconds = InputCmd.FrameDeltaTime;
 		OutputState = InputState;
 
+		IBaseMovementDriver& BaseMovementDriver = Driver->GetBaseMovementDriver();
+
 		// --------------------------------------------------------------
 		//	Rotation Update
 		//	We do the rotational update inside the movement sim so that things like server side teleport will work.
@@ -245,7 +249,7 @@ namespace FlyingMovement
 		if (!Delta.IsNearlyZero(1e-6f))
 		{
 			FHitResult Hit(1.f);
-			Driver->SafeMoveUpdatedComponent(Delta, OutputQuat, true, Hit, ETeleportType::None);
+			BaseMovementDriver.SafeMoveUpdatedComponent(Delta, OutputQuat, true, Hit, ETeleportType::None);
 
 			if (Hit.IsValidBlockingHit())
 			{
@@ -264,7 +268,7 @@ namespace FlyingMovement
 		}
 
 		// Finalize. This is unfortunate. The component mirrors our internal motion state and since we call into it to update, at this point, it has the real position.
-		const FTransform UpdateComponentTransform = Driver->GetUpdateComponentTransform();
+		const FTransform UpdateComponentTransform = BaseMovementDriver.GetUpdateComponentTransform();
 		OutputState.Location = UpdateComponentTransform.GetLocation();
 
 		// Note that we don't pull the rotation out of the final update transform. Converting back from a quat will lead to a different FRotator than what we are storing
@@ -273,47 +277,15 @@ namespace FlyingMovement
 
 	void FMoveState::VisualLog(const FVisualLoggingParameters& Parameters, IFlyingMovementDriver* Driver, IFlyingMovementDriver* LogDriver)
 	{
-		const bool DrawPointOnly = (Parameters.Context == EVisualLoggingContext::OtherMispredicted || Parameters.Context == EVisualLoggingContext::OtherPredicted);
-		const bool DrawPersistent = (Parameters.Lifetime == EVisualLoggingLifetime::Persistent);
-		const float Lifetime = Parameters.Lifetime == EVisualLoggingLifetime::Persistent ? FlyingPawnSimCVars::DrawDebugDefaultLifeTime : -1.f;
-
-		float Radius=0.f;
-		float HalfHeight=0.f;
-		Driver->GetCapsuleDimensions(Radius, HalfHeight);
-
-		if (FlyingPawnSimCVars::UseDrawDebug)
-		{
-			FVector DrawPosition = Location;
-			const float Thickness = 2.f;
-
-			if (DrawPointOnly)
-			{
-				static float PointSize = 3.f;
-				DrawDebugPoint(LogDriver->GetDriverWorld(), DrawPosition, PointSize, Parameters.GetDebugColor(), DrawPersistent, Lifetime);
-			}
-			else
-			{
-				DrawDebugCapsule(LogDriver->GetDriverWorld(), DrawPosition, HalfHeight, Radius, Rotation.Quaternion(), Parameters.GetDebugColor(), DrawPersistent, Lifetime, 0.f, Thickness);
-			}
-		}
+		IBaseMovementDriver& BaseMovementDriver = Driver->GetBaseMovementDriver();		
+		FTransform Transform(Rotation, Location);		
 
-	
-		if (FlyingPawnSimCVars::UseVLogger)
-		{
-			FVector DrawPosition = Location;
-			if (DrawPointOnly)
-			{
-				static float CrumbRadius = 3.f;
-				UE_VLOG_LOCATION(LogDriver->GetVLogOwner(), LogFlyingPawnSimulation, Log, DrawPosition, CrumbRadius, Parameters.GetDebugColor(), TEXT("%d"), Parameters.Keyframe);
-			}
-			else
-			{
-				DrawPosition.Z -= HalfHeight;
-				UE_VLOG_CAPSULE(LogDriver->GetVLogOwner(), LogFlyingPawnSimulation, Log, DrawPosition, HalfHeight, Radius, Rotation.Quaternion(), Parameters.GetDebugColor(), TEXT("%d"), Parameters.Keyframe);
-			}
+		IBaseMovementDriver::FDrawDebugParams DrawParams(Parameters, &LogDriver->GetBaseMovementDriver());		
+		DrawParams.Transform = Transform;		
+		DrawParams.InWorldText = LexToString(Parameters.Keyframe);
+		DrawParams.LogText = FString::Printf(TEXT("[%d] %s. Location: %s. Rotation: %s"), Parameters.Keyframe, *LexToString(Parameters.Context), *Transform.GetLocation().ToString(), *Transform.GetRotation().Rotator().ToString());
 
-			UE_VLOG(LogDriver->GetVLogOwner(), LogFlyingPawnSimulation, Log, TEXT("[%d] %s. Location: %s."), Parameters.Keyframe, *LexToString(Parameters.Context), *Location.ToString());
-		}
+		BaseMovementDriver.DrawDebug(DrawParams);
 	}
 
 } // end namespace
\ No newline at end of file
diff --git a/Engine/Plugins/Runtime/NetworkPrediction/Source/NetworkPrediction/Private/Movement/ParametricMovement.cpp b/Engine/Plugins/Runtime/NetworkPrediction/Source/NetworkPrediction/Private/Movement/ParametricMovement.cpp
new file mode 100644
--- /dev/null
+++ b/Engine/Plugins/Runtime/NetworkPrediction/Source/NetworkPrediction/Private/Movement/ParametricMovement.cpp
@@ -0,0 +1,211 @@
+// Copyright 1998-2019 Epic Games, Inc. All Rights Reserved.
+
+#include "Movement/ParametricMovement.h"
+#include "NetworkSimulationModelDebugger.h"
+
+DEFINE_LOG_CATEGORY_STATIC(LogParametricMovement, Log, All);
+
+namespace ParametricMoverCVars
+{
+static int32 UseVLogger = 1;
+static FAutoConsoleVariableRef CVarUseVLogger(TEXT("parametricmover.Debug.UseUnrealVLogger"),
+	UseVLogger,	TEXT("Use Unreal Visual Logger\n"),	ECVF_Default);
+
+static int32 UseDrawDebug = 1;
+static FAutoConsoleVariableRef CVarUseDrawDebug(TEXT("parametricmover.Debug.UseDrawDebug"),
+	UseVLogger,	TEXT("Use built in DrawDebug* functions for visual logging\n"), ECVF_Default);
+
+static float DrawDebugDefaultLifeTime = 30.f;
+static FAutoConsoleVariableRef CVarDrawDebugDefaultLifeTime(TEXT("parametricmover.Debug.UseDrawDebug.DefaultLifeTime"),
+	DrawDebugDefaultLifeTime, TEXT("Use built in DrawDebug* functions for visual logging"), ECVF_Default);
+
+}
+
+// -------------------------------------------------------------------------------------------------------------------------------
+// ParametricMovement
+// -------------------------------------------------------------------------------------------------------------------------------
+
+namespace ParametricMovement
+{
+	void FMovementSystem::Update(IMovementDriver* Driver, const FInputCmd& InputCmd, const FMoveState& InputState, FMoveState& OutputState, const FAuxState& AuxState)
+	{
+		IBaseMovementDriver& BaseMovementDriver = Driver->GetBaseMovementDriver();
+		const float DeltaSeconds = InputCmd.FrameDeltaTime;
+
+		// Advance parametric time. This won't always be linear: we could loop, rewind/bounce, etc
+		const float InputPlayRate = InputCmd.PlayRate.Get(InputState.PlayRate); // Returns InputCmds playrate if set, else returns previous state's playrate		
+		Driver->AdvanceParametricTime(InputState.Position, InputPlayRate, OutputState.Position, OutputState.PlayRate, InputCmd.FrameDeltaTime);
+
+		// We have our time that we should be at. We just need to move primitive component to that position.
+		// Again, note that we expect this cannot fail. We move like this so that it can push things, but we don't expect failure.
+		// (I think it would need to be a different simulation that supports this as a failure case. The tricky thing would be working out
+		// the ouput Position in the case where a Move is blocked (E.g, you move 50% towards the desired location)
+
+		FTransform StartingTransform;
+		Driver->SetTransformForPosition(InputState.Position, StartingTransform);
+
+		FTransform NewTransform;
+		Driver->SetTransformForPosition(OutputState.Position, NewTransform);
+
+		FHitResult Hit(1.f);
+
+		FVector Delta = NewTransform.GetLocation() - StartingTransform.GetLocation();
+		BaseMovementDriver.MoveUpdatedComponent(Delta, NewTransform.GetRotation(), true, &Hit, ETeleportType::None);
+
+		if (Hit.IsValidBlockingHit())
+		{
+			FTransform ActualTransform = BaseMovementDriver.GetUpdateComponentTransform();
+			FVector ActualDelta = NewTransform.GetLocation() - ActualTransform.GetLocation();
+			UE_LOG(LogParametricMovement, Warning, TEXT("Blocking hit occurred when trying to move parametric mover. ActualDelta: %s"), *ActualDelta.ToString());
+		}
+	}
+
+	void FMoveState::VisualLog(const FVisualLoggingParameters& Parameters, IParametricMovementDriver* Driver, IParametricMovementDriver* LogDriver)
+	{
+		IBaseMovementDriver& BaseMovementDriver = Driver->GetBaseMovementDriver();
+		IBaseMovementDriver& LogMovementDriver = LogDriver->GetBaseMovementDriver();
+
+		FTransform Transform;
+		Driver->SetTransformForPosition(Position, Transform);
+
+		IBaseMovementDriver::FDrawDebugParams DrawParams;
+		DrawParams.DebugWorld = LogMovementDriver.GetDriverWorld();
+		DrawParams.DebugLogOwner = LogMovementDriver.GetVLogOwner();
+		DrawParams.Transform = Transform;
+		DrawParams.DrawType = (Parameters.Context == EVisualLoggingContext::OtherMispredicted || Parameters.Context == EVisualLoggingContext::OtherPredicted) ? EVisualLoggingDrawType::Crumb : EVisualLoggingDrawType::Full;
+		DrawParams.Lifetime = Parameters.Lifetime;
+		DrawParams.DrawColor = Parameters.GetDebugColor();
+		DrawParams.InWorldText = LexToString(Parameters.Keyframe);
+		DrawParams.LogText = FString::Printf(TEXT("[%d] %s. Location: %s. Rotation: %s"), Parameters.Keyframe, *LexToString(Parameters.Context), *Transform.GetLocation().ToString(), *Transform.GetRotation().Rotator().ToString());
+
+		BaseMovementDriver.DrawDebug(DrawParams);
+	}
+};
+
+
+// -------------------------------------------------------------------------------------------------------------------------------
+// UParametricMovementComponent
+// -------------------------------------------------------------------------------------------------------------------------------
+
+UParametricMovementComponent::UParametricMovementComponent()
+{
+
+}
+
+IReplicationProxy* UParametricMovementComponent::InstantiateNetworkSimulation()
+{
+	NetworkSim.Reset(new ParametricMovement::FMovementSystem());
+
+#if NETSIM_MODEL_DEBUG
+	FNetworkSimulationModelDebuggerManager::Get().RegisterNetworkSimulationModel(NetworkSim.Get(), (IParametricMovementDriver*)this, GetOwner(), TEXT("ParametricMovement"));
+#endif
+	return NetworkSim.Get();
+}
+
+void UParametricMovementComponent::InitializeForNetworkRole(ENetRole Role)
+{	
+	check(NetworkSim);
+	NetworkSim->InitializeForNetworkRole(Role, true /* fixme IsLocallyControlled() */, GetSimulationInitParameters(Role));
+}
+
+FNetworkSimulationModelInitParameters UParametricMovementComponent::GetSimulationInitParameters(ENetRole Role)
+{
+	// These are reasonable defaults but may not be right for everyone
+	FNetworkSimulationModelInitParameters InitParams;
+	InitParams.InputBufferSize = Role != ROLE_SimulatedProxy ? 32 : 32;  // Fixme.. not good
+	InitParams.SyncedBufferSize = Role != ROLE_AutonomousProxy ? 2 : 32;
+	InitParams.AuxBufferSize = Role != ROLE_AutonomousProxy ? 2 : 32;
+	InitParams.DebugBufferSize = 32;
+	InitParams.HistoricBufferSize = 128;
+	return InitParams;
+}
+
+void UParametricMovementComponent::BeginPlay()
+{
+	Super::BeginPlay();
+	CachedStartingTransform = UpdatedComponent->GetComponentToWorld();
+}
+
+float UParametricMovementComponent::GetTransformForTime(const float InTime, FTransform& OutTransform) const
+{
+	const FVector Delta = ParametricDelta * InTime;
+	OutTransform = CachedStartingTransform;
+	OutTransform.AddToTranslation(Delta);
+	return InTime;
+}
+
+void UParametricMovementComponent::InitSyncState(ParametricMovement::FMoveState& OutSyncState) const
+{
+	// In this case, we just default to the 0 position. Maybe this could be a starting variable set on the component.
+	OutSyncState.Position = 0.f;
+	OutSyncState.PlayRate = 0.f;
+}
+
+void UParametricMovementComponent::SyncTo(const ParametricMovement::FMoveState& SyncState)
+{
+	FTransform NewTransform;
+	GetTransformForTime(SyncState.Position, NewTransform);
+
+	check(UpdatedComponent);
+	UpdatedComponent->SetWorldTransform(NewTransform, false, nullptr, ETeleportType::TeleportPhysics);
+}
+
+void UParametricMovementComponent::AdvanceParametricTime(const float InPosition, const float InPlayRate, float &OutPosition, float& OutPlayRate, const float DeltaTimeSeconds) const
+{
+	// Real simple oscillation for now
+	OutPosition = InPosition + (InPlayRate * DeltaTimeSeconds);
+	OutPlayRate = InPlayRate;
+
+	const float DeltaMax = OutPosition - MaxTime;
+	if (DeltaMax > SMALL_NUMBER)
+	{
+		OutPosition = MaxTime - DeltaMax;
+		OutPlayRate *= -1.f;
+	}
+	else
+	{
+		const float DeltaMin = OutPosition - MinTime;
+		if (DeltaMin < SMALL_NUMBER)
+		{
+			OutPosition = MinTime - DeltaMin;
+			OutPlayRate *= -1.f;
+		}
+	}
+}
+
+void UParametricMovementComponent::SetTransformForPosition(const float InPosition, FTransform& OutTransform) const
+{
+	const FVector Delta = ParametricDelta * InPosition;
+
+	OutTransform = CachedStartingTransform;
+	OutTransform.AddToTranslation(Delta);
+}
+
+void UParametricMovementComponent::TickComponent(float DeltaTime, enum ELevelTick TickType, FActorComponentTickFunction* ThisTickFunction)
+{
+	Super::TickComponent(DeltaTime, TickType, ThisTickFunction);
+	const ENetRole OwnerRole = GetOwnerRole();
+
+	if (NetworkSim && OwnerRole != ROLE_None)
+	{
+		
+		if (OwnerRole == ROLE_Authority)
+		{
+			if (ParametricMovement::FInputCmd* NextCmd = NetworkSim->GetNextInputForWrite(DeltaTime))
+			{
+				NextCmd->PlayRate = PendingPlayRate;
+				PendingPlayRate.Reset();
+			}
+		}
+
+		ParametricMovement::FMovementSystem::FTickParameters Parameters;
+		Parameters.Role = OwnerRole;
+		Parameters.LocalDeltaTimeSeconds = DeltaTime;
+
+		// Tick the core network sim, this will consume input and generate new sync state
+		NetworkSim->Tick((IParametricMovementDriver*)this, Parameters);
+	}
+
+	// TEMP
+	GetOwner()->ForceNetUpdate();
+}
\ No newline at end of file
diff --git a/Engine/Plugins/Runtime/NetworkPrediction/Source/NetworkPrediction/Private/NetworkPredictionComponent.cpp b/Engine/Plugins/Runtime/NetworkPrediction/Source/NetworkPrediction/Private/NetworkPredictionComponent.cpp
--- a/Engine/Plugins/Runtime/NetworkPrediction/Source/NetworkPrediction/Private/NetworkPredictionComponent.cpp
+++ b/Engine/Plugins/Runtime/NetworkPrediction/Source/NetworkPrediction/Private/NetworkPredictionComponent.cpp
@@ -59,6 +59,18 @@ void UNetworkPredictionComponent::GetLifetimeReplicatedProps(TArray< FLifetimePr
 #endif
 }
 
+FNetworkSimulationModelInitParameters UNetworkPredictionComponent::GetSimulationInitParameters(ENetRole Role)
+{
+	// These are reasonable defaults but may not be right for everyone
+	FNetworkSimulationModelInitParameters InitParams;
+	InitParams.InputBufferSize = Role != ROLE_SimulatedProxy ? 32 : 0;
+	InitParams.SyncedBufferSize = Role != ROLE_AutonomousProxy ? 2 : 32;
+	InitParams.AuxBufferSize = Role != ROLE_AutonomousProxy ? 2 : 32;
+	InitParams.DebugBufferSize = 32;
+	InitParams.HistoricBufferSize = 128;
+	return InitParams;
+}
+
 bool UNetworkPredictionComponent::IsLocallyControlled()
 {
 	// This awkward because, engine wide, "Is Locally Controlled" is really only a built in concept to Pawns.
diff --git a/Engine/Plugins/Runtime/NetworkPrediction/Source/NetworkPrediction/Private/NetworkSimulationModelDebugger.cpp b/Engine/Plugins/Runtime/NetworkPrediction/Source/NetworkPrediction/Private/NetworkSimulationModelDebugger.cpp
--- a/Engine/Plugins/Runtime/NetworkPrediction/Source/NetworkPrediction/Private/NetworkSimulationModelDebugger.cpp
+++ b/Engine/Plugins/Runtime/NetworkPrediction/Source/NetworkPrediction/Private/NetworkSimulationModelDebugger.cpp
@@ -9,6 +9,7 @@
 #include "Misc/NetworkGuid.h"
 #include "Engine/NetConnection.h"
 #include "Engine/NetDriver.h"
+#include "EngineUtils.h"
 
 FNetworkSimulationModelDebuggerManager& FNetworkSimulationModelDebuggerManager::Get()
 {
@@ -82,6 +83,51 @@ FAutoConsoleCommandWithWorldAndArgs NetworkSimulationModelDebugCmd(TEXT("nms.Deb
 	FNetworkSimulationModelDebuggerManager::Get().ToggleDebuggerActive(Pawn);
 }));
 
+FAutoConsoleCommandWithWorldAndArgs NetworkSimulationModelDebugClassCmd(TEXT("nms.Debug.Class"), TEXT(""),
+	FConsoleCommandWithWorldAndArgsDelegate::CreateLambda([](const TArray< FString >& Args, UWorld* World) 
+{
+	if (Args.Num() <= 0)
+	{
+		UE_LOG(LogNetworkSimDebug, Display, TEXT("Usage: nms.Debug.Class <Class>"));
+		return;
+	}
+
+	UClass* Class = FindObject<UClass>(ANY_PACKAGE, *Args[0]);
+	if (Class == nullptr)
+	{
+		UE_LOG(LogNetworkSimDebug, Display, TEXT("Could not find Class: %s. Searching for partial match."), *Args[0]);
+
+		for (TObjectIterator<UClass> ClassIt; ClassIt; ++ClassIt)
+		{
+			if (ClassIt->GetName().Contains(Args[0]))
+			{
+				Class = *ClassIt;
+				UE_LOG(LogNetworkSimDebug, Display, TEXT("Found: %s"), *Class->GetName());
+				break;
+			}
+
+		}
+
+		if (Class == nullptr)
+		{
+			UE_LOG(LogNetworkSimDebug, Display, TEXT("No Matches"));
+			return;
+		}
+	}
+
+	// Probably need to do distance based sorting or a way to specify which one you want.
+	for (TActorIterator<AActor> It(World); It; ++It)
+	{
+		if (It->GetClass()->IsChildOf(Class))
+		{
+			UE_LOG(LogNetworkSimDebug, Display, TEXT("Toggling NetworkSim debugger for %s"), *It->GetName());
+			FNetworkSimulationModelDebuggerManager::Get().ToggleDebuggerActive(*It);
+			break;
+		}
+	}
+	
+}));
+
 
 FAutoConsoleCommandWithWorldAndArgs NetworkSimulationModelDebugToggleContinousCmd(TEXT("nms.Debug.ToggleContinous"), TEXT(""),
 	FConsoleCommandWithWorldAndArgsDelegate::CreateLambda([](const TArray< FString >& InArgs, UWorld* World) 
diff --git a/Engine/Plugins/Runtime/NetworkPrediction/Source/NetworkPrediction/Public/Movement/BaseMovementComponent.h b/Engine/Plugins/Runtime/NetworkPrediction/Source/NetworkPrediction/Public/Movement/BaseMovementComponent.h
new file mode 100644
--- /dev/null
+++ b/Engine/Plugins/Runtime/NetworkPrediction/Source/NetworkPrediction/Public/Movement/BaseMovementComponent.h
@@ -0,0 +1,108 @@
+// Copyright 1998-2019 Epic Games, Inc. All Rights Reserved
+
+#pragma once
+
+#include "NetworkPredictionComponent.h"
+#include "CollisionQueryParams.h"
+#include "Engine/EngineTypes.h"
+#include "WorldCollision.h"
+#include "Components/SceneComponent.h"
+
+#include "BaseMovementComponent.generated.h"
+
+// -------------------------------------------------------------------------------------------------------------------------------
+//	Base component for movement. This essentially has the generic glue for selecting an UpdatedComponent and moving it along the world
+//	It is abstract in that you still need to define which simulation the component runs (via ::InstantiateNetworkSimulation)
+// -------------------------------------------------------------------------------------------------------------------------------
+
+// This is the base driver, to be used by the backing movement TNetworkedSimulationModel. 
+// It is the basic interface for moving and updating a primitive component that is owned by the UNetworkPredictionComponent component.
+// It is unfortunate and not ideal for perf that we have to go through a virtual function call to do collision queries. 
+// It also creates an awkward "dual hierarchy" when trying to share code like this.
+// A world where the simulation can do all the collision work and just push the final result back to the component layer would be much better.
+class IBaseMovementDriver
+{
+public:
+	virtual bool SafeMoveUpdatedComponent(const FVector& Delta, const FQuat& NewRotation, bool bSweep, FHitResult& OutHit, ETeleportType Teleport) const = 0;
+	virtual bool MoveUpdatedComponent(const FVector& Delta, const FQuat& NewRotation, bool bSweep, FHitResult* OutHit, ETeleportType Teleport) const = 0;
+	virtual FTransform GetUpdateComponentTransform() const = 0;
+
+	virtual UWorld* GetDriverWorld() const = 0;
+	virtual UObject* GetVLogOwner() const = 0;
+
+	struct FDrawDebugParams
+	{
+		FDrawDebugParams() { }
+		FDrawDebugParams(const struct FVisualLoggingParameters& Parameters, IBaseMovementDriver* LogDriver);
+
+		UWorld* DebugWorld;		// World to draw stuff to (can be different than this->GetWorld() in PIE case)
+		UObject* DebugLogOwner; // Object that should own any VLog entries
+		
+		EVisualLoggingDrawType DrawType;
+		EVisualLoggingLifetime Lifetime;
+		FColor DrawColor;
+
+		FTransform Transform;	// Where to draw (important to understand: you don't draw your current state, you draw where the simulation says to draw)
+		FString InWorldText;
+		FString LogText;
+	};
+	
+	// This is the final debug drawing function that the driver can implement. UBaseMovementComponent implements a generic one that will draw capsules or actor bounds.
+	virtual void DrawDebug(const FDrawDebugParams& Params) const = 0;
+};
+
+UCLASS(Abstract)
+class NETWORKPREDICTION_API UBaseMovementComponent : public UNetworkPredictionComponent, public IBaseMovementDriver
+{
+	GENERATED_BODY()
+
+public:
+
+	UBaseMovementComponent();
+
+	virtual void InitializeComponent() override;
+	virtual void OnRegister() override;
+	virtual void RegisterComponentTickFunctions(bool bRegister) override;
+
+	// IBaseMovementDriver
+	bool SafeMoveUpdatedComponent(const FVector& Delta, const FQuat& NewRotation, bool bSweep, FHitResult& OutHit, ETeleportType Teleport) const override final;
+	bool MoveUpdatedComponent(const FVector& Delta, const FQuat& NewRotation, bool bSweep, FHitResult* OutHit, ETeleportType Teleport) const override final;
+	FTransform GetUpdateComponentTransform() const override final;
+	
+	UWorld* GetDriverWorld() const override final { return GetWorld(); }
+	UObject* GetVLogOwner() const override final;	
+
+	virtual void DrawDebug(const IBaseMovementDriver::FDrawDebugParams& Params) const override;
+
+protected:
+
+	// Basic "Update Component/Ticking"
+	virtual void SetUpdatedComponent(USceneComponent* NewUpdatedComponent);
+	virtual void UpdateTickRegistration();
+
+	UFUNCTION()
+	virtual void PhysicsVolumeChanged(class APhysicsVolume* NewVolume);
+
+	UPROPERTY()
+	USceneComponent* UpdatedComponent = nullptr;
+
+	UPROPERTY()
+	UPrimitiveComponent* UpdatedPrimitive = nullptr;
+
+private:
+
+	/** Transient flag indicating whether we are executing OnRegister(). */
+	bool bInOnRegister = false;
+	
+	/** Transient flag indicating whether we are executing InitializeComponent(). */
+	bool bInInitializeComponent = false;
+
+	static FVector GetPenetrationAdjustment(const FHitResult& Hit);
+
+	bool OverlapTest(const FVector& Location, const FQuat& RotationQuat, const ECollisionChannel CollisionChannel, const FCollisionShape& CollisionShape, const AActor* IgnoreActor) const;
+	void InitCollisionParams(FCollisionQueryParams &OutParams, FCollisionResponseParams& OutResponseParam) const;
+	bool ResolvePenetration(const FVector& ProposedAdjustment, const FHitResult& Hit, const FQuat& NewRotationQuat) const;
+
+	/**  Flags that control the behavior of calls to MoveComponent() on our UpdatedComponent. */
+	mutable EMoveComponentFlags MoveComponentFlags = MOVECOMP_NoFlags; // Mutable because we sometimes need to disable these flags ::ResolvePenetration. Better way may be possible
+};
\ No newline at end of file
diff --git a/Engine/Plugins/Runtime/NetworkPrediction/Source/NetworkPrediction/Public/Movement/FlyingMovement.h b/Engine/Plugins/Runtime/NetworkPrediction/Source/NetworkPrediction/Public/Movement/FlyingMovement.h
--- a/Engine/Plugins/Runtime/NetworkPrediction/Source/NetworkPrediction/Public/Movement/FlyingMovement.h
+++ b/Engine/Plugins/Runtime/NetworkPrediction/Source/NetworkPrediction/Public/Movement/FlyingMovement.h
@@ -14,7 +14,7 @@
 #include "Misc/OutputDevice.h"
 #include "Misc/CoreDelegates.h"
 #include "NetworkSimulationModelTemplates.h"
-#include "NetworkPredictionComponent.h"
+#include "BaseMovementComponent.h"
 
 #include "FlyingMovement.generated.h"
 
@@ -119,16 +119,8 @@ namespace FlyingMovement
 		class IMovementDriver : public IDriver
 		{
 		public:
-
-			virtual bool SafeMoveUpdatedComponent(const FVector& Delta, const FQuat& NewRotation, bool bSweep, FHitResult& OutHit, ETeleportType Teleport) const = 0;
-			virtual bool MoveUpdatedComponent(const FVector& Delta, const FQuat& NewRotation, bool bSweep, FHitResult* OutHit, ETeleportType Teleport) const = 0;
-			virtual FTransform GetUpdateComponentTransform() const = 0;
-
-
-			virtual void GetCapsuleDimensions(float &Radius, float& HalfHeight) const = 0;
-			virtual UWorld* GetDriverWorld() const = 0;
-			virtual UObject* GetVLogOwner() const = 0;
-			virtual FTransform GetDebugWorldTransform() const = 0;
+			// FlyingMovement only really needs the base driver functions for moving a primitive component around.
+			virtual IBaseMovementDriver& GetBaseMovementDriver() = 0;
 		};
 
 		/** Main update function */
@@ -145,75 +137,37 @@ namespace FlyingMovement
 
 // Needed to trick UHT into letting UMockNetworkSimulationComponent implement. UHT cannot parse the ::
 // Also needed for forward declaring. Can't just be a typedef/using =
-class IFlyingMovementDriver : public FlyingMovement::FMovementSystem::IMovementDriver { };
+class IFlyingMovementDriver : virtual public FlyingMovement::FMovementSystem::IMovementDriver { };
 
 // -------------------------------------------------------------------------------------------------------------------------------
 // ActorComponent for running FlyingMovement 
 // -------------------------------------------------------------------------------------------------------------------------------
 
 UCLASS(BlueprintType, meta=(BlueprintSpawnableComponent))
-class NETWORKPREDICTION_API UFlyingMovementComponent : public UNetworkPredictionComponent, public IFlyingMovementDriver
+class NETWORKPREDICTION_API UFlyingMovementComponent : public UBaseMovementComponent, public IFlyingMovementDriver
 {
 	GENERATED_BODY()
 
 public:
 
 	UFlyingMovementComponent();
-
-	virtual void InitializeComponent() override;
-	virtual void OnRegister() override;
-	virtual void RegisterComponentTickFunctions(bool bRegister) override;
+	
 	virtual void TickComponent(float DeltaTime, enum ELevelTick TickType, FActorComponentTickFunction *ThisTickFunction) override;
 
 	// Gets writable copy of FClientInputCmd for this frame. Should only be called once per frame by the local player.
 	FlyingMovement::FInputCmd* GetNextClientInputCmdForWrite(float DeltaTime);
 
-	// IFlyingMovementDriver
-	bool SafeMoveUpdatedComponent(const FVector& Delta, const FQuat& NewRotation, bool bSweep, FHitResult& OutHit, ETeleportType Teleport) const;
-	bool MoveUpdatedComponent(const FVector& Delta, const FQuat& NewRotation, bool bSweep, FHitResult* OutHit, ETeleportType Teleport) const;
-	FTransform GetUpdateComponentTransform() const;
+	IBaseMovementDriver& GetBaseMovementDriver() override final { return *static_cast<IBaseMovementDriver*>(this); }
 
 	void InitSyncState(FlyingMovement::FMoveState& OutSyncState) const override;
 	void SyncTo(const FlyingMovement::FMoveState& SyncState) override;
-	virtual UWorld* GetDriverWorld() const override final { return GetWorld(); }
-	virtual UObject* GetVLogOwner() const override final;
-	virtual FTransform GetDebugWorldTransform() const override final;
-	virtual void GetCapsuleDimensions(float &Radius, float& HalfHeight) const;
 
 protected:
 
-	// Basic "Update Component/Ticking"
-	virtual void SetUpdatedComponent(USceneComponent* NewUpdatedComponent);
-	virtual void UpdateTickRegistration();
-
-	UFUNCTION()
-	virtual void PhysicsVolumeChanged(class APhysicsVolume* NewVolume);
-
 	// Network Prediction
-	virtual IReplicationProxy* InstantiateNetworkSimulation() override;	
+	virtual IReplicationProxy* InstantiateNetworkSimulation() override;
 	virtual void InitializeForNetworkRole(ENetRole Role) override;
-	TUniquePtr<FlyingMovement::FMovementSystem> NetworkSim; // The Network sim that this component is managing. This is what is doing all the work.
 
-private:
-
-	UPROPERTY()
-	USceneComponent* UpdatedComponent = nullptr;
-
-	UPROPERTY()
-	UPrimitiveComponent* UpdatedPrimitive = nullptr;
-
-	/** Transient flag indicating whether we are executing OnRegister(). */
-	bool bInOnRegister = false;
+	TUniquePtr<FlyingMovement::FMovementSystem> NetworkSim; // The Network sim that this component is managing. This is what is doing all the work.
 	
-	/** Transient flag indicating whether we are executing InitializeComponent(). */
-	bool bInInitializeComponent = false;
-
-	static FVector GetPenetrationAdjustment(const FHitResult& Hit);
-
-	bool OverlapTest(const FVector& Location, const FQuat& RotationQuat, const ECollisionChannel CollisionChannel, const FCollisionShape& CollisionShape, const AActor* IgnoreActor) const;
-	void InitCollisionParams(FCollisionQueryParams &OutParams, FCollisionResponseParams& OutResponseParam) const;
-	bool ResolvePenetration(const FVector& ProposedAdjustment, const FHitResult& Hit, const FQuat& NewRotationQuat) const;
-
-	/**  Flags that control the behavior of calls to MoveComponent() on our UpdatedComponent. */
-	mutable EMoveComponentFlags MoveComponentFlags = MOVECOMP_NoFlags; // Mutable because we sometimes need to disable these flags ::ResolvePenetration. Better way may be possible
 };
\ No newline at end of file
diff --git a/Engine/Plugins/Runtime/NetworkPrediction/Source/NetworkPrediction/Public/Movement/ParametricMovement.h b/Engine/Plugins/Runtime/NetworkPrediction/Source/NetworkPrediction/Public/Movement/ParametricMovement.h
new file mode 100644
--- /dev/null
+++ b/Engine/Plugins/Runtime/NetworkPrediction/Source/NetworkPrediction/Public/Movement/ParametricMovement.h
@@ -0,0 +1,182 @@
+// Copyright 1998-2019 Epic Games, Inc. All Rights Reserved
+
+#pragma once
+
+#include "NetworkSimulationModelTemplates.h"
+#include "BaseMovementComponent.h"
+
+#include "ParametricMovement.generated.h"
+
+class IParametricMovementDriver;
+
+namespace ParametricMovement
+{
+	// State the client generates
+	struct FInputCmd
+	{
+		// Client's FrameTime. This is essential for a non-fixed step simulation. The server will run the movement simulation for this client at this rate.
+		float FrameDeltaTime = 0.f;
+
+		// Input Playrate. This being set can be thought of "telling the simulation what its new playrate should be"
+		TOptional<float> PlayRate;
+
+		void NetSerialize(const FNetSerializeParams& P)
+		{
+			// FIXME: quantization and effecient packing of this data needs to be looked at/generalized in a nice way
+			P.Ar << FrameDeltaTime;
+			P.Ar << PlayRate;
+		}
+
+		void Log(FStandardLoggingParameters& P) const
+		{
+			if (P.Context == EStandardLoggingContext::HeaderOnly)
+			{
+				P.Ar->Logf(TEXT(" %d "), P.Keyframe);
+			}
+			else if (P.Context == EStandardLoggingContext::Full)
+			{
+				P.Ar->Logf(TEXT("FrameDeltaTime: %.4f"), FrameDeltaTime);
+				if (PlayRate.IsSet())
+				{
+					P.Ar->Logf(TEXT("PlayRate: %.2f"), PlayRate.GetValue());
+				}
+			}
+		}
+	};
+
+	// State we are evolving frame to frame and keeping in sync
+	struct FMoveState
+	{
+		float Position;
+		float PlayRate;
+
+		void NetSerialize(const FNetSerializeParams& P)
+		{
+			P.Ar << Position;
+			P.Ar << PlayRate;
+		}
+
+		// Compare this state with AuthorityState. return true if a reconcile (correction) should happen
+		bool ShouldReconcile(const FMoveState& AuthorityState)
+		{
+			const float ErrorTolerance = 0.01f;
+			return (FMath::Abs<float>(AuthorityState.Position - Position) > ErrorTolerance) || (FMath::Abs<float>(AuthorityState.PlayRate - PlayRate) > ErrorTolerance);
+		}
+
+		void Log(FStandardLoggingParameters& Params) const
+		{
+			if (Params.Context == EStandardLoggingContext::HeaderOnly)
+			{
+				Params.Ar->Logf(TEXT(" %d "), Params.Keyframe);
+			}
+			else if (Params.Context == EStandardLoggingContext::Full)
+			{
+				Params.Ar->Logf(TEXT("Frame: %d"), Params.Keyframe);
+				Params.Ar->Logf(TEXT("Pos: %.2f"), Position);
+				Params.Ar->Logf(TEXT("Rate: %.2f"), PlayRate);
+			}
+		}
+
+		void VisualLog(const FVisualLoggingParameters& Parameters, IParametricMovementDriver* Driver, IParametricMovementDriver* LogDriver);
+	};
+
+	// Auxiliary state that is input into the simulation. Doesn't change during the simulation tick.
+	// (It can change and even be predicted but doing so will trigger more bookeeping, etc. Changes will happen "next tick")
+	struct FAuxState
+	{
+		float Multiplier=1;
+		void NetSerialize(const FNetSerializeParams& P)
+		{
+			P.Ar << Multiplier;
+		}
+	};
+
+	// Actual definition of our network simulation.
+	class FMovementSystem : public TNetworkedSimulationModel<FMovementSystem, FInputCmd, FMoveState, FAuxState>
+	{
+	public:
+
+		// Interface between the simulation and owning component driving it. Functions added here are available in ::Update and must be implemented by UMockNetworkSimulationComponent.
+		class IMovementDriver : public IDriver
+		{
+		public:
+
+			// BaseMovement driver API (functions for moving around a primitive component)
+			virtual IBaseMovementDriver& GetBaseMovementDriver() = 0;
+
+			// Advance parametric time. This is meant to do simple things like looping/reversing etc.
+			// Note how this should be STATIC and not rely on state outside of what is passed in (such thing would need to be done inside the simulation, not through the driver!)
+			virtual void AdvanceParametricTime(const float InPosition, const float InPlayRate, float &OutPosition, float& OutPlayRate, const float DeltaTimeSeconds) const = 0;
+
+			// Actually turn the given position into a transform. Again, should be static and not conditional on changing state outside of the network sim
+			virtual void SetTransformForPosition(const float InPosition, FTransform& OutTransform) const = 0;
+		};
+
+		/** Main update function */
+		static void Update(IMovementDriver* Driver, const FInputCmd& InputCmd, const FMoveState& InputState, FMoveState& OutputState, const FAuxState& AuxState);
+
+		/** Dev tool to force simple mispredict */
+		static bool ForceMispredict;
+	};
+
+} // End namespace
+
+// Needed to trick UHT into letting UMockNetworkSimulationComponent implement. UHT cannot parse the ::
+// Also needed for forward declaring. Can't just be a typedef/using =
+class IParametricMovementDriver : public ParametricMovement::FMovementSystem::IMovementDriver { };
+
+// -------------------------------------------------------------------------------------------------------------------------------
+//	ActorComponent for running basic Parametric movement. 
+//	Parametric movement could be anything that takes a Time and returns an FTransform.
+//	
+//	Initially, we will support pushing (ie, we sweep as we update the mover's position).
+//	But we will not allow a parametric mover from being blocked. 
+//
+// -------------------------------------------------------------------------------------------------------------------------------
+
+UCLASS(BlueprintType, meta=(BlueprintSpawnableComponent))
+class NETWORKPREDICTION_API UParametricMovementComponent : public UBaseMovementComponent, public IParametricMovementDriver
+{
+	GENERATED_BODY()
+
+	UParametricMovementComponent();
+
+	virtual void BeginPlay() override;
+	virtual void TickComponent(float DeltaTime, enum ELevelTick TickType, FActorComponentTickFunction* ThisTickFunction) override;
+
+	// Base TNetworkModelSimulation driver
+	void InitSyncState(ParametricMovement::FMoveState& OutSyncState) const override;
+	void SyncTo(const ParametricMovement::FMoveState& SyncState) override;
+
+	// Base Movement Driver
+	IBaseMovementDriver& GetBaseMovementDriver() override final { return *static_cast<IBaseMovementDriver*>(this); }
+
+	// Parametric Movement Driver
+	float GetTransformForTime(const float InTime, FTransform& OutTransform) const;
+
+	virtual void AdvanceParametricTime(const float InPosition, const float InPlayRate, float &OutPosition, float& OutPlayRate, const float DeltaTimeSeconds) const override;
+	virtual void SetTransformForPosition(const float InPosition, FTransform& OutTransform) const override;
+
+protected:
+
+	virtual IReplicationProxy* InstantiateNetworkSimulation() override;	
+	void InitializeForNetworkRole(ENetRole Role) override;
+	FNetworkSimulationModelInitParameters GetSimulationInitParameters(ENetRole Role) override;
+
+	TUniquePtr<ParametricMovement::FMovementSystem> NetworkSim;
+
+	// Temp....
+	UPROPERTY(EditAnywhere, BlueprintReadOnly, Category=ParametricMovement)
+	FVector ParametricDelta = FVector(0.f, 0.f, 500.f);
+
+	UPROPERTY(EditAnywhere, BlueprintReadOnly, Category=ParametricMovement)
+	float MinTime = -1.f;
+
+	UPROPERTY(EditAnywhere, BlueprintReadOnly, Category=ParametricMovement)
+	float MaxTime = 1.f;
+
+	FTransform CachedStartingTransform;
+
+	TOptional<float> PendingPlayRate = 1.f;
+};
+
diff --git a/Engine/Plugins/Runtime/NetworkPrediction/Source/NetworkPrediction/Public/NetworkPredictionComponent.h b/Engine/Plugins/Runtime/NetworkPrediction/Source/NetworkPrediction/Public/NetworkPredictionComponent.h
--- a/Engine/Plugins/Runtime/NetworkPrediction/Source/NetworkPrediction/Public/NetworkPredictionComponent.h
+++ b/Engine/Plugins/Runtime/NetworkPrediction/Source/NetworkPrediction/Public/NetworkPredictionComponent.h
@@ -13,6 +13,9 @@
 //
 //	This is an abstract component and cannot function on its own. It must be subclassed and InstantiateNetworkSimulation must be implemented.
 //	The subclass is responsible for running the network simulation. This at a bare minimum would mean ticking the simulation and supplying it with input.
+//
+//	Its also worth pointing out that nothing about being a UActorComponent is essential here. All that this component does could be done within an AActor itself.
+//	An actor component makes sense for flexible/reusable code provided by the plugin. But there is nothing stopping you from copying this directly into an actor if you had reason to.
 // ----------------------------------------------------------------------------------------------------------------------------------------------------------------------
 UCLASS(Abstract)
 class NETWORKPREDICTION_API UNetworkPredictionComponent : public UActorComponent
@@ -41,10 +44,13 @@ protected:
 
 	// Child classes must allocate and manage lifetime of their own NetworkSim (E.g, TNetworkedSimulationModel). Recommend to just store in a TUniquePtr.
 	// Child classes must also register with Network Sim debugger here (FNetworkSimulationModelDebuggerManager RegisterNetworkSimulationModel)
-	virtual IReplicationProxy* InstantiateNetworkSimulation() { return nullptr; }
+	virtual IReplicationProxy* InstantiateNetworkSimulation() { check(false); return nullptr; }
+
+	// Child must override this an initialize their NetworkSim here (call 
+	virtual void InitializeForNetworkRole(ENetRole Role) { check(false); }
 
-	// Child classes should override this an initialize their NetworkSim here
-	virtual void InitializeForNetworkRole(ENetRole Role) { }
+	// Doesn't need to be overridden. Expected to be used in ::InitializeForNetworkRole implementation
+	virtual FNetworkSimulationModelInitParameters GetSimulationInitParameters(ENetRole Role);
 
 	// Attempts to determine if the component is locally controlled or not, meaning we expect input cmds to be generated locally (and hence need to prepare the input buffer for that)
 	virtual bool IsLocallyControlled();
@@ -76,4 +82,9 @@ private:
 
 	UPROPERTY(Replicated, transient)
 	FReplicationProxy ReplicationProxy_Debug;
+
+protected:
+
+	UFUNCTION()
+	virtual void OnRep_SimulatedProxy() { }
 };
\ No newline at end of file
diff --git a/Engine/Plugins/Runtime/NetworkPrediction/Source/NetworkPrediction/Public/NetworkSimulationModel.h b/Engine/Plugins/Runtime/NetworkPrediction/Source/NetworkPrediction/Public/NetworkSimulationModel.h
--- a/Engine/Plugins/Runtime/NetworkPrediction/Source/NetworkPrediction/Public/NetworkSimulationModel.h
+++ b/Engine/Plugins/Runtime/NetworkPrediction/Source/NetworkPrediction/Public/NetworkSimulationModel.h
@@ -98,6 +98,12 @@ inline FString LexToString(EVisualLoggingContext A)
 	return *UEnum::GetValueAsString(TEXT("NetworkPrediction.EVisualLoggingContext"), A);
 }
 
+enum class EVisualLoggingDrawType : uint8
+{
+	Full,				// Draw "the full thing" (maybe a collision capsule for example)
+	Crumb,				// Draw a small/minimal representation (E.g, a point or small axis)
+};
+
 enum class EVisualLoggingLifetime : uint8
 {
 	Transient,			// This logging is transient and will (probably) be done every frame. Don't persist (in contexts where that would matter).
diff --git a/Engine/Plugins/Runtime/NetworkPrediction/Source/NetworkPrediction/Public/NetworkSimulationModelTemplates.h b/Engine/Plugins/Runtime/NetworkPrediction/Source/NetworkPrediction/Public/NetworkSimulationModelTemplates.h
--- a/Engine/Plugins/Runtime/NetworkPrediction/Source/NetworkPrediction/Public/NetworkSimulationModelTemplates.h
+++ b/Engine/Plugins/Runtime/NetworkPrediction/Source/NetworkPrediction/Public/NetworkSimulationModelTemplates.h
@@ -527,8 +527,6 @@ private:
 	bool bReconcileFaultDetected = false;	// A fault was detected: we received state from the server that we are unable to reconcile with locally predicted state
 };
 
-
-
 // ----------------------------------------------------------------------------------------------------------------------------------------------------------------------
 //	TNetworkedSimulationModel
 //	
@@ -724,11 +722,15 @@ public:
 			}
 		}
 
+		// FIXME: this needs to be sorted out. We really want to check if there is new sync state and then call this here.
 		// Call into the driver to sync to the latest state if we processed any input
-		if (NumProcessed > 0)
+		//if (NumProcessed > 0)
 		{
-			check(SyncBuffer.GetNumValidElements() > 0);
-			Driver->SyncTo(*SyncBuffer.GetElementFromHead(0));
+			//check(SyncBuffer.GetNumValidElements() > 0);
+			if (SyncBuffer.GetNumValidElements() > 0)
+			{
+				Driver->SyncTo(*SyncBuffer.GetElementFromHead(0));
+			}
 		}
 
 		// -------------------------------------------------------------------------------------------------------------------------------------------------
@@ -751,7 +753,7 @@ public:
 		}
 	}	
 	
-	void InitializeForNetworkRole(const ENetRole Role, const bool IsLocallyControleld, const FNetworkSimulationModelInitParameters& Parameters)
+	void InitializeForNetworkRole(const ENetRole Role, const bool IsLocallyControlled, const FNetworkSimulationModelInitParameters& Parameters)
 	{
 		InputBuffer.SetBufferSize(Parameters.InputBufferSize);
 		SyncBuffer.SetBufferSize(Parameters.SyncedBufferSize);
@@ -769,8 +771,9 @@ public:
 			MyHistoricBuffers->AuxBuffer.SetBufferSize(Parameters.HistoricBufferSize);
 		}
 
-		if (IsLocallyControleld)
+		if (IsLocallyControlled)
 		{
+			check(Parameters.InputBufferSize > 0); // If you tell me this is locally controlled, you need to have an input buffer.
 			InitLocalInputBuffer();
 		}
 	}
diff --git a/Engine/Plugins/Runtime/NetworkPrediction/readme.txt b/Engine/Plugins/Runtime/NetworkPrediction/readme.txt
--- a/Engine/Plugins/Runtime/NetworkPrediction/readme.txt
+++ b/Engine/Plugins/Runtime/NetworkPrediction/readme.txt
@@ -8,7 +8,7 @@ High level overview:
 -Long term we hope this replaces UCharacterMovementComponent and its networking system (ServerMove, ClientAck*, ClientAdjust*, etc).
 -The GameplayAbility system will be updated to interop with this as well.
 -NetworkPrediction plugin is the core generalized system that you need to enable to use anything.
-	-For now at least, new movement classes build on top of this will live here but could one day be moved out into a dedicated movement system plugin.
+	-For now at least, new movement classes built on top of this will live here but could one day be moved out into a dedicated movement system plugin.
 -NetworkPredictionExtras plugin contains some examples that may be helpful in learning the system. It is not inteded to be used directly in shipping games.
 
 High level goals:
@@ -17,12 +17,14 @@ High level goals:
 -Provide better tooling to understand what is going on in these complex predicted gameplay systems.
 -This all leads to: ability to develop richer and more interactive gameplay systems that work in multiplayer games.
 
+What this is NOT:
+-We are not moving to a global fixed ticking system. Though these systems have huge advantages, we don't think it is a good universal solution for the engine.
+-That said, we recognize that it is the right solution for many things and we want to enable support for that. Nothing in this system /requires/ dynamic ticking, but it is something we feel the need to support.
+-To be clear: widely different, variable, fluctuating tick rates between clients and server is basically the root problems we are fighting here.
+-Related, we are also not all in on a "predict everything: seemlessly and automtically" solution. We still feel player prediction is best kept to a minimum (meaning: predict the minimum amount of stuff you can get away with).
+-But, we are tring to make it way easier to write predictive gameplay systems. Easier to write, to debug and reason about.
+-We do hope the new movement system will be automatic in this sense: if you write a new movement mode within the movement simulation class, things should "just work". (But we are not trying to make your entire game predicted)
 
-Current State (7-25-19)
--Initial public check in. Core simulation model structure is in place. Lots of pieces missing but overall system is roughed in.
--Two main examples right now: Flying Movement and 'MockNetworkSimulation' (in NetworkPredictionExtras)
--Flying Movement: a port of FloatingPawnMovement into the new system. Essentially just a basic flying movement system.
--MockNetworkSimulation: a very basic minimal "simulation" that demonstrates how the system works. Not tied to movement or anything physical: just a running counter/sum.
 
 Getting Started:
 -Both NetworkPrediction and NetworkPredictionExtras are disabled by default. You will need to enable them for your project (manually edit .uproject or do so through editor plugin screen).
@@ -61,7 +63,10 @@ TODO: Major missing elements
 
 Short term road map:
 -The focus right now is on the generalized prediction system, rather than actual movement code.
--Next will be interopping the Flying Movement example with Ability System (E.g., movement abilities, stamina/fuel attributes, etc).
+-We are exploring "forward predicting non autonomous proxies" with the parametric mover. Unclear how deep we will go with this.
+-Then Aux buffer and "events"
+-Then mixing in ability system + other scripting options
+
 -Once everything feels solid we will go deeper on movement and build out a more extendable movement system.
 -E.g, We don't intend all users of the engine to write their own network simulation model movement system.
 -We hope to provide a generalized movement system that is easy to extend in itself without knowing all the details of the NetwokrPrediction plugin.
@@ -69,5 +74,21 @@ Short term road map:
 
 
 
+// ----------------------------------------------------------------------------------------------------------
+// Release notes
+// ----------------------------------------------------------------------------------------------------------
+
+Update (8-2-19)
+-We want to take a look non player controlled simulations, such as doors, elevators, "pushers" before going deeper on ability system or movement.
+-Basic ParametricMoverment system checked in. This is also in an incomplete state. ("Pushing" the flying movement component does not work and won't be the focus for now until we are looking for closely at movement specifically).
+-The short term goal here is to be able to forward predict these even if they are simulated proxies. When/why/how is very tbd, but we want to see what it looks like.
+
+
+Current State (7-25-19)
+-Initial public check in. Core simulation model structure is in place. Lots of pieces missing but overall system is roughed in.
+-Two main examples right now: Flying Movement and 'MockNetworkSimulation' (in NetworkPredictionExtras)
+-Flying Movement: a port of FloatingPawnMovement into the new system. Essentially just a basic flying movement system.
+-MockNetworkSimulation: a very basic minimal "simulation" that demonstrates how the system works. Not tied to movement or anything physical: just a running counter/sum.
+
 
 
diff --git a/Engine/Plugins/Runtime/NetworkPredictionExtras/Source/NetworkPredictionExtras/Private/MockNetworkSimulation.cpp b/Engine/Plugins/Runtime/NetworkPredictionExtras/Source/NetworkPredictionExtras/Private/MockNetworkSimulation.cpp
--- a/Engine/Plugins/Runtime/NetworkPredictionExtras/Source/NetworkPredictionExtras/Private/MockNetworkSimulation.cpp
+++ b/Engine/Plugins/Runtime/NetworkPredictionExtras/Source/NetworkPredictionExtras/Private/MockNetworkSimulation.cpp
@@ -139,17 +139,10 @@ IReplicationProxy* UMockNetworkSimulationComponent::InstantiateNetworkSimulation
 	return NetworkSim.Get();
 }
 
-	// Child classes should override this an initialize their NetworkSim here
+// Child classes should override this an initialize their NetworkSim here
 void UMockNetworkSimulationComponent::InitializeForNetworkRole(ENetRole Role)
-{
-	FNetworkSimulationModelInitParameters InitParams;
-	InitParams.InputBufferSize = Role != ROLE_SimulatedProxy ? 32 : 0;
-	InitParams.SyncedBufferSize = Role != ROLE_AutonomousProxy ? 2 : 32;
-	InitParams.AuxBufferSize = Role != ROLE_AutonomousProxy ? 2 : 32;
-	InitParams.DebugBufferSize = 32;
-	InitParams.HistoricBufferSize = 128;
-	
-	NetworkSim->InitializeForNetworkRole(Role, IsLocallyControlled(), InitParams);
+{	
+	NetworkSim->InitializeForNetworkRole(Role, IsLocallyControlled(), GetSimulationInitParameters(Role));
 }
 
 void UMockNetworkSimulationComponent::TickComponent(float DeltaTime, enum ELevelTick TickType, FActorComponentTickFunction *ThisTickFunction)
diff --git a/Engine/Plugins/Runtime/NetworkPredictionExtras/Source/NetworkPredictionExtras/Private/NetworkPredictionExtrasFlyingPawn.cpp b/Engine/Plugins/Runtime/NetworkPredictionExtras/Source/NetworkPredictionExtras/Private/NetworkPredictionExtrasFlyingPawn.cpp
--- a/Engine/Plugins/Runtime/NetworkPredictionExtras/Source/NetworkPredictionExtras/Private/NetworkPredictionExtrasFlyingPawn.cpp
+++ b/Engine/Plugins/Runtime/NetworkPredictionExtras/Source/NetworkPredictionExtras/Private/NetworkPredictionExtrasFlyingPawn.cpp
@@ -11,7 +11,8 @@
 #include "Camera/PlayerCameraManager.h"
 #include "Engine/World.h"
 #include "Engine/LocalPlayer.h"
-
+#include "Components/CapsuleComponent.h"
+#include "DrawDebugHelpers.h"
 
 namespace FlyingPawnCVars
 {
@@ -139,7 +140,6 @@ void ANetworkPredictionExtrasFlyingPawn::InputAxis_MoveDown(float Value)
 void ANetworkPredictionExtrasFlyingPawn::Tick( float DeltaSeconds)
 {
 	Super::Tick(DeltaSeconds);
-
 	// Do whatever you want here. By now we have the latest movement state and latest input processed.
 }
 
@@ -152,8 +152,10 @@ void ANetworkPredictionExtrasFlyingPawn::GenerateLocalInput(float DeltaSeconds)
 	//
 	// Its worth calling out: the code that happens here is happening *outside* of the flying movement simulation. All we are doing
 	// is generating the input being fed into that simulation. That said, this means that A) the code below does not run on the server
-	// (and non controlling clients) and B) the code is not rerun during reconcile/resimulates. This this information guide any
-	// decisions about where something should go (such as aim assist, lock on targeting systems, etc).
+	// (and non controlling clients) and B) the code is not rerun during reconcile/resimulates. Use this information guide any
+	// decisions about where something should go (such as aim assist, lock on targeting systems, etc): it is hard to give absolute
+	// answers and will depend on the game and its specific needs. In general, at this time, I'd recommend aim assist and lock on 
+	// targeting systems to happen /outside/ of the system, i.e, here. But I can think of scenarios where that may not be ideal too.
 
 	check(Controller);
 	if (FlyingMovement::FInputCmd* NextCmd = FlyingMovementComponent->GetNextClientInputCmdForWrite(DeltaSeconds))
