PVS fix + comments
#rnx

diff --git a/Engine/Plugins/Runtime/NetworkPrediction/Source/NetworkPrediction/Private/Movement/ParametricMovement.cpp b/Engine/Plugins/Runtime/NetworkPrediction/Source/NetworkPrediction/Private/Movement/ParametricMovement.cpp
--- a/Engine/Plugins/Runtime/NetworkPrediction/Source/NetworkPrediction/Private/Movement/ParametricMovement.cpp
+++ b/Engine/Plugins/Runtime/NetworkPrediction/Source/NetworkPrediction/Private/Movement/ParametricMovement.cpp
@@ -112,7 +112,7 @@ FNetworkSimulationModelInitParameters UParametricMovementComponent::GetSimulatio
 {
 	// These are reasonable defaults but may not be right for everyone
 	FNetworkSimulationModelInitParameters InitParams;
-	InitParams.InputBufferSize = Role != ROLE_SimulatedProxy ? 32 : 32;  // Fixme.. not good
+	InitParams.InputBufferSize = 32; //Role != ROLE_SimulatedProxy ? 32 : 32;  // Fixme.. not good
 	InitParams.SyncedBufferSize = Role != ROLE_AutonomousProxy ? 2 : 32;
 	InitParams.AuxBufferSize = Role != ROLE_AutonomousProxy ? 2 : 32;
 	InitParams.DebugBufferSize = 32;
diff --git a/Engine/Plugins/Runtime/NetworkPrediction/Source/NetworkPrediction/Public/Movement/ParametricMovement.h b/Engine/Plugins/Runtime/NetworkPrediction/Source/NetworkPrediction/Public/Movement/ParametricMovement.h
--- a/Engine/Plugins/Runtime/NetworkPrediction/Source/NetworkPrediction/Public/Movement/ParametricMovement.h
+++ b/Engine/Plugins/Runtime/NetworkPrediction/Source/NetworkPrediction/Public/Movement/ParametricMovement.h
@@ -165,7 +165,13 @@ protected:
 
 	TUniquePtr<ParametricMovement::FMovementSystem> NetworkSim;
 
-	// Temp....
+	// ------------------------------------------------------------------------
+	// Temp Parametric movement example
+	//	The essence of this movement simulation is to map some Time value to a transform. That is it.
+	//	(It could be mapped via a spline, a curve, a simple blueprint function, etc).
+	//	What is below is just a simple C++ implementation to stand things up. Most likely we would 
+	//	do additional subclasses to vary the way this is implemented)
+	// ------------------------------------------------------------------------
 	UPROPERTY(EditAnywhere, BlueprintReadOnly, Category=ParametricMovement)
 	FVector ParametricDelta = FVector(0.f, 0.f, 500.f);
 
diff --git a/Engine/Plugins/Runtime/NetworkPrediction/Source/NetworkPrediction/Public/NetworkSimulationModelTemplates.h b/Engine/Plugins/Runtime/NetworkPrediction/Source/NetworkPrediction/Public/NetworkSimulationModelTemplates.h
--- a/Engine/Plugins/Runtime/NetworkPrediction/Source/NetworkPrediction/Public/NetworkSimulationModelTemplates.h
+++ b/Engine/Plugins/Runtime/NetworkPrediction/Source/NetworkPrediction/Public/NetworkSimulationModelTemplates.h
@@ -470,13 +470,8 @@ struct TReplicator_BasicReconciliar
 
 		TSyncState* LatestClient = SyncBuffer.GetElementFromHead(0);
 
-		
+		//Temp		
 		UE_LOG(LogNetworkSim, Warning, TEXT("Client misprediction. Frame: %d"), ReconciliationKeyframe);
-		/*
-		UE_LOG(LogNetworkSim, Warning, TEXT("  Client: %s"), *ClientMotionState->Location.ToString());
-		UE_LOG(LogNetworkSim, Warning, TEXT("  Server: %s"), *ServerState->Location.ToString());
-		UE_LOG(LogNetworkSim, Warning, TEXT("  LatestClient: %s"), *LatestClient->Location.ToString());
-		*/
 
 		// Copy authoritative state over client state (FIXME: we may want to store this off historically somewhere? Or will VLog be enough for debugging?)
 		*ClientSyncState = *ServerState;
@@ -685,7 +680,7 @@ public:
 				{
 					if (LastProcessedInputKeyframe != 0)
 					{
-						// This shouldn't happen, but is not fatal. We are reseting the motion state buffer.
+						// This shouldn't happen, but is not fatal. We are reseting the sync state buffer.
 						UE_LOG(LogNetworkSim, Warning, TEXT("Break in SyncState continuity. LastProcessedInputKeyframe: %d. SyncBuffer.GetHeadKeyframe(): %d. Role=%d"), LastProcessedInputKeyframe, SyncBuffer.GetHeadKeyframe(), (int32)Parameters.Role);
 					}
 
