NetworkPrediction plugin WIP continued. Mostly internal refactoring of NetworkSimulationModel. Working towards:
-Interpolated/Extrapolated/ForwardPredicted modes for simulated proxy simulations
-Global Reconcile and Tick management
-More formalized paths for advancing simulation and tracking where we are in input/sync buffer progression

#rnx

diff --git a/Engine/Plugins/Runtime/NetworkPrediction/Source/NetworkPrediction/Private/Movement/BaseMovementComponent.cpp b/Engine/Plugins/Runtime/NetworkPrediction/Source/NetworkPrediction/Private/Movement/BaseMovementComponent.cpp
--- a/Engine/Plugins/Runtime/NetworkPrediction/Source/NetworkPrediction/Private/Movement/BaseMovementComponent.cpp
+++ b/Engine/Plugins/Runtime/NetworkPrediction/Source/NetworkPrediction/Private/Movement/BaseMovementComponent.cpp
@@ -7,7 +7,7 @@
 #include "Components/CapsuleComponent.h"
 #include "DrawDebugHelpers.h"
 #include "VisualLogger/VisualLogger.h"
-#include "NetworkSimulationModel.h"
+#include "NetworkPredictionTypes.h"
 
 
 DEFINE_LOG_CATEGORY_STATIC(LogBaseMovement, Log, All);
diff --git a/Engine/Plugins/Runtime/NetworkPrediction/Source/NetworkPrediction/Private/Movement/FlyingMovementComponent.cpp b/Engine/Plugins/Runtime/NetworkPrediction/Source/NetworkPrediction/Private/Movement/FlyingMovementComponent.cpp
--- a/Engine/Plugins/Runtime/NetworkPrediction/Source/NetworkPrediction/Private/Movement/FlyingMovementComponent.cpp
+++ b/Engine/Plugins/Runtime/NetworkPrediction/Source/NetworkPrediction/Private/Movement/FlyingMovementComponent.cpp
@@ -104,6 +104,8 @@ void UFlyingMovementComponent::TickComponent(float DeltaTime, enum ELevelTick Ti
 	// -------------------------------------
 	if (NetworkSim)
 	{
+		PreTickSimulation(DeltaTime); // Fixme
+
 		// Check if we should trip a mispredict. (Note how its not possible to do this inside the Update function!)
 		if (OwnerRole == ROLE_Authority && FlyingMovementCVars::RequestMispredict)
 		{
diff --git a/Engine/Plugins/Runtime/NetworkPrediction/Source/NetworkPrediction/Private/Movement/NetworkInterpolation.cpp b/Engine/Plugins/Runtime/NetworkPrediction/Source/NetworkPrediction/Private/Movement/NetworkInterpolation.cpp
new file mode 100644
--- /dev/null
+++ b/Engine/Plugins/Runtime/NetworkPrediction/Source/NetworkPrediction/Private/Movement/NetworkInterpolation.cpp
@@ -0,0 +1 @@
+// Copyright 1998-2019 Epic Games, Inc. All Rights Reserved.
\ No newline at end of file
diff --git a/Engine/Plugins/Runtime/NetworkPrediction/Source/NetworkPrediction/Private/Movement/ParametricMovement.cpp b/Engine/Plugins/Runtime/NetworkPrediction/Source/NetworkPrediction/Private/Movement/ParametricMovement.cpp
--- a/Engine/Plugins/Runtime/NetworkPrediction/Source/NetworkPrediction/Private/Movement/ParametricMovement.cpp
+++ b/Engine/Plugins/Runtime/NetworkPrediction/Source/NetworkPrediction/Private/Movement/ParametricMovement.cpp
@@ -188,14 +188,26 @@ void UParametricMovementComponent::SetTransformForPosition(const float InPositio
 void UParametricMovementComponent::TickComponent(float DeltaTime, enum ELevelTick TickType, FActorComponentTickFunction* ThisTickFunction)
 {
 	Super::TickComponent(DeltaTime, TickType, ThisTickFunction);
+
+	TickSimulation(DeltaTime); // fixme
+}
+
+void UParametricMovementComponent::Reconcile()
+{
+
+}
+
+void UParametricMovementComponent::TickSimulation(float DeltaTimeSeconds)
+{
 	const ENetRole OwnerRole = GetOwnerRole();
 
+	PreTickSimulation(DeltaTimeSeconds); // Fixme
+
 	if (NetworkSim && OwnerRole != ROLE_None)
 	{
-		
 		if (OwnerRole == ROLE_Authority)
 		{
-			if (ParametricMovement::FInputCmd* NextCmd = NetworkSim->GetNextInputForWrite(DeltaTime))
+			if (ParametricMovement::FInputCmd* NextCmd = NetworkSim->GetNextInputForWrite(DeltaTimeSeconds))
 			{
 				NextCmd->PlayRate = PendingPlayRate;
 				PendingPlayRate.Reset();
@@ -204,7 +216,7 @@ void UParametricMovementComponent::TickComponent(float DeltaTime, enum ELevelTic
 
 		ParametricMovement::FMovementSystem::FTickParameters Parameters;
 		Parameters.Role = OwnerRole;
-		Parameters.LocalDeltaTimeSeconds = DeltaTime;
+		Parameters.LocalDeltaTimeSeconds = DeltaTimeSeconds;
 
 		// Tick the core network sim, this will consume input and generate new sync state
 		NetworkSim->Tick((IParametricMovementDriver*)this, Parameters);
diff --git a/Engine/Plugins/Runtime/NetworkPrediction/Source/NetworkPrediction/Private/NetworkPredictionComponent.cpp b/Engine/Plugins/Runtime/NetworkPrediction/Source/NetworkPrediction/Private/NetworkPredictionComponent.cpp
--- a/Engine/Plugins/Runtime/NetworkPrediction/Source/NetworkPrediction/Private/NetworkPredictionComponent.cpp
+++ b/Engine/Plugins/Runtime/NetworkPrediction/Source/NetworkPrediction/Private/NetworkPredictionComponent.cpp
@@ -93,7 +93,7 @@ bool UNetworkPredictionComponent::IsLocallyControlled()
 	return bIsLocallyControlled;
 }
 
-void UNetworkPredictionComponent::TickComponent(float DeltaTime, enum ELevelTick TickType, FActorComponentTickFunction *ThisTickFunction)
+void UNetworkPredictionComponent::PreTickSimulation(float DeltaTime)
 {
 	CheckOwnerRoleChange();
 	if (PreTickLocallyControlledSim.IsBound() && IsLocallyControlled())
diff --git a/Engine/Plugins/Runtime/NetworkPrediction/Source/NetworkPrediction/Private/NetworkPredictionTypes.cpp b/Engine/Plugins/Runtime/NetworkPrediction/Source/NetworkPrediction/Private/NetworkPredictionTypes.cpp
new file mode 100644
--- /dev/null
+++ b/Engine/Plugins/Runtime/NetworkPrediction/Source/NetworkPrediction/Private/NetworkPredictionTypes.cpp
@@ -0,0 +1,114 @@
+// Copyright 1998-2019 Epic Games, Inc. All Rights Reserved.
+
+#include "NetworkPredictionTypes.h"
+#include "Engine/NetConnection.h"
+
+DEFINE_LOG_CATEGORY(LogNetworkSim);
+
+FColor FVisualLoggingParameters::DebugColors[(int32)EVisualLoggingContext::MAX] =
+{
+	FColor::Cyan,		// LastPredicted
+	FColor::Cyan,		// OtherPredicted
+	FColor::Green,		// LastConfirmed
+	FColor::Red,		// FirstMispredicted
+	FColor(255,81,101),	// OtherMispredicted
+	FColor(255,91,91),	// LastMispredicted
+	FColor::Purple,		// CurrentServerPIE
+};
+
+// -------------------------------------------------------------------------------------------------------------------------------
+//	FReplicationProxy
+// -------------------------------------------------------------------------------------------------------------------------------
+
+void FReplicationProxy::Init(class IReplicationProxy* InNetworkSimModel, EReplicationProxyTarget InReplicationTarget)
+{
+	NetworkSimModel = InNetworkSimModel;
+	ReplicationTarget = InReplicationTarget;
+}
+
+bool FReplicationProxy::NetSerialize(FArchive& Ar, class UPackageMap* Map, bool& bOutSuccess)
+{
+	NetworkSimModel->NetSerializeProxy(ReplicationTarget, Ar);
+	return true;
+}
+
+void FReplicationProxy::OnPreReplication()
+{
+	CachedDirtyCount = NetworkSimModel->GetProxyDirtyCount(ReplicationTarget);
+}
+	
+bool FReplicationProxy::Identical(const FReplicationProxy* Other, uint32 PortFlags) const
+{
+	return (CachedDirtyCount == Other->CachedDirtyCount);
+}
+
+// -------------------------------------------------------------------------------------------------------------------------------
+//	FServerRPCProxyParameter
+// -------------------------------------------------------------------------------------------------------------------------------
+
+TArray<uint8> FServerReplicationRPCParameter::TempStorage;
+
+bool FServerReplicationRPCParameter::NetSerialize(FArchive& Ar, class UPackageMap* Map, bool& bOutSuccess)
+{
+	if (Ar.IsLoading())
+	{
+		// Loading: serialize to temp storage. We'll do the real deserialize in a manual call to ::NetSerializeToProxy
+		FNetBitReader& BitReader = (FNetBitReader&)Ar;
+		CachedNumBits = BitReader.GetBitsLeft();
+		CachedPackageMap = Map;
+
+		const int64 BytesLeft = BitReader.GetBytesLeft();
+		check(BytesLeft > 0); // Should not possibly be able to get here with an empty archive
+		TempStorage.Reset(BytesLeft);
+		TempStorage.SetNumUninitialized(BytesLeft);
+		TempStorage.Last() = 0;
+
+		BitReader.SerializeBits(TempStorage.GetData(), CachedNumBits);
+	}
+	else
+	{
+		// Saving: directly call into the proxy's NetSerialize. No need for temp storage.
+		check(Proxy); // Must have been set before, via ctor.
+		return Proxy->NetSerialize(Ar, Map, bOutSuccess);
+	}
+
+	return true;
+}
+
+void FServerReplicationRPCParameter::NetSerializeToProxy(FReplicationProxy& InProxy)
+{
+	check(CachedPackageMap != nullptr);
+	check(CachedNumBits != -1);
+
+	FNetBitReader BitReader(CachedPackageMap, TempStorage.GetData(), CachedNumBits);
+
+	bool bOutSuccess = true;
+	InProxy.NetSerialize(BitReader, CachedPackageMap, bOutSuccess);
+
+	CachedNumBits = -1;
+	CachedPackageMap = nullptr;
+}
+
+// -------------------------------------------------------------------------------------------------------------------------------
+//	FScopedBandwidthLimitBypass
+// -------------------------------------------------------------------------------------------------------------------------------
+
+FScopedBandwidthLimitBypass::FScopedBandwidthLimitBypass(AActor* OwnerActor)
+{
+	if (OwnerActor)
+	{
+		CachedNetConnection = OwnerActor->GetNetConnection();
+		if (CachedNetConnection)
+		{
+			RestoreBits = CachedNetConnection->QueuedBits + CachedNetConnection->SendBuffer.GetNumBits();
+		}
+	}
+}
+
+FScopedBandwidthLimitBypass::~FScopedBandwidthLimitBypass()
+{
+	if (CachedNetConnection)
+	{
+		CachedNetConnection->QueuedBits = RestoreBits - CachedNetConnection->SendBuffer.GetNumBits();
+	}
+}
\ No newline at end of file
diff --git a/Engine/Plugins/Runtime/NetworkPrediction/Source/NetworkPrediction/Private/NetworkSimulationGlobalManager.cpp b/Engine/Plugins/Runtime/NetworkPrediction/Source/NetworkPrediction/Private/NetworkSimulationGlobalManager.cpp
new file mode 100644
--- /dev/null
+++ b/Engine/Plugins/Runtime/NetworkPrediction/Source/NetworkPrediction/Private/NetworkSimulationGlobalManager.cpp
@@ -0,0 +1,16 @@
+#include "NetworkSimulationGlobalManager.h"
+
+UNetworkSimulationGlobalManager::UNetworkSimulationGlobalManager()
+{
+
+}
+
+void UNetworkSimulationGlobalManager::ReconcileSimulationsPostNetworkUpdate()
+{
+
+}
+
+void UNetworkSimulationGlobalManager::TickNewSimulationFrame(float DeltaTimeSeconds)
+{
+
+}
\ No newline at end of file
diff --git a/Engine/Plugins/Runtime/NetworkPrediction/Source/NetworkPrediction/Private/NetworkSimulationModel.cpp b/Engine/Plugins/Runtime/NetworkPrediction/Source/NetworkPrediction/Private/NetworkSimulationModel.cpp
deleted file mode 100644
diff --git a/Engine/Plugins/Runtime/NetworkPrediction/Source/NetworkPrediction/Private/NetworkSimulationModelDebugger.cpp b/Engine/Plugins/Runtime/NetworkPrediction/Source/NetworkPrediction/Private/NetworkSimulationModelDebugger.cpp
--- a/Engine/Plugins/Runtime/NetworkPrediction/Source/NetworkPrediction/Private/NetworkSimulationModelDebugger.cpp
+++ b/Engine/Plugins/Runtime/NetworkPrediction/Source/NetworkPrediction/Private/NetworkSimulationModelDebugger.cpp
@@ -1,7 +1,7 @@
 // Copyright 1998-2019 Epic Games, Inc. All Rights Reserved.
 
 #include "NetworkSimulationModelDebugger.h"
-#include "NetworkSimulationModel.h"
+#include "NetworkPredictionTypes.h"
 #include "Engine/World.h"
 #include "Engine/LocalPlayer.h"
 #include "GameFramework/Pawn.h"
diff --git a/Engine/Plugins/Runtime/NetworkPrediction/Source/NetworkPrediction/Public/Movement/FlyingMovement.h b/Engine/Plugins/Runtime/NetworkPrediction/Source/NetworkPrediction/Public/Movement/FlyingMovement.h
--- a/Engine/Plugins/Runtime/NetworkPrediction/Source/NetworkPrediction/Public/Movement/FlyingMovement.h
+++ b/Engine/Plugins/Runtime/NetworkPrediction/Source/NetworkPrediction/Public/Movement/FlyingMovement.h
@@ -13,7 +13,7 @@
 #include "Components/SceneComponent.h"
 #include "Misc/OutputDevice.h"
 #include "Misc/CoreDelegates.h"
-#include "NetworkSimulationModelTemplates.h"
+#include "NetworkSimulationModel.h"
 #include "BaseMovementComponent.h"
 
 #include "FlyingMovement.generated.h"
@@ -110,8 +110,10 @@ namespace FlyingMovement
 		}
 	};
 
+	using TMovementBufferTypes = TNetworkSimBufferTypes<FInputCmd, FMoveState, FAuxState>;
+
 	// Actual definition of our network simulation.
-	class FMovementSystem : public TNetworkedSimulationModel<FMovementSystem, FInputCmd, FMoveState, FAuxState>
+	class FMovementSystem : public TNetworkedSimulationModel<FMovementSystem, TMovementBufferTypes>
 	{
 	public:
 
diff --git a/Engine/Plugins/Runtime/NetworkPrediction/Source/NetworkPrediction/Public/Movement/ParametricMovement.h b/Engine/Plugins/Runtime/NetworkPrediction/Source/NetworkPrediction/Public/Movement/ParametricMovement.h
--- a/Engine/Plugins/Runtime/NetworkPrediction/Source/NetworkPrediction/Public/Movement/ParametricMovement.h
+++ b/Engine/Plugins/Runtime/NetworkPrediction/Source/NetworkPrediction/Public/Movement/ParametricMovement.h
@@ -2,7 +2,7 @@
 
 #pragma once
 
-#include "NetworkSimulationModelTemplates.h"
+#include "NetworkSimulationModel.h"
 #include "BaseMovementComponent.h"
 
 #include "ParametricMovement.generated.h"
@@ -91,8 +91,10 @@ namespace ParametricMovement
 		}
 	};
 
+	using TMovementBufferTypes = TNetworkSimBufferTypes<FInputCmd, FMoveState, FAuxState>;
+
 	// Actual definition of our network simulation.
-	class FMovementSystem : public TNetworkedSimulationModel<FMovementSystem, FInputCmd, FMoveState, FAuxState>
+	class FMovementSystem : public TNetworkedSimulationModel<FMovementSystem, TMovementBufferTypes>
 	{
 	public:
 
@@ -144,6 +146,9 @@ class NETWORKPREDICTION_API UParametricMovementComponent : public UBaseMovementC
 	virtual void BeginPlay() override;
 	virtual void TickComponent(float DeltaTime, enum ELevelTick TickType, FActorComponentTickFunction* ThisTickFunction) override;
 
+	virtual void Reconcile() override;
+	virtual void TickSimulation(float DeltaTimeSeconds) override;
+
 	// Base TNetworkModelSimulation driver
 	void InitSyncState(ParametricMovement::FMoveState& OutSyncState) const override;
 	void FinalizeFrame(const ParametricMovement::FMoveState& SyncState) override;
diff --git a/Engine/Plugins/Runtime/NetworkPrediction/Source/NetworkPrediction/Public/NetworkInterpolation.h b/Engine/Plugins/Runtime/NetworkPrediction/Source/NetworkPrediction/Public/NetworkInterpolation.h
new file mode 100644
--- /dev/null
+++ b/Engine/Plugins/Runtime/NetworkPrediction/Source/NetworkPrediction/Public/NetworkInterpolation.h
@@ -0,0 +1,16 @@
+// Copyright 1998-2019 Epic Games, Inc. All Rights Reserved.
+
+#pragma once
+
+class UNetworkInterpolationManager
+{
+
+};
+
+
+
+template <typename TState>
+class TNetworkInterpolationHelper
+{
+
+};
\ No newline at end of file
diff --git a/Engine/Plugins/Runtime/NetworkPrediction/Source/NetworkPrediction/Public/NetworkPredictionComponent.h b/Engine/Plugins/Runtime/NetworkPrediction/Source/NetworkPrediction/Public/NetworkPredictionComponent.h
--- a/Engine/Plugins/Runtime/NetworkPrediction/Source/NetworkPrediction/Public/NetworkPredictionComponent.h
+++ b/Engine/Plugins/Runtime/NetworkPrediction/Source/NetworkPrediction/Public/NetworkPredictionComponent.h
@@ -2,10 +2,16 @@
 
 #pragma once
 #include "Components/ActorComponent.h"
-#include "NetworkSimulationModel.h"
+#include "NetworkPredictionTypes.h"
 
 #include "NetworkPredictionComponent.generated.h"
 
+class INetworkSimulationOwner
+{
+	virtual void Reconcile() = 0;
+	virtual void TickSimulation(float DeltaTimeSeconds) = 0;
+};
+
 // ----------------------------------------------------------------------------------------------------------------------------------------------------------------------
 //	UNetworkPredictionComponent
 //	This is the base component for running a TNetworkedSimulationModel through an actor component. This contains the boiler plate hooks into getting the system
@@ -30,15 +36,17 @@ public:
 	virtual void PreReplication(IRepChangedPropertyTracker & ChangedPropertyTracker) override;
 	virtual void PreNetReceive() override;
 
-	// TickComponent calls CheckOwnerRoleChange and invokes PreTickLocallyControlledSim. You should always call it at the top of your child's tick implementation.
-	virtual void TickComponent(float DeltaTime, enum ELevelTick TickType, FActorComponentTickFunction *ThisTickFunction) override;
-
 	// PreSimulationTick callback.
 	// Called right in ::TickComponent for autonomous proxies. This is optional and intended to be a convencience for generating local input.
 	// If you don't use this function, you should manually be submitting client input prior to simulation tick (via GetNextClientInputCmdForWrite).
 	// (It is important to submit input prior to ticking the sim. If you get it backwards, you introduce local latency and cause problems with variable time steps).
 	DECLARE_DELEGATE_OneParam(FPreSimTickDelegate, float /* DeltaSeconds */)
 	void SetLocallyControlledPreTick(const FPreSimTickDelegate& PreSimTick) { PreTickLocallyControlledSim = PreSimTick; }
+	
+	virtual void Reconcile() { }
+	virtual void TickSimulation(float DeltaTimeSeconds) { }
+
+	void PreTickSimulation(float DeltaTime);
 
 protected:
 
diff --git a/Engine/Plugins/Runtime/NetworkPrediction/Source/NetworkPrediction/Public/NetworkPredictionTypes.h b/Engine/Plugins/Runtime/NetworkPrediction/Source/NetworkPrediction/Public/NetworkPredictionTypes.h
new file mode 100644
--- /dev/null
+++ b/Engine/Plugins/Runtime/NetworkPrediction/Source/NetworkPrediction/Public/NetworkPredictionTypes.h
@@ -0,0 +1,242 @@
+// Copyright 1998-2019 Epic Games, Inc. All Rights Reserved.
+
+#pragma once
+
+#include "Engine/EngineTypes.h"
+#include "Net/UnrealNetwork.h" // For MakeRelative
+#include "NetworkPredictionTypes.generated.h"
+
+// -------------------------------------------------------------------------------------------------------------------------------
+//	This file contains supporting types that are used by both the UE4 side (UNetworkPredictionComponent etc) and templated
+//	network simulation model side (TNetworkedSimulationModel).
+// -------------------------------------------------------------------------------------------------------------------------------
+
+#define NETSIM_MODEL_DEBUG !(UE_BUILD_SHIPPING || UE_BUILD_TEST)
+
+NETWORKPREDICTION_API DECLARE_LOG_CATEGORY_EXTERN(LogNetworkSim, Log, All);
+
+UENUM()
+enum class EReplicationProxyTarget: uint8
+{
+	ServerRPC,			// Client -> Server
+	AutonomousProxy,	// Owning/Controlling client
+	SimulatedProxy,		// Non owning client
+	Replay,				// Replay net driver
+	Debug,				// Replication target that is disabled in shipping
+};
+
+inline FString LexToString(EReplicationProxyTarget A)
+{
+	return *UEnum::GetValueAsString(TEXT("NetworkPrediction.EReplicationProxyTarget"), A);
+}
+
+// The parameters for NetSerialize that are passed around the system. Everything should use this, expecting to have to add more.
+struct FNetSerializeParams
+{
+	FNetSerializeParams(FArchive& InAr) : Ar(InAr) { }
+	FArchive& Ar;
+};
+
+struct FNetworkSimulationModelInitParameters
+{
+	int32 InputBufferSize=0;
+	int32 SyncedBufferSize=0;
+	int32 AuxBufferSize=0;
+	int32 DebugBufferSize=0;
+	int32 HistoricBufferSize=0;
+};
+
+struct FNetSimProcessedFrameDebugInfo
+{
+	int32 LocalGFrameNumber; // Local GFrame number
+	float LocalDeltaTimeSeconds; // Local frame time
+	TArray<int32> ProcessedKeyframes; // Which keyframes were processed this frame
+	int32 LastProcessedKeyframe; // What LastProcessedKeyframe was at the end of the frame. Does NOT mean we processed it this frame!
+	int32 HeadKeyframe; // Head keyframe of the inputbuffer when the frame ended
+
+	void NetSerialize(const FNetSerializeParams& P)
+	{	
+		P.Ar << LocalGFrameNumber;
+		P.Ar << LocalDeltaTimeSeconds;
+		P.Ar << LastProcessedKeyframe;
+		P.Ar << HeadKeyframe;
+
+		SafeNetSerializeTArray_Default<31>(P.Ar, ProcessedKeyframes);
+	}
+};
+
+enum class EStandardLoggingContext: uint8
+{
+	HeaderOnly,		// Minimal logging
+	Full			// "Logs everything"
+};
+
+struct FStandardLoggingParameters
+{
+	FStandardLoggingParameters(FOutputDevice* InAr, EStandardLoggingContext InContext, int32 InKeyframe)
+		: Ar(InAr), Context(InContext), Keyframe(InKeyframe) { }
+	FOutputDevice* Ar;
+	EStandardLoggingContext Context;
+	int32 Keyframe;
+};
+
+UENUM()
+enum class EVisualLoggingContext: uint8
+{
+	LastPredicted,		// The last predicted state. Where the character "is now"
+	OtherPredicted,		// "Middle" states between LastComfirmed and LastPredicted. Recommend to draw small crumbs (not full model/cylinder/etc) to avoid washing out the scene.
+	LastConfirmed,		// The last confirmed state from the server
+	FirstMispredicted,	// The first state that was mispredicted on the client. This correlates with LastConfirmed.
+	OtherMispredicted,	// subsequent mispredicted states on the client. Recommend to draw small crumbs (not full model/cylinder/etc) to avoid washing out the scene.
+	LastMispredicted,	// The latest predictive state when a misprediction was detected.
+	CurrentServerPIE,	// The current server position *right now*. Only available in PIE
+	MAX
+};
+
+inline FString LexToString(EVisualLoggingContext A)
+{
+	return *UEnum::GetValueAsString(TEXT("NetworkPrediction.EVisualLoggingContext"), A);
+}
+
+enum class EVisualLoggingDrawType : uint8
+{
+	Full,				// Draw "the full thing" (maybe a collision capsule for example)
+	Crumb,				// Draw a small/minimal representation (E.g, a point or small axis)
+};
+
+enum class EVisualLoggingLifetime : uint8
+{
+	Transient,			// This logging is transient and will (probably) be done every frame. Don't persist (in contexts where that would matter).
+	Persistent,			// This is more of a persistent/one-off event that should be drawn for some longer amount of time (probably configurable at whatever level is doing the logging)
+};
+
+struct FVisualLoggingParameters
+{
+	FVisualLoggingParameters(EVisualLoggingContext InContext, int32 InKeyframe, EVisualLoggingLifetime InLifetime) :
+		 Context(InContext), Keyframe(InKeyframe), Lifetime(InLifetime) { }
+
+	EVisualLoggingContext Context;
+	int32 Keyframe;
+	EVisualLoggingLifetime Lifetime;
+
+	FColor GetDebugColor() const { return DebugColors[(int32)Context]; }
+	static NETWORKPREDICTION_API FColor DebugColors[(int32)EVisualLoggingContext::MAX];
+};
+
+UENUM()
+enum class ESimulatedUpdateMode : uint8
+{
+	Interpolate,		// Update from previous to current known states from the server. This puts the simulation further "behind" due to having to buffer the known state (but is never "wrong" and doesn't require a simulation update)
+	Extrapolate,		// Extrapolate the simulation once per local frame, by synthesizing (guessing) input commands
+	ForwardPredict		// Predict the simulation ahead. For a simulated proxy to do this, it must be tied to an autonomous proxy
+};
+
+inline FString LexToString(ESimulatedUpdateMode A)
+{
+	return *UEnum::GetValueAsString(TEXT("NetworkPrediction.ESimulatedUpdateMode"), A);
+}
+
+// -------------------------------------------------------------------------------------------------------------------------------
+// Ticking parameters use to drive the simulation
+// -------------------------------------------------------------------------------------------------------------------------------
+
+class FNetworkSimulationTickingParameters
+{
+	float LocalFrameDeltaTimeSeconds;
+};
+
+
+// -------------------------------------------------------------------------------------------------------------------------------
+// Interface that the proxy talks to. This is what will implement the replication.
+// -------------------------------------------------------------------------------------------------------------------------------
+class IReplicationProxy
+{
+public:
+
+	virtual ~IReplicationProxy() { }
+	
+	virtual void NetSerializeProxy(EReplicationProxyTarget Target, const FNetSerializeParams& Params) = 0;
+	virtual int32 GetProxyDirtyCount(EReplicationProxyTarget Target) = 0;
+};
+
+// -------------------------------------------------------------------------------------------------------------------------------
+//	FReplicationProxy
+//	Replicated USTRUCT that point to the IReplicationProxy to do the replication.
+// -------------------------------------------------------------------------------------------------------------------------------
+USTRUCT()
+struct FReplicationProxy
+{
+	GENERATED_BODY()
+
+	void Init(IReplicationProxy* InNetworkSimModel, EReplicationProxyTarget InReplicationTarget);
+	bool NetSerialize(FArchive& Ar, class UPackageMap* Map, bool& bOutSuccess);
+	void OnPreReplication();	
+	bool Identical(const FReplicationProxy* Other, uint32 PortFlags) const;
+	
+private:
+
+	IReplicationProxy* NetworkSimModel;
+	EReplicationProxyTarget ReplicationTarget;
+	int32 CachedDirtyCount = 0;
+};
+
+template<>
+struct TStructOpsTypeTraits<FReplicationProxy> : public TStructOpsTypeTraitsBase2<FReplicationProxy>
+{
+	enum
+	{
+		WithNetSerializer = true,
+		WithIdentical = true,
+	};
+};
+
+// -------------------------------------------------------------------------------------------------------------------------------
+//	FServerRPCProxyParameter
+//	Used for the client->Server RPC. Since this is instantiated on the stack by the replication system prior to net serializing,
+//	we have no opportunity to point the RPC parameter to the member variables we want. So we serialize into a generic temp byte buffer
+//	and move them into the real buffers on the component in the RPC body (via ::NetSerialzeToProxy).
+// -------------------------------------------------------------------------------------------------------------------------------
+USTRUCT()
+struct FServerReplicationRPCParameter
+{
+	GENERATED_BODY()
+
+	// Receive flow: ctor() -> NetSerializetoProxy
+	FServerReplicationRPCParameter() : Proxy(nullptr)	{ }
+	void NetSerializeToProxy(FReplicationProxy& InProxy);
+
+	// Send flow: ctor(Proxy) -> NetSerialize
+	FServerReplicationRPCParameter(FReplicationProxy& InProxy) : Proxy(&InProxy) { }
+	bool NetSerialize(FArchive& Ar, class UPackageMap* Map, bool& bOutSuccess);
+
+private:
+
+	static TArray<uint8> TempStorage;
+
+	FReplicationProxy* Proxy;
+	int64 CachedNumBits = -1;
+	class UPackageMap* CachedPackageMap = nullptr;
+};
+
+template<>
+struct TStructOpsTypeTraits<FServerReplicationRPCParameter> : public TStructOpsTypeTraitsBase2<FServerReplicationRPCParameter>
+{
+	enum
+	{
+		WithNetSerializer = true
+	};
+};
+
+// Helper struct to bypass the bandwidth limit imposed by the engine's NetDriver (QueuedBits, NetSpeed, etc).
+// This is really a temp measure to make the system easier to drop in/try in a project without messing with your engine settings.
+// (bandwidth optimizations have not been done yet and the system in general hasn't been stressed with packetloss / gaps in command streams)
+// So, you are free to use this in your own code but it may be removed one day. Hopefully in general bandwidth limiting will also become more robust.
+struct FScopedBandwidthLimitBypass
+{
+	FScopedBandwidthLimitBypass(AActor* OwnerActor);
+	~FScopedBandwidthLimitBypass();
+private:
+
+	int64 RestoreBits = 0;
+	class UNetConnection* CachedNetConnection = nullptr;
+};
\ No newline at end of file
diff --git a/Engine/Plugins/Runtime/NetworkPrediction/Source/NetworkPrediction/Public/NetworkSimulationGlobalManager.h b/Engine/Plugins/Runtime/NetworkPrediction/Source/NetworkPrediction/Public/NetworkSimulationGlobalManager.h
new file mode 100644
--- /dev/null
+++ b/Engine/Plugins/Runtime/NetworkPrediction/Source/NetworkPrediction/Public/NetworkSimulationGlobalManager.h
@@ -0,0 +1,48 @@
+// Copyright 1998-2019 Epic Games, Inc. All Rights Reserved.
+
+#pragma once
+
+#include "CoreMinimal.h"
+#include "UObject/ObjectMacros.h"
+#include "UObject/UObjectGlobals.h"
+#include "UObject/Object.h"
+
+#include "NetworkSimulationGlobalManager.generated.h"
+
+// ----------------------------------------------------------------------------------------------------------------------------
+//	The purpose of the global manager is to enforce frame boundaries and ordering amongst all instances of network simulations.
+//	Other than registration it is not something much code should interact with. It mainly manages when network sims reconcile and tick.
+//
+//
+//	Frame Overview:
+//
+//		Network TickFlush (receive network traffic)
+//			NetSerialize is called on individual net sims (in an essentially random order)
+//				RepNotifies, PostNetReceive, etc are called on actors *right after they individually receive traffic* (this is kind of unfortunate but how it is)
+//		
+//		UNetworkSimulationGlobalManager::ReconcilePostNetworkUpdate	 (Reconcile is called in deterministic order on all network sims)
+//			(think: "Get yourself right with the latest update from the server *before* we start a new frame")
+//
+//		UNetworkSimulationGlobalManager::BeginNewSimulationFrame ("Add" "allowed" time into the global "space" which will allow sims to continue to tick)
+//			Simulations are ticked in deterministic order (if registered with UNetworkSimulationGlobalManager)
+//
+//		UE4 Actors Tick
+//
+//		Network Tick Dispatch (send outgoing network traffic)
+//
+//
+// ----------------------------------------------------------------------------------------------------------------------------
+
+UCLASS()
+class UNetworkSimulationGlobalManager : public UObject
+{
+	GENERATED_BODY()
+
+public:
+
+	UNetworkSimulationGlobalManager();
+
+	void ReconcileSimulationsPostNetworkUpdate();
+	void TickNewSimulationFrame(float DeltaTimeSeconds);
+
+};
\ No newline at end of file
diff --git a/Engine/Plugins/Runtime/NetworkPrediction/Source/NetworkPrediction/Public/NetworkSimulationModel.h b/Engine/Plugins/Runtime/NetworkPrediction/Source/NetworkPrediction/Public/NetworkSimulationModel.h
--- a/Engine/Plugins/Runtime/NetworkPrediction/Source/NetworkPrediction/Public/NetworkSimulationModel.h
+++ b/Engine/Plugins/Runtime/NetworkPrediction/Source/NetworkPrediction/Public/NetworkSimulationModel.h
@@ -1,232 +1,799 @@
-// Copyright 1998-2019 Epic Games, Inc. All Rights Reserved.
+// Copyright 1998-2019 Epic Game s, Inc. All Rights Reserved.
 
 #pragma once
+#include "NetworkPredictionTypes.h"
+#include "NetworkSimulationModelBuffer.h"
+#include "NetworkSimulationModelTypes.h"
 
-#include "Engine/EngineTypes.h"
-#include "Net/UnrealNetwork.h" // For MakeRelative
-#include "NetworkSimulationModel.generated.h"
+// ----------------------------------------------------------------------------------------------------------------------------------------------------------------------
+//	CVars and compile time constants
+// ----------------------------------------------------------------------------------------------------------------------------------------------------------------------
 
-// -------------------------------------------------------------------------------------------------------------------------------
-//	This file contains supporting types that are used by both the UE4 side (UNetworkPredictionComponent etc) and templated
-//	network simulation model side (TNetworkedSimulationModel).
-// -------------------------------------------------------------------------------------------------------------------------------
-
-#define NETSIM_MODEL_DEBUG !(UE_BUILD_SHIPPING || UE_BUILD_TEST)
-
-NETWORKPREDICTION_API DECLARE_LOG_CATEGORY_EXTERN(LogNetworkSim, Log, All);
-
-UENUM()
-enum class EReplicationProxyTarget: uint8
+namespace NetworkSimulationModelCVars
 {
-	ServerRPC,			// Client -> Server
-	AutonomousProxy,	// Owning/Controlling client
-	SimulatedProxy,		// Non owning client
-	Replay,				// Replay net driver
-	Debug,				// Replication target that is disabled in shipping
-};
+static int32 EnableLocalPrediction = 1;
+static FAutoConsoleVariableRef CVarEnableLocalPrediction(TEXT("ns.EnableLocalPrediction"),
+	EnableLocalPrediction, TEXT("Toggle local prediction. Meant for debugging."), ECVF_Default);
 
-inline FString LexToString(EReplicationProxyTarget A)
-{
-	return *UEnum::GetValueAsString(TEXT("NetworkPrediction.EReplicationProxyTarget"), A);
+static int32 ForceReconcile = 0;
+static FAutoConsoleVariableRef CVarForceReconcile(TEXT("ns.ForceReconcile"),
+	ForceReconcile, TEXT("Forces reconcile even if state does not differ. E.g, force resimulation after every netupdate."), ECVF_Default);
+
+static int32 MaxInputCmdsFrame = 1;
+static FAutoConsoleVariableRef CVarMaxInputCmdsFrame(TEXT("ns.MaxInputCmdsFrame"),
+	MaxInputCmdsFrame, TEXT("Max cap on how many input cmds can be processed in a frame."), ECVF_Default);
 }
 
-// The parameters for NetSerialize that are passed around the system. Everything should use this, expecting to have to add more.
-struct FNetSerializeParams
-{
-	FNetSerializeParams(FArchive& InAr) : Ar(InAr) { }
-	FArchive& Ar;
-};
+#ifndef NETSIM_NETCONSTANT_NUM_BITS_KEYFRAME
+	#define NETSIM_NETCONSTANT_NUM_BITS_KEYFRAME 8	// Allows you to override this setting via UBT, but access via cleaner FActorMotionNetworkingConstants::NUM_BITS_KEYFRAME
+#endif
 
-struct FNetworkSimulationModelInitParameters
+struct FNetworkSimulationSerialization
 {
-	int32 InputBufferSize=0;
-	int32 SyncedBufferSize=0;
-	int32 AuxBufferSize=0;
-	int32 DebugBufferSize=0;
-	int32 HistoricBufferSize=0;
-};
+	// How many bits we use to encode the key frame number for buffers.
+	// Client Keyframes are stored locally as 32 bit integers, but we use a smaller # of bits to NetSerialize.
+	// Frames are only relatively relevant: the absolute value doesn't really matter. We just need to detect newer/older.
+	enum { NUM_BITS_KEYFRAME = NETSIM_NETCONSTANT_NUM_BITS_KEYFRAME };		
 
-struct FNetSimProcessedFrameDebugInfo
-{
-	int32 LocalGFrameNumber; // Local GFrame number
-	float LocalDeltaTimeSeconds; // Local frame time
-	TArray<int32> ProcessedKeyframes; // Which keyframes were processed this frame
-	int32 LastProcessedKeyframe; // What LastProcessedKeyframe was at the end of the frame. Does NOT mean we processed it this frame!
-	int32 HeadKeyframe; // Head keyframe of the inputbuffer when the frame ended
+	// Abs max value we encode into the bit writer
+	enum { MAX_KEYFRAME_WRITE = 1 << NUM_BITS_KEYFRAME };
 
-	void NetSerialize(const FNetSerializeParams& P)
-	{	
-		P.Ar << LocalGFrameNumber;
-		P.Ar << LocalDeltaTimeSeconds;
-		P.Ar << LastProcessedKeyframe;
-		P.Ar << HeadKeyframe;
+	// This is the threshold at which we would wrap around and incorrectly assign a keyframe on the receiving side.
+	// E.g, If there are KEYFRAME_ERROR_THRESHOLD keyframes that do not make it across from sender->receiver, the
+	// receiver will have incorrect local values. With 8 bits, this works out to be 128 frames or about 2 seconds at 60fps.
+	enum { KEYFRAME_ERROR_THRESHOLD = MAX_KEYFRAME_WRITE / 2};
 
-		SafeNetSerializeTArray_Default<31>(P.Ar, ProcessedKeyframes);
+	// Helper to serialize the int32 HeadKeyframe. Returns the unpacked value (this will be same as input in the save path)
+	static int32 SerializeKeyframe(FArchive& Ar, int32 LocalHeadKeyframe)
+	{
+		if (Ar.IsSaving())
+		{
+			((FNetBitWriter&)Ar).WriteIntWrapped( LocalHeadKeyframe, MAX_KEYFRAME_WRITE );
+			return LocalHeadKeyframe;
+		}
+		
+		return MakeRelative(((FNetBitReader&)Ar).ReadInt( MAX_KEYFRAME_WRITE ), LocalHeadKeyframe, MAX_KEYFRAME_WRITE );
 	}
 };
 
-enum class EStandardLoggingContext: uint8
-{
-	HeaderOnly,		// Minimal logging
-	Full			// "Logs everything"
-};
 
-struct FStandardLoggingParameters
-{
-	FStandardLoggingParameters(FOutputDevice* InAr, EStandardLoggingContext InContext, int32 InKeyframe)
-		: Ar(InAr), Context(InContext), Keyframe(InKeyframe) { }
-	FOutputDevice* Ar;
-	EStandardLoggingContext Context;
-	int32 Keyframe;
-};
+// ----------------------------------------------------------------------------------------------------------------------------------------------------------------------
+//	Templated Replicators/Reconciliars
+// ----------------------------------------------------------------------------------------------------------------------------------------------------------------------
 
-UENUM()
-enum class EVisualLoggingContext: uint8
+/** Replicates only the latest element. Does not synchronize keyframe */
+template<typename TBufferTypes, ENetworkSimBufferTypeId BufferId>
+struct TReplicator_LatestOnly
 {
-	LastPredicted,		// The last predicted state. Where the character "is now"
-	OtherPredicted,		// "Middle" states between LastComfirmed and LastPredicted. Recommend to draw small crumbs (not full model/cylinder/etc) to avoid washing out the scene.
-	LastConfirmed,		// The last confirmed state from the server
-	FirstMispredicted,	// The first state that was mispredicted on the client. This correlates with LastConfirmed.
-	OtherMispredicted,	// subsequent mispredicted states on the client. Recommend to draw small crumbs (not full model/cylinder/etc) to avoid washing out the scene.
-	LastMispredicted,	// The latest predictive state when a misprediction was detected.
-	CurrentServerPIE,	// The current server position *right now*. Only available in PIE
-	MAX
-};
+	using TState = typename TBufferTypes::template select_type<BufferId>::type;
 
-inline FString LexToString(EVisualLoggingContext A)
-{
-	return *UEnum::GetValueAsString(TEXT("NetworkPrediction.EVisualLoggingContext"), A);
-}
+	int32 GetProxyDirtyCount(TNetworkSimBufferContainer<TBufferTypes>& Buffers) const
+	{ 
+		return Buffers.template Get<BufferId>().GetDirtyCount(); 
+	}
 
-enum class EVisualLoggingDrawType : uint8
-{
-	Full,				// Draw "the full thing" (maybe a collision capsule for example)
-	Crumb,				// Draw a small/minimal representation (E.g, a point or small axis)
+	void NetSerialize(const FNetSerializeParams& P, TNetworkSimBufferContainer<TBufferTypes>& Buffers, const TSimulationTickInfo<TBufferTypes>& TickInfo)
+	{
+		FArchive& Ar = P.Ar;
+
+		auto& Buffer = Buffers.template Get<BufferId>();
+		TState* State = nullptr;
+		
+		if (Ar.IsSaving())
+		{
+			State = Buffer.GetElementFromHead(0);
+			check(State); // We should not be here if the buffer is empty. Want to avoid serializing an "empty" flag at the top here.
+		}
+		else
+		{
+			State = Buffer.GetWriteNext();
+		}
+
+		State->NetSerialize(Ar);
+	}
 };
 
-enum class EVisualLoggingLifetime : uint8
+template<typename TBufferTypes>
+struct TReplicator_SimulatedExtrapolatedReconciliar
 {
-	Transient,			// This logging is transient and will (probably) be done every frame. Don't persist (in contexts where that would matter).
-	Persistent,			// This is more of a persistent/one-off event that should be drawn for some longer amount of time (probably configurable at whatever level is doing the logging)
-};
+	using TSyncState = typename TBufferTypes::TSyncState;
+	
+	int32 GetProxyDirtyCount(TNetworkSimBufferContainer<TBufferTypes>& Buffers) const
+	{
+		return Buffers.Sync.GetDirtyCount(); 
+	}
 
-struct FVisualLoggingParameters
-{
-	FVisualLoggingParameters(EVisualLoggingContext InContext, int32 InKeyframe, EVisualLoggingLifetime InLifetime) :
-		 Context(InContext), Keyframe(InKeyframe), Lifetime(InLifetime) { }
+	void NetSerialize(const FNetSerializeParams& P, TNetworkSimBufferContainer<TBufferTypes>& Buffers, const TSimulationTickInfo<TBufferTypes>& TickInfo)
+	{
+		FArchive& Ar = P.Ar;
+		TSyncState* State = nullptr;
+		if (Ar.IsSaving())
+		{
+			State = Buffers.Sync.GetElementFromHead(0);
+			check(State); // We should not be here if the buffer is empty. Want to avoid serializing an "empty" flag at the top here.
+			SerializedTimeKeeper = TickInfo.ProcessedSimulationTime;
+		}
+		else
+		{
+			State = Buffers.Sync.GetWriteNext();
+			bPendingReconciliation = true;
+		}
+
+		SerializedTimeKeeper.NetSerialize(Ar);
+		State->NetSerialize(Ar);
+	}
 
-	EVisualLoggingContext Context;
-	int32 Keyframe;
-	EVisualLoggingLifetime Lifetime;
+	template<typename T, typename TDriver>
+	void Reconcile(TDriver* Driver, TNetworkSimBufferContainer<TBufferTypes>& Buffers, TSimulationTickInfo<TBufferTypes>& TickInfo)
+	{
+		if (!bPendingReconciliation)
+			return;
+
+		bPendingReconciliation = false;
+
+		// This is still being worked out. We are almost there but not quite.
+		// Processing input commands is pretty rigid right now. Since we just added a new SyncState keyframe in ::NetSerialize,
+		// we need to insert a fake input command so the extrapolation code in the NetSimModel can continue to work.
+		// This can be improved.
+
+		{
+			// Generate a new, fake, command since we just added a new sync state to head (via GetWriteNext)
+			TickInfo.GetNextInputForWrite(Buffers);
+			// Gotta fake this too
+			TickInfo.LastLocalInputGFrameNumber--; 
+			// Set our LastProcessedInputKeyframe to fake that we handled it
+			TickInfo.LastProcessedInputKeyframe = Buffers.Sync.GetHeadKeyframe();
+		}
+	}
 
-	FColor GetDebugColor() const { return DebugColors[(int32)Context]; }
-	static NETWORKPREDICTION_API FColor DebugColors[(int32)EVisualLoggingContext::MAX];
+private:
+	
+	TSimulationTimeKeeper<TBufferTypes> SerializedTimeKeeper;
+	bool bPendingReconciliation = false;
 };
 
-UENUM()
-enum class ESimulatedUpdateMode : uint8
+/** Replicates the latest sequence of N elements. N is dynamic, not compiled in. */
+template<typename TBufferTypes, ENetworkSimBufferTypeId BufferId>
+struct TReplicator_DynamicSequence
 {
-	Interpolate,		// Update from previous to current known states from the server. This puts the simulation further "behind" due to having to buffer the known state (but is never "wrong" and doesn't require a simulation update)
-	Extrapolate,		// Extrapolate the simulation once per local frame, by synthesizing (guessing) input commands
-	ForwardPredict		// Predict the simulation ahead. For a simulated proxy to do this, it must be tied to an autonomous proxy
+	int32 GetProxyDirtyCount(TNetworkSimBufferContainer<TBufferTypes>& Buffers) const 
+	{
+		return Buffers.template Get<BufferId>().GetDirtyCount(); 
+	}
+	
+	void NetSerialize(const FNetSerializeParams& P, TNetworkSimBufferContainer<TBufferTypes>& Buffers, const TSimulationTickInfo<TBufferTypes>& TickInfo)
+	{
+		auto& Buffer = Buffers.template Get<BufferId>();
+		FArchive& Ar = P.Ar;
+		uint8 SerializedNumElements = FMath::Min<uint8>(MaxNumElements, Buffer.GetNumValidElements());
+		Ar << SerializedNumElements;
+
+		const int32 HeadKeyframe = FNetworkSimulationSerialization::SerializeKeyframe(Ar, Buffer.GetHeadKeyframe());
+		const int32 StartingKeyframe = FMath::Max(0, HeadKeyframe - SerializedNumElements + 1);
+
+		if (Ar.IsLoading())
+		{
+			const int32 PrevHead = Buffer.GetHeadKeyframe();
+			if (PrevHead < StartingKeyframe && PrevHead >= 0)
+			{
+				// There is a gap in the stream. In some cases, we want this to be a "fault" and bubble up. We may want to synthesize state or maybe we just skip ahead.
+				UE_LOG(LogNetworkSim, Warning, TEXT("Fault: gap in received buffer. PrevHead: %d. Received: %d-%d. Reseting previous buffer contents"), PrevHead, StartingKeyframe, HeadKeyframe);
+			}
+
+			Buffer.ResetNextHeadKeyframe(StartingKeyframe);
+		}
+
+		for (int32 Keyframe = StartingKeyframe; Keyframe <= HeadKeyframe; ++Keyframe)
+		{
+			// This, as is, is bad. The intention is that these functions serialize multiple items in some delta compressed fashion.
+			// As is, we are just serializing the elements individually.
+			auto* Cmd = Ar.IsLoading() ? Buffer.GetWriteNext() : Buffer.FindElementByKeyframe(Keyframe);
+			Cmd->NetSerialize(P);
+		}
+	}
+
+private:
+	
+	uint8 MaxNumElements = 3;
 };
 
-inline FString LexToString(ESimulatedUpdateMode A)
+/** Replicates TSyncState and does basic reconciliation. */
+template<typename TBufferTypes>
+struct TReplicator_BasicReconciliar
 {
-	return *UEnum::GetValueAsString(TEXT("NetworkPrediction.ESimulatedUpdateMode"), A);
-}
+	using TInputCmd = typename TBufferTypes::TInputCmd;
+	using TSyncState = typename TBufferTypes::TSyncState;
+	using TAuxState = typename TBufferTypes::TAuxState;
 
-// -------------------------------------------------------------------------------------------------------------------------------
-// Interface that the proxy talks to. This is what will implement the replication.
-// -------------------------------------------------------------------------------------------------------------------------------
-class IReplicationProxy
-{
-public:
+	int32 GetLastSerializedKeyframe() const { return LastSerializedKeyframe; }
+	bool IsReconcileFaultDetected() const { return bReconcileFaultDetected; }
+	const TSimulationTimeKeeper<TBufferTypes>& GetLastSerializedSimulationTimeKeeper() const { return SerializedTimeKeeper; }
 
-	virtual ~IReplicationProxy() { }
-	
-	virtual void NetSerializeProxy(EReplicationProxyTarget Target, const FNetSerializeParams& Params) = 0;
-	virtual int32 GetProxyDirtyCount(EReplicationProxyTarget Target) = 0;
-};
+	int32 GetProxyDirtyCount(TNetworkSimBufferContainer<TBufferTypes>& Buffers) const
+	{
+		return Buffers.Sync.GetDirtyCount();
+	}
 
-// -------------------------------------------------------------------------------------------------------------------------------
-//	FReplicationProxy
-//	Replicated USTRUCT that point to the IReplicationProxy to do the replication.
-// -------------------------------------------------------------------------------------------------------------------------------
-USTRUCT()
-struct FReplicationProxy
-{
-	GENERATED_BODY()
+	// --------------------------------------------------------------------
+	//	NetSerialize
+	// --------------------------------------------------------------------
+	void NetSerialize(const FNetSerializeParams& P, TNetworkSimBufferContainer<TBufferTypes>& Buffers, const TSimulationTickInfo<TBufferTypes>& TickInfo)
+	{
+		FArchive& Ar = P.Ar;
+
+		const int32 SerializedHeadKeyframe = FNetworkSimulationSerialization::SerializeKeyframe(Ar, Buffers.Sync.GetHeadKeyframe());
+		TSyncState* SerializedState = nullptr;
+
+		// Serialize total simulation time. This isn't really necessary since we have the keyframe above. 
+		SerializedTimeKeeper = TickInfo.ProcessedSimulationTime;
+		SerializedTimeKeeper.NetSerialize(P);
+
+		if (Ar.IsSaving())
+		{
+			// Server serialize the latest state
+			SerializedState = Buffers.Sync.GetElementFromHead(0);
+		}
+		else
+		{
+			// Its possible we process two packets in between ticks. This isn't really a problem but should be rare.
+			if (bPendingReconciliation)
+			{
+				UE_LOG(LogNetworkSim, Warning, TEXT("bPendingReconciliation while in ::NetSerialize. LastSerializedKeyframe: %d. New SerializedHeadKeyframe: %d."), LastSerializedKeyframe, SerializedHeadKeyframe);
+			}
+
+			// Lazy init the reconciliation buffer. We don't need this on the server/writing side. (Fixme, with templated approach we could easily specialize this to an inline element instead of buffer?)
+			if (ReconciliationBuffer.GetMaxNumElements() == 0)
+			{
+				ReconciliationBuffer.SetBufferSize(1);
+			}
+
+			// Reset the head to match the server. This is not critical, we just need a place to store this, but it makes the code in the reconciliation a bit nicer since the reconcilated frame number = head frame.
+			ReconciliationBuffer.ResetNextHeadKeyframe(SerializedHeadKeyframe);
+			SerializedState = ReconciliationBuffer.GetWriteNext();
+		}
+
+		SerializedState->NetSerialize(Ar);
+
+		if (Ar.IsLoading())
+		{
+			bReconcileFaultDetected = false;
+			bPendingReconciliation = false;
+
+			// The state the client predicted that corresponds to the state the server just serialized to us
+			if (TSyncState* ClientExistingState = Buffers.Sync.FindElementByKeyframe(SerializedHeadKeyframe))
+			{
+				if (ClientExistingState->ShouldReconcile(*SerializedState) || (NetworkSimulationModelCVars::ForceReconcile > 0))
+				{
+					UE_CLOG(!Buffers.Input.IsValidKeyframe(SerializedHeadKeyframe), LogNetworkSim, Error, TEXT("::NetSerialize: Client InputBuffer does not contain data for frame %d. {%s} {%s}"), SerializedHeadKeyframe, *Buffers.Input.GetBasicDebugStr(), *Buffers.Sync.GetBasicDebugStr());
+					bPendingReconciliation =  true;
+				}
+			}
+			else
+			{
+				// We don't have corresponding local state. There are two cases:
+				if (NetworkSimulationModelCVars::EnableLocalPrediction) // Fixme: this is awkward. Expected if player prediction is disabled but coupling like this feels bad.
+				{
+					UE_LOG(LogNetworkSim, Warning, TEXT("::NetSerialize Fault: SyncBuffer does not contain data for frame %d. [%d-%d]"), SerializedHeadKeyframe, Buffers.Sync.GetTailKeyframe(), Buffers.Sync.GetHeadKeyframe());
+				}
+
+				// Case 1: the serialized state is older than what we've kept in our buffer. A bigger buffer would solve this! (at the price of more resimulated frames to recover when this happens)
+				if (SerializedHeadKeyframe < Buffers.Sync.GetTailKeyframe())
+				{
+					// This is a reconcile fault and we just need to chill. We'll stop sending user commands until the cmds in flight flush through the system and we catch back up.
+					bReconcileFaultDetected = true;
+				}
+				else
+				{
+					// Case 2: We've received a newer frame than what we've processed locally. This could happen if we are buffering our inputs locally (but still sending to the server).
+					// Since this doesn't require resimulating, we can just set the authoritative state here and continue on. Ticking logic will need to detect this gap and skip
+					// our LastProcessedInputKeyframe ahead.
+
+					// However, this keyframe better be a valid input. Otherwise how did the server generate it?
+					checkf(SerializedHeadKeyframe <= Buffers.Input.GetHeadKeyframe(), TEXT("Received newer motionstate that doesn't correspond to valid input cmd. SerializedHeadKeyframe: %d. {%s} {%s}"),
+						SerializedHeadKeyframe, *Buffers.Input.GetBasicDebugStr(), *Buffers.Sync.GetBasicDebugStr());
+
+					Buffers.Sync.ResetNextHeadKeyframe(SerializedHeadKeyframe);
+					TSyncState* ClientMotionState = Buffers.Sync.GetWriteNext();
+					*ClientMotionState = *SerializedState;
+				}
+			}
+		}
+
+		LastSerializedKeyframe = SerializedHeadKeyframe;
+	}
+
+	// --------------------------------------------------------------------
+	//	Reconcile
+	// --------------------------------------------------------------------
+	template<typename T, typename TDriver>
+	void Reconcile(TDriver* Driver, TNetworkSimBufferContainer<TBufferTypes>& Buffers, TSimulationTickInfo<TBufferTypes>& TickInfo)
+	{
+		if (bPendingReconciliation == false)
+		{
+			return;
+		}
+		bPendingReconciliation = false;
+		check(ReconciliationBuffer.GetNumValidElements() == 1); // this function assumes only 1 element
+
+		const int32 ReconciliationKeyframe = ReconciliationBuffer.GetHeadKeyframe();
+		TSyncState* ServerState = ReconciliationBuffer.GetElementFromHead(0);
+
+		TInputCmd* ClientInputCmd = Buffers.Input.FindElementByKeyframe( ReconciliationKeyframe );
+		if (ClientInputCmd == nullptr)
+		{
+			// Fault: no historic command for this frame to reconcile with.
+			UE_LOG(LogNetworkSim, Error, TEXT("Client InputBuffer does not contain data for frame %d. {%s} {%s}"), ReconciliationKeyframe, *Buffers.Input.GetBasicDebugStr(), *Buffers.Sync.GetBasicDebugStr());
+			return;
+		}
+
+		TSyncState* ClientSyncState = Buffers.Sync.FindElementByKeyframe( ReconciliationKeyframe );
+
+		// Should not be here if the client doesn't have a valid state. See ::NetSerialize
+		checkf(ClientSyncState != nullptr, TEXT("SyncBuffer does not contain data for frame %d. %s"), ReconciliationKeyframe, *Buffers.Sync.GetBasicDebugStr());
+
+		// -------------------------------------------------------------------------------------------------------------------------
+		// Resimulate
+		// -------------------------------------------------------------------------------------------------------------------------
+
+		ServerState->VisualLog( FVisualLoggingParameters(EVisualLoggingContext::LastConfirmed, ReconciliationKeyframe, EVisualLoggingLifetime::Persistent), Driver, Driver);
+		ClientSyncState->VisualLog( FVisualLoggingParameters(EVisualLoggingContext::FirstMispredicted, ReconciliationKeyframe, EVisualLoggingLifetime::Persistent), Driver, Driver);
+
+		TSyncState* LatestClient = Buffers.Sync.GetElementFromHead(0);
+
+		// Copy authoritative state over client state (FIXME: we may want to store this off historically somewhere? Or will VLog be enough for debugging?)
+		*ClientSyncState = *ServerState;
+
+		// Set the canonical simulation time to what we received (we will advance it as we resimulate)
+		TickInfo.ProcessedSimulationTime = SerializedTimeKeeper;
+		TickInfo.LastProcessedInputKeyframe = ReconciliationKeyframe;
+		
+		// Resimulate up to our latest SyncBuffer frame. Note that we may have unprocessed user commands in the command buffer at this point.
+		const int32 LatestKeyframe = Buffers.Sync.GetHeadKeyframe();		
+		for (int32 Keyframe = ReconciliationKeyframe+1; Keyframe <= LatestKeyframe; ++Keyframe)
+		{
+			// Keyframe is the frame we are resimulating right now.
+			TInputCmd* ResimulateCmd	= Buffers.Input.FindElementByKeyframe(Keyframe);
+			TSyncState* PrevMotionState = Buffers.Sync.FindElementByKeyframe(Keyframe - 1);
+			TSyncState* NextMotionState = Buffers.Sync.FindElementByKeyframe(Keyframe);
+			
+			check(ResimulateCmd);
+			check(PrevMotionState);
+			check(NextMotionState);
+
+			// TEMP (Aux buffer not fully plumbed through the system yet)
+			// ------------------------------------------------------
+			TAuxState TempAuxState;
+			TAuxState* AuxState = Buffers.Aux.FindElementByKeyframe(Keyframe);
+			if (AuxState == nullptr)
+			{
+				AuxState = &TempAuxState;
+			}
+			// ------------------------------------------------------
+
+			// Log out the Mispredicted state that we are about to overwrite.
+			NextMotionState->VisualLog( FVisualLoggingParameters(Keyframe == LatestKeyframe ? EVisualLoggingContext::LastMispredicted : EVisualLoggingContext::OtherMispredicted, Keyframe, EVisualLoggingLifetime::Persistent), Driver, Driver);
+
+			// Do the actual update
+			T::Update(Driver, *ResimulateCmd, *PrevMotionState, *NextMotionState, *AuxState);
+			
+			// Update TickInfo
+			TickInfo.ProcessedSimulationTime.AccumulateTimeFromInputCmd(*ResimulateCmd);
+			TickInfo.LastProcessedInputKeyframe = Keyframe;
+
+			// Log out the newly predicted state that we got.
+			NextMotionState->VisualLog( FVisualLoggingParameters(Keyframe == LatestKeyframe ? EVisualLoggingContext::LastPredicted : EVisualLoggingContext::OtherPredicted, Keyframe, EVisualLoggingLifetime::Persistent), Driver, Driver);
+		}
+	}
 
-	void Init(IReplicationProxy* InNetworkSimModel, EReplicationProxyTarget InReplicationTarget);
-	bool NetSerialize(FArchive& Ar, class UPackageMap* Map, bool& bOutSuccess);
-	void OnPreReplication();	
-	bool Identical(const FReplicationProxy* Other, uint32 PortFlags) const;
-	
 private:
+	
+	TReplicationBuffer<TSyncState> ReconciliationBuffer;
+	
+	TSimulationTimeKeeper<TBufferTypes> SerializedTimeKeeper; // last serialized time keeper
 
-	IReplicationProxy* NetworkSimModel;
-	EReplicationProxyTarget ReplicationTarget;
-	int32 CachedDirtyCount = 0;
+	int32 LastSerializedKeyframe = -1;
+	bool bPendingReconciliation = false;	// Reconciliation is pending: we need to reconcile state from the server that differs from the locally predicted state
+	bool bReconcileFaultDetected = false;	// A fault was detected: we received state from the server that we are unable to reconcile with locally predicted state
 };
 
-template<>
-struct TStructOpsTypeTraits<FReplicationProxy> : public TStructOpsTypeTraitsBase2<FReplicationProxy>
+// ----------------------------------------------------------------------------------------------------------------------------------------------------------------------
+//	TNetworkedSimulationModel
+//	
+//	* Has all logic for "ticking, advancing buffers, calling Update, calling ServerRPC etc
+//	* Doesn't have anything about update component, movesweep, etc
+//	* Concept of "IDriver" which is the owning object that is driving the network sim. This is the interface to the outside UE4 world.
+//	* Has 4 buffers:
+//		-Input: Generated by a client / not the authority.
+//		-Sync: What we are trying to keep in sync. The state that evolves frame to frame with an Update function.
+//		-Aux: State that is also an input into the simulation but does not intrinsically evolve from to frame. Changes to this state can be trapped/tracked/predicted.
+//		-Debug: Replicated buffer from server->client with server-frame centered debug information. Compiled out of shipping builds.
+//
+// ----------------------------------------------------------------------------------------------------------------------------------------------------------------------
+
+template <
+	typename T,
+	typename TBufferTypes,
+	typename TRepProxyServerRPC = TReplicator_DynamicSequence<TBufferTypes, ENetworkSimBufferTypeId::Input>,
+	typename TRepProxyAutonomous = TReplicator_BasicReconciliar<TBufferTypes>,
+	typename TRepProxySimulated = TReplicator_SimulatedExtrapolatedReconciliar<TBufferTypes>,
+	typename TRepProxyReplay = TReplicator_DynamicSequence<TBufferTypes, ENetworkSimBufferTypeId::Sync>,
+	typename TRepProxyDebug = TReplicator_DynamicSequence<TBufferTypes, ENetworkSimBufferTypeId::Debug>
+>
+class TNetworkedSimulationModel : public IReplicationProxy
 {
-	enum
+public:
+
+	using TInputCmd = typename TBufferTypes::TInputCmd;
+	using TSyncState = typename TBufferTypes::TSyncState;
+	using TAuxState = typename TBufferTypes::TAuxState;
+	using TDebugState = typename TBufferTypes::TDebugState;
+
+	class IDriver
 	{
-		WithNetSerializer = true,
-		WithIdentical = true,
+	public:
+		virtual void InitSyncState(TSyncState& OutSyncState) const = 0;	// Called to create initial value of the sync state.
+		virtual void FinalizeFrame(const TSyncState& SyncState) = 0; // Called from the Network Sim at the end of the sim frame when there is new sync data.
 	};
-};
 
-// -------------------------------------------------------------------------------------------------------------------------------
-//	FServerRPCProxyParameter
-//	Used for the client->Server RPC. Since this is instantiated on the stack by the replication system prior to net serializing,
-//	we have no opportunity to point the RPC parameter to the member variables we want. So we serialize into a generic temp byte buffer
-//	and move them into the real buffers on the component in the RPC body (via ::NetSerialzeToProxy).
-// -------------------------------------------------------------------------------------------------------------------------------
-USTRUCT()
-struct FServerReplicationRPCParameter
-{
-	GENERATED_BODY()
+	struct FTickParameters
+	{
+		ENetRole Role;
+		float LocalDeltaTimeSeconds;
+	};
 
-	// Receive flow: ctor() -> NetSerializetoProxy
-	FServerReplicationRPCParameter() : Proxy(nullptr)	{ }
-	void NetSerializeToProxy(FReplicationProxy& InProxy);
+	template<typename TDriver>
+	void Tick(TDriver* Driver, const FTickParameters& Parameters)
+	{
+		TDebugState* const DebugState = GetNextDebugStateWrite();
+		if (DebugState)
+		{
+			DebugState->LocalDeltaTimeSeconds = Parameters.LocalDeltaTimeSeconds;
+			DebugState->LocalGFrameNumber = GFrameNumber;
+			DebugState->ProcessedKeyframes.Reset();
+		}
+
+		// How many commands we are allowed to process right now. This will need to be built out a bit more to handle shifting network conditions etc. For now, just a cvar.
+		int32 AllowCmds = NetworkSimulationModelCVars::MaxInputCmdsFrame;
+
+		switch (Parameters.Role) {
+
+		// -------------------------------------------------------------------------------------------------------------------------------------------------
+		//														Reconciliation
+		// -------------------------------------------------------------------------------------------------------------------------------------------------
+		case ROLE_AutonomousProxy:
+		{
+			// Client: don't allow buffering and ensure continuous stream of inputs. This will need to be built out a bit more. The main points are:
+			// -Don't buffer input locally (unless we explicitly opt into it as a lag hiding measurement). Don't want to accidentally introduce application latency.
+			// -Don't send server input cmds that we didn't actually process. Current Input buffer replication is rather hard coded to send continuous frames and doesn't handle gaps.
+
+			// Recovering from this is possible but messy. We could immediately set things up so that we process the latest user cmd.
+			// The tricky part is making sure we the server does the same thing. We would need to set the tail of the input cmd stream so
+			// that we don't replicate any cmds that were skipped. Then server side it would have to detect this and know to jump ahead.
+			// At this point in development, the complexity this adds to the overall system is not worth the benefit. The better way to
+			// fix this for now is to closely guard what we allow into the client's InputCmdBuffer
+
+			// You can be at head (no new cmds to process) or right before the head (1 new cmd to process)
+			checkf(Buffers.Input.GetHeadKeyframe() - TickInfo.LastProcessedInputKeyframe <= 1, TEXT("Client Input Processing is out of sync. LastProcessedInputKeyframe: %d. InputBuffer: [%d - %d]"), 
+				TickInfo.LastProcessedInputKeyframe, Buffers.Input.GetTailKeyframe(), Buffers.Input.GetHeadKeyframe());
+
+			// We should be processing input from this GFrameNumber. If we aren't, then our DeltaTime is probably off (input cmd's delta time wont match the actual delta time of the frame we sim on which will cause artifacts).
+			// We also want to avoid accidentally introducing input latency as a matter of principle (input latency bad)
+			// We *can* support local input buffering in two ways (neither currently implemented):
+			//	-Fixed time steps (delta times will be constant so buffering wont matter)
+			//	-Variable time steps with "interpolation layer": we will need to be able to absorb fluctuations in frame rate so there will be render frames that doesn't process any new input and
+			//		render frames that have to process multiple inputs. (E.g, the goal would be to always "see" an update every render frame, even if we can't actually tick a new sim frame every render frame. Interpolation would be needed to accomplish this)
+			// If either of the above ideas are implemented, it will be done as an optional/templated parameter to the network sim class.
+			checkf(TickInfo.LastLocalInputGFrameNumber == 0 || GFrameNumber == TickInfo.LastLocalInputGFrameNumber, TEXT("TNetworkedSimulationModel running on stale input. You may have an ordering issue between TNetworkedSimulationModel::GetNextInputForWrite and TNetworkedSimulationModel::Tick. "
+				"GFrameNumber: %d LastLocalInputGFrameNumber: %d"), GFrameNumber, TickInfo.LastLocalInputGFrameNumber);
+
+		
+			// Note that is is important for the NetworkSimModel to control *when* the reconcile happens. The replicator objects themselves should never do
+			// reconciliation inside a NetSerialze call. Tick may not be the permanent place for this, but it can't be in NetSerialize.
+			RepProxy_Autonomous.template Reconcile<T, TDriver>(Driver, Buffers, TickInfo);
+
+			// Don't process new user commands if we are dealing with a reconciliation fault. That means we weren't able to correlate our predicted position with what the
+			// server last sent us. If we are hitting this, it means we are waiting for the in flight commands to flush through the system. The proxy rep object will handle recovering.
+			if (RepProxy_Autonomous.IsReconcileFaultDetected())
+			{
+				AllowCmds = 0;
+			}
+
+			// Check SyncBuffer being ahead of processed Keyframes. This would happen in cases where we are either not predicting or are buffering our input locally
+			// while sending latest cmds to the server. Essentially, we got the authoritative motion state from the server before we ran the simulation locally.
+			if (Buffers.Sync.GetHeadKeyframe() > TickInfo.LastProcessedInputKeyframe)
+			{
+				checkf(Buffers.Input.IsValidKeyframe(Buffers.Sync.GetHeadKeyframe()), TEXT("MotionState and InputCmd buffers are out of system. LastProcessedInputKeyframe: %d {%s} vs {%s}"),
+					TickInfo.LastProcessedInputKeyframe, *Buffers.Sync.GetBasicDebugStr(), *Buffers.Input.GetBasicDebugStr());
+
+				UE_LOG(LogNetworkSim, Warning, TEXT("Skipping local input frames because we have newer data in SyncBuffer. LastProcessedInputKeyframe: %d. {%s} {%s}"),
+					TickInfo.LastProcessedInputKeyframe, *Buffers.Sync.GetBasicDebugStr(), *Buffers.Input.GetBasicDebugStr());
+
+				TickInfo.LastProcessedInputKeyframe = Buffers.Sync.GetHeadKeyframe();
+			}
+
+			if (NetworkSimulationModelCVars::EnableLocalPrediction == 0)
+			{
+				AllowCmds = 0; // Don't process any commands this frame
+				TickInfo.LastProcessedInputKeyframe = Buffers.Input.GetHeadKeyframe(); // Increment so we can accept a new command next frame
+			}
+
+			break;
+		}
+
+		case ROLE_Authority:
+		{
+			if ( TickInfo.LastProcessedInputKeyframe+1 < Buffers.Input.GetTailKeyframe() )
+			{
+				// We've missed commands
+				UE_LOG(LogNetworkSim, Warning, TEXT("::Tick missing inputcmds. LastProcessedInputKeyframe: %d. %s"), TickInfo.LastProcessedInputKeyframe, *Buffers.Input.GetBasicDebugStr());
+				TickInfo.LastProcessedInputKeyframe = Buffers.Input.GetTailKeyframe()+1;
+			}
+			break;
+		}
+
+		case ROLE_SimulatedProxy:
+		{
+			RepProxy_Simulated.template Reconcile<T, TDriver>(Driver, Buffers, TickInfo);
+
+			// Extrapolation: this still needs work  (needs to be a setting/state somewhere and we need to better handle the relationship with NetSerialize and Reconcile)
+			if (TInputCmd* SynthesizedCmd = GetNextInputForWrite(Parameters.LocalDeltaTimeSeconds))
+			{
+
+			}
+			break;
+		}
+
+		} // end switch
+
+		// -------------------------------------------------------------------------------------------------------------------------------------------------
+		//														Input Processing
+		// -------------------------------------------------------------------------------------------------------------------------------------------------
+
+		int32 NumProcessed = 0;
+		while (AllowCmds-- > 0)
+		{
+			const int32 Keyframe = TickInfo.LastProcessedInputKeyframe+1;
+			if (TInputCmd* NextCmd = Buffers.Input.FindElementByKeyframe(Keyframe))
+			{
+				// The SyncedState buffer needs to be in sync here:
+				//	-We want it to have a SyncedState, but it may not on the first frame through (thats ok).
+				//  -Its HeadKeyframe should be one behind the Keyframe we are about to process.
+				//
+				// Note, InputCmds start @ Keyframe=1. The first SyncedState that Update produces will go in KeyFrame=1.
+				// (E.g, InputCmd @ keyframe=X is used to generate MotionState @ keyframe=X)
+				// This means that SyncedState @ keyframe=0 is always created here via InitSyncState.
+				// This also means that we never actually process InputCmd @ keyframe=0. Which is why LastProcessedInputKeyframe is initialized to 0 ("already processed")
+				// and the buffer has an empty element inserted in InitLocalInputBuffer.
+
+				if (Buffers.Sync.GetHeadKeyframe() != TickInfo.LastProcessedInputKeyframe)
+				{
+					if (TickInfo.LastProcessedInputKeyframe != 0)
+					{
+						// This shouldn't happen, but is not fatal. We are reseting the sync state buffer.
+						UE_LOG(LogNetworkSim, Warning, TEXT("Break in SyncState continuity. LastProcessedInputKeyframe: %d. SyncBuffer.GetHeadKeyframe(): %d. Role=%d"), TickInfo.LastProcessedInputKeyframe, Buffers.Sync.GetHeadKeyframe(), (int32)Parameters.Role);
+					}
+
+					// We need an initial/current state. Get this from the sim driver
+					Buffers.Sync.ResetNextHeadKeyframe(TickInfo.LastProcessedInputKeyframe);
+					TSyncState* StartingState = Buffers.Sync.GetWriteNext();
+					Driver->InitSyncState(*StartingState);
+					
+				}
+
+				TSyncState* PrevSyncState = Buffers.Sync.FindElementByKeyframe(TickInfo.LastProcessedInputKeyframe);
+				TSyncState* NextSyncState = Buffers.Sync.GetWriteNext();
+
+				check(PrevSyncState != nullptr);
+				check(NextSyncState != nullptr);
+				check(Buffers.Sync.GetHeadKeyframe() == Keyframe);
+				
+				if (DebugState)
+				{
+					DebugState->ProcessedKeyframes.Add(Keyframe);
+				}
+
+				TAuxState AuxState; // Temp: aux buffer not implemented yet
+
+				T::Update(Driver, *NextCmd, *PrevSyncState, *NextSyncState, AuxState);
+				TickInfo.ProcessedSimulationTime.AccumulateTimeFromInputCmd(*NextCmd);
+				NumProcessed++;
+
+				TickInfo.LastProcessedInputKeyframe = Keyframe;
+			}
+			else
+			{
+				// Not a warning in itself anymore. Though the concept of "haven't processed cmds in x MS" is something we want to define.
+				break;
+			}
+		}
+
+		// FIXME: this needs to be sorted out. We really want to check if there is new sync state and then call this here.
+		// Call into the driver to sync to the latest state if we processed any input
+		//if (NumProcessed > 0)
+		{
+			//check(SyncBuffer.GetNumValidElements() > 0);
+			if (Buffers.Sync.GetNumValidElements() > 0)
+			{
+				Driver->FinalizeFrame(*Buffers.Sync.GetElementFromHead(0));
+			}
+		}
+
+		// -------------------------------------------------------------------------------------------------------------------------------------------------
+		//														Debug
+		// -------------------------------------------------------------------------------------------------------------------------------------------------
+
+		// Finish debug state buffer recording (what the server processed each frame)
+		if (DebugState)
+		{
+			DebugState->LastProcessedKeyframe = TickInfo.LastProcessedInputKeyframe;
+			DebugState->HeadKeyframe = Buffers.Input.GetHeadKeyframe();
+		}
+
+		// Historical data recording (longer buffers for historical reference)
+		if (auto* HistoricData = GetHistoricBuffers())
+		{
+			HistoricData->Input.CopyAndMerge(Buffers.Input);
+			HistoricData->Sync.CopyAndMerge(Buffers.Sync);
+			HistoricData->Aux.CopyAndMerge(Buffers.Aux);
+		}
+	}	
+	
+	void InitializeForNetworkRole(const ENetRole Role, const bool IsLocallyControlled, const FNetworkSimulationModelInitParameters& Parameters)
+	{
+		Buffers.Input.SetBufferSize(Parameters.InputBufferSize);
+		Buffers.Sync.SetBufferSize(Parameters.SyncedBufferSize);
+		Buffers.Aux.SetBufferSize(Parameters.AuxBufferSize);
+
+		if (GetDebugBuffer())
+		{
+			GetDebugBuffer()->SetBufferSize(Parameters.DebugBufferSize);
+		}
+
+		if (auto* MyHistoricBuffers = GetHistoricBuffers(true))
+		{
+			MyHistoricBuffers->Input.SetBufferSize(Parameters.HistoricBufferSize);
+			MyHistoricBuffers->Sync.SetBufferSize(Parameters.HistoricBufferSize);
+			MyHistoricBuffers->Aux.SetBufferSize(Parameters.HistoricBufferSize);
+		}
+
+		if (IsLocallyControlled)
+		{
+			check(Parameters.InputBufferSize > 0); // If you tell me this is locally controlled, you need to have an input buffer.
+			InitLocalInputBuffer();
+		}
+	}
 
-	// Send flow: ctor(Proxy) -> NetSerialize
-	FServerReplicationRPCParameter(FReplicationProxy& InProxy) : Proxy(&InProxy) { }
-	bool NetSerialize(FArchive& Ar, class UPackageMap* Map, bool& bOutSuccess);
+	void NetSerializeProxy(EReplicationProxyTarget Target, const FNetSerializeParams& Params)
+	{
+		// You are not allowed to change the simulations TickInfo while NetSerializing. Since all rep proxies are templated
+		// and there is not interface/base class, we are enforcing the constness at the call site.
+		const TSimulationTickInfo<TBufferTypes>& ConstTickInfo = const_cast<const TSimulationTickInfo<TBufferTypes>&>(TickInfo);
+
+		switch(Target)
+		{
+		case EReplicationProxyTarget::ServerRPC:
+			RepProxy_ServerRPC.NetSerialize(Params, Buffers, ConstTickInfo);
+			break;
+		case EReplicationProxyTarget::AutonomousProxy:
+			RepProxy_Autonomous.NetSerialize(Params, Buffers, ConstTickInfo);
+			break;
+		case EReplicationProxyTarget::SimulatedProxy:
+			RepProxy_Simulated.NetSerialize(Params, Buffers, ConstTickInfo);
+			break;
+		case EReplicationProxyTarget::Replay:
+			RepProxy_Replay.NetSerialize(Params, Buffers, ConstTickInfo);
+			break;
+		case EReplicationProxyTarget::Debug:
+#if NETSIM_MODEL_DEBUG
+			RepProxy_Debug.NetSerialize(Params, Buffers, ConstTickInfo);
+			break;
+#endif
+		default:
+			checkf(false, TEXT("Unknown: %d"), (int32)Target);
+		};
+	}
 
-private:
+	int32 GetProxyDirtyCount(EReplicationProxyTarget Target)
+	{
+		switch(Target)
+		{
+		case EReplicationProxyTarget::ServerRPC:
+			return RepProxy_ServerRPC.GetProxyDirtyCount(Buffers);
+		case EReplicationProxyTarget::AutonomousProxy:
+			return RepProxy_Autonomous.GetProxyDirtyCount(Buffers);
+		case EReplicationProxyTarget::SimulatedProxy:
+			return RepProxy_Simulated.GetProxyDirtyCount(Buffers);
+		case EReplicationProxyTarget::Replay:
+			return RepProxy_Replay.GetProxyDirtyCount(Buffers);
+		case EReplicationProxyTarget::Debug:
+#if NETSIM_MODEL_DEBUG
+			return RepProxy_Debug.GetProxyDirtyCount(Buffers);
+#endif
+		default:
+			checkf(false, TEXT("Unknown: %d"), (int32)Target);
+			return 0;
+		};
+	}
+	
+	TInputCmd* GetNextInputForWrite(float DeltaTime)
+	{
+		if (TInputCmd* Next = TickInfo.GetNextInputForWrite(Buffers))
+		{
+			Next->FrameDeltaTime = DeltaTime;
+			return Next;
+		}
+		return nullptr;
+	}
+
+	void InitLocalInputBuffer()
+	{
+		// Buffer should also be empty before calling this
+		check(Buffers.Input.GetHeadKeyframe() == INDEX_NONE);
 
-	static TArray<uint8> TempStorage;
+		// We want to start with an empty command in the input buffer. See notes in input buffer processing function.
+		*Buffers.Input.GetWriteNext() = TInputCmd();
+		TickInfo.LastLocalInputGFrameNumber = 0;
+	}	
+	
+	TSimulationTickInfo<TBufferTypes> TickInfo;	// Manages simulation time and what inputs we are processed
 
-	FReplicationProxy* Proxy;
-	int64 CachedNumBits = -1;
-	class UPackageMap* CachedPackageMap = nullptr;
-};
+	TNetworkSimBufferContainer<TBufferTypes> Buffers;
 
-template<>
-struct TStructOpsTypeTraits<FServerReplicationRPCParameter> : public TStructOpsTypeTraitsBase2<FServerReplicationRPCParameter>
-{
-	enum
+	TRepProxyServerRPC RepProxy_ServerRPC;
+	TRepProxyAutonomous RepProxy_Autonomous;
+	TRepProxySimulated RepProxy_Simulated;
+	TRepProxyReplay RepProxy_Replay;
+
+	// ------------------------------------------------------------------
+	// RPC Sending helper: provides basic send frequency settings for tracking when the Server RPC can be invoked.
+	// Note that the Driver is the one that must call the RPC, that cannot be rolled into this templated structure.
+	// More flexbile/dynamic send rates may be desireable. There is not reason this *has* to be done here, it could
+	// completely be tracked at the driver level, but that will also push more boilerplate to that layer for users.
+	// ------------------------------------------------------------------
+
+	void SetDesiredServerRPCSendFrequency(float DesiredHz) { ServerRPCThresholdTimeSeconds = 1.f / DesiredHz; }
+	bool ShouldSendServerRPC(ENetRole OwnerRole, float DeltaTimeSeconds)
 	{
-		WithNetSerializer = true
-	};
-};
+		// Don't allow a large delta time to pollute the accumulator
+		const float CappedDeltaTimeSeconds = FMath::Min<float>(DeltaTimeSeconds, ServerRPCThresholdTimeSeconds);
+		if (OwnerRole == ROLE_AutonomousProxy)
+		{
+			ServerRPCAccumulatedTimeSeconds += DeltaTimeSeconds;
+			if (ServerRPCAccumulatedTimeSeconds >= ServerRPCThresholdTimeSeconds)
+			{
+				ServerRPCAccumulatedTimeSeconds -= ServerRPCThresholdTimeSeconds;
+				return true;
+			}
+		}
+		return false;
+	}
+private:
+	float ServerRPCAccumulatedTimeSeconds = 0.f;
+	float ServerRPCThresholdTimeSeconds = 1.f / 999.f; // Default is to send at a max of 999hz. This part of the system needs to be build out more (better handling of super high FPS clients and fixed rate servers)
+
+	// ------------------------------------------------------------------
+	//	Debugging
+	// ------------------------------------------------------------------
+public:
+
+#if NETSIM_MODEL_DEBUG
+	TReplicationBuffer<TDebugState>* GetDebugBuffer() {	return &Buffers.Debug; }
+	TDebugState* GetNextDebugStateWrite() { return Buffers.Debug.GetWriteNext(); }
+	TNetworkSimBufferContainer<TBufferTypes>* GetHistoricBuffers(bool bCreate=false)
+	{
+		if (HistoricBuffers.IsValid() == false && bCreate) { HistoricBuffers.Reset(new TNetworkSimBufferContainer<TBufferTypes>()); }
+		return HistoricBuffers.Get();
+	}
+#else
+	TReplicationBuffer<TDebugState>* GetDebugBuffer(bool bCreate=false) {	return nullptr; }
+	TDebugState* GetNextDebugStateWrite() { return nullptr; }
+	TNetworkSimBufferContainer<TBufferTypes>* GetHistoricBuffers() { return nullptr; }
+#endif
 
-// Helper struct to bypass the bandwidth limit imposed by the engine's NetDriver (QueuedBits, NetSpeed, etc).
-// This is really a temp measure to make the system easier to drop in/try in a project without messing with your engine settings.
-// (bandwidth optimizations have not been done yet and the system in general hasn't been stressed with packetloss / gaps in command streams)
-// So, you are free to use this in your own code but it may be removed one day. Hopefully in general bandwidth limiting will also become more robust.
-struct FScopedBandwidthLimitBypass
-{
-	FScopedBandwidthLimitBypass(AActor* OwnerActor);
-	~FScopedBandwidthLimitBypass();
 private:
 
-	int64 RestoreBits = 0;
-	class UNetConnection* CachedNetConnection = nullptr;
-};
\ No newline at end of file
+#if NETSIM_MODEL_DEBUG
+	TRepProxyDebug RepProxy_Debug;
+	TUniquePtr<TNetworkSimBufferContainer<TBufferTypes>> HistoricBuffers;
+#endif
+};
diff --git a/Engine/Plugins/Runtime/NetworkPrediction/Source/NetworkPrediction/Public/NetworkSimulationModelBuffer.h b/Engine/Plugins/Runtime/NetworkPrediction/Source/NetworkPrediction/Public/NetworkSimulationModelBuffer.h
new file mode 100644
--- /dev/null
+++ b/Engine/Plugins/Runtime/NetworkPrediction/Source/NetworkPrediction/Public/NetworkSimulationModelBuffer.h
@@ -0,0 +1,209 @@
+// Copyright 1998-2019 Epic Game s, Inc. All Rights Reserved.
+
+#pragma once
+
+// ----------------------------------------------------------------------------------------------------------------------------------------------------------------------
+//	TReplicationBuffer
+//
+//	Generic Cyclic buffer. Has canonical head position. This is the "Client Frame" or "Keyframe" identifier used by the rest of the system.
+//	Contract: Elements in buffer are contiguously valid. We don't allow "gaps" in the buffer!
+//	Use ::GetWriteNext<T> to write to the buffer.
+//	Use ::CreateIterator to iterate from Tail->Head
+//
+//	for (auto It = Buffer.CreateIterator(); It; ++It) {
+//		FMyType* Element = It.Element();
+//		int32 Keyframe = It.Keyframe();
+//	}
+//
+// ----------------------------------------------------------------------------------------------------------------------------------------------------------------------
+template<typename T>
+struct TReplicationBuffer
+{
+	// Set the buffer size. Resizing is supported but not optimized because we preserve head/tail iteration (not a simple memcopy since where buffer wraps around will change)
+	// In general, avoid resizing outside of startup/initialization.
+	void SetBufferSize(int32 NewMaxNumElements)
+	{
+		check(NewMaxNumElements >= 0);	
+		if (Data.Num() == NewMaxNumElements)
+		{
+			return;
+		}
+		
+		if (Data.Num() == 0)
+		{
+			Data.SetNum(NewMaxNumElements, false);
+		}
+		else
+		{
+			// Grow or shrink. This is far from optimal but this operation should be rare
+			TReplicationBuffer<T> Old(MoveTemp(*this)); // Move old data to a copy
+			Data.SetNum(NewMaxNumElements, false);
+			ResetNextHeadKeyframe(Old.GetTailKeyframe(), true); // Reset so our next write is at the old copies tail keyframe
+
+			for (auto It = Old.CreateIterator(); It; ++It)
+			{
+				T* WriteNext = GetWriteNext();
+				T* ReadNext = It.Element();
+				FMemory::Memcpy(WriteNext, ReadNext, sizeof(T));
+			}
+
+			// Blow away the temp copy so we don't run the dtor on it on scope end (because we memcpy'd it over, not copy assigned)
+			Old.Data.Reset();
+			Old.NumValidElements = 0;
+		}
+	}
+
+	// Gets the oldest element in the buffer (or null)
+	T* GetElementFromTail(int32 OffsetFromTail)
+	{
+		check(OffsetFromTail >= 0); // Must always pass in valid offset
+
+		// Out of range. This is not fatal, just don't return the wrong element.
+		if (NumValidElements <= OffsetFromTail)
+		{
+			return nullptr;
+		}
+		
+		const int32 Tail = GetTailKeyframe();
+		const int32 Position = Tail + OffsetFromTail;
+		const int32 Offset = (Position % Data.Num());
+		return &Data[Offset];
+	}
+
+	// Gets the newest element in the buffer (or null)
+	T* GetElementFromHead(int32 OffsetFromHead)
+	{
+		check(OffsetFromHead >= 0); // Must always pass in valid offset
+
+		// Out of range. This is not fatal, just don't return the wrong element.
+		if (NumValidElements <= OffsetFromHead)
+		{
+			return nullptr;
+		}
+		
+		const int32 Position = Head - OffsetFromHead;
+		check(Position >= 0);
+		const int32 Offset = (Position % Data.Num());
+		return &Data[Offset];
+	}
+
+	const T* FindElementByKeyframe(int32 Key) const
+	{
+		return FindElementByKeyframeImpl(Key);
+	}
+
+	T* FindElementByKeyframe(int32 Key)
+	{
+		const T* Element = const_cast<TReplicationBuffer<T>*>(this)->FindElementByKeyframeImpl(Key);
+		return const_cast<T*>(Element);
+	}
+
+	// Returns the next element for writing. The contents of this element are unknown (could be stale content!). This also advances the head. Note the element returned is immediately considered "valid"!
+	T* GetWriteNext()
+	{
+		check(Data.Num() > 0); // Buffer must be initialized first
+		
+		++Head;
+		++DirtyCount;
+		NumValidElements = FMath::Min<int32>(NumValidElements+1, Data.Num());
+
+		const int32 Offset = (Head % Data.Num());
+		return &Data[Offset];
+	}
+
+	// Resets the head keyframen, so that NextHeadKeyframe will be the keyframe of the next GetWriteNext<>(). E.g, call this before you want to write to NextHeadKeyframe.
+	// To remphasize: NextHeadKeyframe is where the NEXT write will go. NOT "where the current head is".
+	// Contents will be preserved if possible. E.g, if NextHeadKeyframe-1 is a valid keyframe, it will remain valid. If its not (if this creates "gaps" in the buffer) then
+	// the contents will effectively be cleared. bForceClearContents will force this behavior even if the previous keyframe was valid.
+	void ResetNextHeadKeyframe(int32 NextHeadKeyframe=0, bool bForceClearContents=false)
+	{
+		const int32 NewHeadKeyframe = NextHeadKeyframe - 1;
+		if (bForceClearContents || NewHeadKeyframe < GetTailKeyframe() || NewHeadKeyframe > GetHeadKeyframe())
+		{
+			NumValidElements = 0;
+		}
+		else
+		{
+			NumValidElements += (NewHeadKeyframe - GetHeadKeyframe());
+			check(NumValidElements >= 0 && NumValidElements <= Data.Num());
+		}
+
+		Head =  NewHeadKeyframe;
+		DirtyCount++;
+	}
+
+	// Copies the contents of the SourceBuffer into this buffer. This existing buffer's data will be preserved as much as possible, but the guarantee is that the SourceBuffer's
+	// data will all be added to this buffer.
+	
+	// Example 1: Target={1...5}, Source={3...9}  --> Target={1...9}
+	// Example 2: Target={1...5}, Source={7...9}  --> Target={7...9}
+	// Exmaple 3: Target={6...9}, Source={1...4}  --> Target={1...4}
+	void CopyAndMerge(const TReplicationBuffer<T>& SourceBuffer)
+	{
+		const int32 StartingHeadKeyframe = GetHeadKeyframe();
+
+		ResetNextHeadKeyframe(SourceBuffer.GetTailKeyframe());
+		for (int32 Keyframe = SourceBuffer.GetTailKeyframe(); Keyframe <= SourceBuffer.GetHeadKeyframe(); ++Keyframe)
+		{
+			T* TargetData = GetWriteNext();
+			check(GetHeadKeyframe() == Keyframe);
+
+			const T* SourceData = SourceBuffer.FindElementByKeyframe(Keyframe);
+			check(SourceData != nullptr);
+
+			*TargetData = *SourceData;
+		}
+	}
+
+	FString GetBasicDebugStr() const
+	{
+		return FString::Printf(TEXT("Elements: [%d/%d]. Keyframes: [%d-%d]"), NumValidElements, Data.Num(), GetTailKeyframe(), GetHeadKeyframe());
+	}
+
+	int32 GetNumValidElements() const { return NumValidElements; }
+	int32 GetMaxNumElements() const { return Data.Num(); }
+	int32 GetHeadKeyframe() const { return Head; }
+	int32 GetTailKeyframe() const { return Head - NumValidElements + 1; }
+	bool IsValidKeyframe(int32 Keyframe) const { return Keyframe >= GetTailKeyframe() && Keyframe <= GetHeadKeyframe(); }
+	int32 GetDirtyCount() const { return DirtyCount; }
+
+	// Create an iterator from tail->head. Note that no matter what template class is, this WILL iterate correctly across the element. The templated type
+	// is just for casting the return element type. E.g, it is fine to use <uint8> in generic code, this will still step by 'StructSize'.
+	struct TGenericIterator
+	{
+		TGenericIterator(TReplicationBuffer<T>& InBuffer) : Buffer(InBuffer)
+		{
+			CurrentKeyframe = Buffer.GetTailKeyframe();
+		}
+
+		int32 Keyframe() const { return CurrentKeyframe; }
+		T* Element() const { return Buffer.FindElementByKeyframe(CurrentKeyframe); }
+		void operator++() { CurrentKeyframe++; }
+		operator bool()	{ return CurrentKeyframe <= Buffer.GetHeadKeyframe(); }
+
+	private:
+		TReplicationBuffer<T>& Buffer;
+		int32 CurrentKeyframe = -1;
+	};
+		
+	TGenericIterator CreateIterator() { return TGenericIterator(*this); }
+
+private:
+
+	const T* FindElementByKeyframeImpl(int32 Key) const
+	{
+		const int32 RelativeToHead = Key - Head;
+		if (RelativeToHead > 0 || RelativeToHead <= -NumValidElements)
+		{
+			return nullptr;
+		}
+
+		return &Data[Key % Data.Num()];
+	}
+
+	int32 Head = INDEX_NONE;
+	int32 DirtyCount = 0;
+	int32 NumValidElements = 0;
+
+	TArray<T> Data;
+};
\ No newline at end of file
diff --git a/Engine/Plugins/Runtime/NetworkPrediction/Source/NetworkPrediction/Public/NetworkSimulationModelDebugger.h b/Engine/Plugins/Runtime/NetworkPrediction/Source/NetworkPrediction/Public/NetworkSimulationModelDebugger.h
--- a/Engine/Plugins/Runtime/NetworkPrediction/Source/NetworkPrediction/Public/NetworkSimulationModelDebugger.h
+++ b/Engine/Plugins/Runtime/NetworkPrediction/Source/NetworkPrediction/Public/NetworkSimulationModelDebugger.h
@@ -10,11 +10,11 @@
 #include "Engine/Engine.h"
 #include "UObject/WeakObjectPtrTemplates.h"
 #include "Tickable.h"
-#include "NetworkSimulationModel.h"
+#include "NetworkPredictionTypes.h"
 #include "CanvasItem.h"
 #include "Containers/Array.h"
 #include "Templates/UniquePtr.h"
-#include "NetworkSimulationModelTemplates.h"
+#include "NetworkSimulationModel.h"
 
 DEFINE_LOG_CATEGORY_STATIC(LogNetworkSimDebug, Log, All);
 
@@ -355,10 +355,10 @@ struct TNetworkSimulationModelDebugger : public INetworkSimulationModelDebugger
 		// ------------------------------------------------------------------------------------------------------------------------------------------------
 
 		Out.Emit(FString::Printf(TEXT("%s - %s"), *Owner->GetName(), *UEnum::GetValueAsString(TEXT("Engine.ENetRole"), Owner->Role)), FColor::Yellow);
-		Out.Emit(FString::Printf(TEXT("LastProcessedInputKeyframe: %d (%d Buffered)"), NetworkSim->LastProcessedInputKeyframe, NetworkSim->InputBuffer.GetHeadKeyframe() - NetworkSim->LastProcessedInputKeyframe));
+		Out.Emit(FString::Printf(TEXT("LastProcessedInputKeyframe: %d (%d Buffered)"), NetworkSim->TickInfo.LastProcessedInputKeyframe, NetworkSim->Buffers.Input.GetHeadKeyframe() - NetworkSim->TickInfo.LastProcessedInputKeyframe));
 
 		// Autorproxy
-		{
+		{			
 			FColor Color = FColor::White;
 			const bool FaultDetected = NetworkSim->RepProxy_Autonomous.IsReconcileFaultDetected();
 
@@ -366,15 +366,15 @@ struct TNetworkSimulationModelDebugger : public INetworkSimulationModelDebugger
 
 			// Calc how much predicted time we have processed. Note that we use the motionstate buffer to iterate but the MS is on the input cmd. (if we are buffering cmds, don't want to count them)
 			float PredictedMS = 0.f;
-			for (int32 PredKeyrame = LastSerializedKeyframe+1; PredKeyrame <= NetworkSim->SyncBuffer.GetHeadKeyframe(); ++PredKeyrame)
+			for (int32 PredKeyrame = LastSerializedKeyframe+1; PredKeyrame <= NetworkSim->Buffers.Sync.GetHeadKeyframe(); ++PredKeyrame)
 			{
-				if (auto* Cmd = NetworkSim->InputBuffer.FindElementByKeyframe(PredKeyrame))
+				if (auto* Cmd = NetworkSim->Buffers.Input.FindElementByKeyframe(PredKeyrame))
 				{
 					PredictedMS += Cmd->FrameDeltaTime * 1000.f;
 				}
 			}
 
-			FString ConfirmedFrameStr = FString::Printf(TEXT("LastConfirmedFrame: %d. Prediction: %d Frames, %.2f MS"), LastSerializedKeyframe, NetworkSim->SyncBuffer.GetHeadKeyframe() - LastSerializedKeyframe, PredictedMS);
+			FString ConfirmedFrameStr = FString::Printf(TEXT("LastConfirmedFrame: %d. Prediction: %d Frames, %.2f MS"), LastSerializedKeyframe, NetworkSim->Buffers.Sync.GetHeadKeyframe() - LastSerializedKeyframe, PredictedMS);
 			if (FaultDetected)
 			{
 				ConfirmedFrameStr += TEXT(" RECONCILE FAULT DETECTED!");
@@ -382,6 +382,11 @@ struct TNetworkSimulationModelDebugger : public INetworkSimulationModelDebugger
 			}
 
 			Out.Emit(*ConfirmedFrameStr, Color);
+
+			FString SimulationTimeString = FString::Printf(TEXT("Local SimulationTime: %s. SerialisedSimulationTime: %s. Difference MS: %s"), *NetworkSim->TickInfo.ProcessedSimulationTime.ToString(),
+				*NetworkSim->RepProxy_Autonomous.GetLastSerializedSimulationTimeKeeper().ToString(), *(NetworkSim->TickInfo.ProcessedSimulationTime - NetworkSim->RepProxy_Autonomous.GetLastSerializedSimulationTimeKeeper()).ToString());
+			Out.Emit(*SimulationTimeString, Color);
+			
 		}
 
 		auto EmitBuffer = [&Out](FString BufferName, auto& Buffer)
@@ -393,8 +398,8 @@ struct TNetworkSimulationModelDebugger : public INetworkSimulationModelDebugger
 			Out.EmitElement(Buffer, FStandardLoggingParameters(nullptr, EStandardLoggingContext::Full, Buffer.GetHeadKeyframe()));
 		};
 
-		EmitBuffer(TEXT("InputBuffer"), NetworkSim->InputBuffer);
-		EmitBuffer(TEXT("SyncBuffer"), NetworkSim->SyncBuffer);
+		EmitBuffer(TEXT("InputBuffer"), NetworkSim->Buffers.Input);
+		EmitBuffer(TEXT("SyncBuffer"), NetworkSim->Buffers.Sync);
 
 		// ------------------------------------------------------------------------------------------------------------------------------------------------
 		//	Canvas
@@ -403,7 +408,7 @@ struct TNetworkSimulationModelDebugger : public INetworkSimulationModelDebugger
 		auto* DebugBuffer = NetworkSim->GetDebugBuffer();
 		if (Canvas && DebugBuffer && DebugBuffer->GetNumValidElements() > 0)
 		{
-			auto& InputBuffer = NetworkSim->GetHistoricBuffers() ? NetworkSim->GetHistoricBuffers()->InputBuffer : NetworkSim->InputBuffer;
+			auto& InputBuffer = NetworkSim->GetHistoricBuffers() ? NetworkSim->GetHistoricBuffers()->Input : NetworkSim->Buffers.Input;
 
 			static float StartPctX = 0.3f;
 			static float StartPctY = 0.6f;
@@ -497,26 +502,26 @@ struct TNetworkSimulationModelDebugger : public INetworkSimulationModelDebugger
 
 		UWorld* World = Owner->GetWorld();
 
-		if (auto* LatestSync = NetworkSim->SyncBuffer.GetElementFromHead(0))
+		if (auto* LatestSync = NetworkSim->Buffers.Sync.GetElementFromHead(0))
 		{
-			LatestSync->VisualLog( FVisualLoggingParameters(EVisualLoggingContext::LastPredicted, NetworkSim->SyncBuffer.GetHeadKeyframe(), EVisualLoggingLifetime::Transient), Driver, Driver);
+			LatestSync->VisualLog( FVisualLoggingParameters(EVisualLoggingContext::LastPredicted, NetworkSim->Buffers.Sync.GetHeadKeyframe(), EVisualLoggingLifetime::Transient), Driver, Driver);
 		}
 
 		FStuff ServerPIEStuff = GetServerPIEStuff();
 		if (ServerPIEStuff.Driver && ServerPIEStuff.NetworkSim)
 		{
-			if (auto* ServerLatestSync = ServerPIEStuff.NetworkSim->SyncBuffer.GetElementFromHead(0))
+			if (auto* ServerLatestSync = ServerPIEStuff.NetworkSim->Buffers.Sync.GetElementFromHead(0))
 			{
-				ServerLatestSync->VisualLog( FVisualLoggingParameters(EVisualLoggingContext::CurrentServerPIE, ServerPIEStuff.NetworkSim->SyncBuffer.GetHeadKeyframe(), EVisualLoggingLifetime::Transient), ServerPIEStuff.Driver, Driver);
+				ServerLatestSync->VisualLog( FVisualLoggingParameters(EVisualLoggingContext::CurrentServerPIE, ServerPIEStuff.NetworkSim->Buffers.Sync.GetHeadKeyframe(), EVisualLoggingLifetime::Transient), ServerPIEStuff.Driver, Driver);
 			}
 		}
-		
-		for (int32 Keyframe = NetworkSim->RepProxy_Autonomous.GetLastSerializedKeyframe(); Keyframe < NetworkSim->SyncBuffer.GetHeadKeyframe(); ++Keyframe)
+
+		for (int32 Keyframe = NetworkSim->RepProxy_Autonomous.GetLastSerializedKeyframe(); Keyframe < NetworkSim->Buffers.Sync.GetHeadKeyframe(); ++Keyframe)
 		{
-			if (auto* SyncState = NetworkSim->SyncBuffer.FindElementByKeyframe(Keyframe))
+			if (auto* SyncState = NetworkSim->Buffers.Sync.FindElementByKeyframe(Keyframe))
 			{
 				const EVisualLoggingContext Context = (Keyframe == NetworkSim->RepProxy_Autonomous.GetLastSerializedKeyframe()) ? EVisualLoggingContext::LastConfirmed : EVisualLoggingContext::OtherPredicted;
-				SyncState->VisualLog( FVisualLoggingParameters(Context, NetworkSim->SyncBuffer.GetHeadKeyframe(), EVisualLoggingLifetime::Transient), Driver, Driver);
+				SyncState->VisualLog( FVisualLoggingParameters(Context, NetworkSim->Buffers.Sync.GetHeadKeyframe(), EVisualLoggingLifetime::Transient), Driver, Driver);
 			}
 		}
 	}
diff --git a/Engine/Plugins/Runtime/NetworkPrediction/Source/NetworkPrediction/Public/NetworkSimulationModelTypes.h b/Engine/Plugins/Runtime/NetworkPrediction/Source/NetworkPrediction/Public/NetworkSimulationModelTypes.h
new file mode 100644
--- /dev/null
+++ b/Engine/Plugins/Runtime/NetworkPrediction/Source/NetworkPrediction/Public/NetworkSimulationModelTypes.h
@@ -0,0 +1,240 @@
+// Copyright 1998-2019 Epic Game s, Inc. All Rights Reserved.
+
+#pragma once
+
+#include "NetworkSimulationModelBuffer.h"
+
+// ---------------------------------------------------------------------------------------------------------------------
+//	BufferTypes: template helpers for addressing the different buffer types of the system.
+// ---------------------------------------------------------------------------------------------------------------------
+
+// Enum to refer to buffer type. These are used as template arguments to write generic code that can act on any of the buffers.
+enum ENetworkSimBufferTypeId
+{
+	Input,
+	Sync,
+	Aux,
+	Debug
+};
+
+// A collection of the system's buffer types. This allows us to collapse the 4 types into a single type to use a template argument elsewhere.
+template<typename InInputCmd, typename InSyncState, typename InAuxState, typename InDebugState = FNetSimProcessedFrameDebugInfo>
+struct TNetworkSimBufferTypes
+{
+	// Quick access to types when you know what you want
+	using TInputCmd = InInputCmd;
+	using TSyncState = InSyncState;
+	using TAuxState = InAuxState;
+	using TDebugState = InDebugState;
+
+	// Template access via ENetworkSimBufferTypeId when "which buffer" is parameterized
+	template<ENetworkSimBufferTypeId>
+	struct select_type
+	{
+		using type = void;
+	};
+
+	template<>
+	struct select_type<ENetworkSimBufferTypeId::Input>
+	{
+		using type = TInputCmd;
+	};
+
+	template<>
+	struct select_type<ENetworkSimBufferTypeId::Sync>
+	{
+		using type = TSyncState;
+	};
+
+	template<>
+	struct select_type<ENetworkSimBufferTypeId::Aux>
+	{
+		using type = TAuxState;
+	};
+
+	template<>
+	struct select_type<ENetworkSimBufferTypeId::Debug>
+	{
+		using type = TDebugState;
+	};
+};
+
+// ---------------------------------------------------------------------------------------------------------------------
+//	Ticking and Input processing helpers. We want to support fixed tick systems (but variable is most important).
+//	Whether your system is fixed or not depends on your input cmd type.
+// ---------------------------------------------------------------------------------------------------------------------
+
+// Basic struct trait for whether the input cmd is fixed or not. By default you are not fixed, you must specialize this template for your type and set value=true
+template<typename T>
+struct TNetworkSimInput_is_fixedtick
+{
+	static const bool value = false;
+};
+
+// Variable tick simulation time. For now this is a double.
+template<bool IsFixedTick=false>
+struct TSimulationTime
+{
+	TSimulationTime(double InTime=0.f) : Time(InTime) { }
+
+	template<typename T>
+	void AccumulateTimeFromInputCmd(const T& InputCmd)
+	{
+		Time += InputCmd.FrameDeltaTime;
+	}
+	
+	void NetSerialize(const FNetSerializeParams& P)
+	{
+		// Temp: we can devise a system like keyframe replication to make this more efficient
+		P.Ar << Time;
+	}
+
+	double Time = 0.f;
+};
+
+// Fixed tick simulation time. Just a frame counter.
+template<>
+struct TSimulationTime<true>
+{
+	TSimulationTime(int32 InTime=0) : Time(InTime) { }
+
+	template<typename T>
+	void AccumulateTimeFromInputCmd(const T& InputCmd)
+	{
+		Time++;
+	}
+
+	void NetSerialize(const FNetSerializeParams& P)
+	{
+		// Temp: we can devise a system like keyframe replication to make this more efficient
+		P.Ar << Time;
+	}
+
+	int32 Time = 0;
+};
+
+// Actual struct to use to track time. Parameterized by your NetworkSimBufferTypes and uses the input cmd to determine if you are fixed or not.
+template<typename TBufferTypes, typename TParent = TSimulationTime< TNetworkSimInput_is_fixedtick<typename TBufferTypes::TInputCmd>::value >>
+struct TSimulationTimeKeeper : public TParent
+{
+	using T = TSimulationTimeKeeper<TBufferTypes>;
+	using TParent::TParent;
+
+	FString ToString() const { return LexToString(this->Time); }
+
+	T& operator+= (const T &rhs) { this->Time += rhs.Time; return(*this); }
+	T& operator-= (const T &rhs) { this->Time -= rhs.Time; return(*this); }
+	
+	T operator+ (const T &rhs) { return T(this->Time - rhs.Time); }
+	T operator- (const T &rhs) { return T(this->Time - rhs.Time); }
+
+	bool operator<  (const T &rhs) const { return(this->Time < rhs.Time); }
+	bool operator<= (const T &rhs) const { return(this->Time <= rhs.Time); }
+	bool operator>  (const T &rhs) const { return(this->Time > rhs.Time); }
+	bool operator>= (const T &rhs) const { return(this->Time >= rhs.Time); }
+	bool operator== (const T &rhs) const { return(this->Time == rhs.Time); }
+	bool operator!= (const T &rhs) const { return(this->Time != rhs.Time); }
+};
+
+
+// ---------------------------------------------------------------------------------------------------------------------
+//	TNetworkSimBufferContainer
+//	Container for the actual replicated buffers that the system uses.
+//	Has compile time accessors for retrieving the buffers based on templated enum value.
+// ---------------------------------------------------------------------------------------------------------------------
+
+template<typename T>
+struct TNetworkSimBufferContainer
+{
+	// Collection of types we were assigned
+	using TBufferTypes = T;
+
+	// helper that returns the buffer type for a given BufferId (not the underlying type: the actual TReplicatedBuffer<TUnderlyingType>
+	template<ENetworkSimBufferTypeId TypeId>
+	struct select_buffer_type
+	{
+		// We are just wrapping a TReplicationBuffer around whatever type TBufferTypes::select_type returns (which returns the underlying type)
+		using type = TReplicationBuffer<typename TBufferTypes::template select_type<TypeId>::type>;
+	};
+
+	// The buffer types. This may look intimidating but is not that bad!
+	// We are using select_buffer_type to decide what the type should be for a given ENetworkSimBufferTypeId.
+	using TInputBuffer = typename select_buffer_type<ENetworkSimBufferTypeId::Input>::type;
+	using TSyncBuffer = typename select_buffer_type<ENetworkSimBufferTypeId::Sync>::type;
+	using TAuxBuffer = typename select_buffer_type<ENetworkSimBufferTypeId::Aux>::type;
+	using TDebugBuffer = typename select_buffer_type<ENetworkSimBufferTypeId::Debug>::type;
+
+	// The buffers themselves. Types are already declared above.
+	// If you are reading just to find the damn underlying type here, its TReplicationBuffer< whatever your struct type is >
+	TInputBuffer Input;
+	TSyncBuffer Sync;
+	TAuxBuffer Aux;
+	TDebugBuffer Debug;
+
+	// Finally, template accessor for getting buffers based on enum. This is really what all this junk is about.
+	// This allows other templated classes in the system to access a specific buffer from another templated argument
+	template<ENetworkSimBufferTypeId BufferId>
+	typename select_buffer_type<BufferId>::type& Get()
+	{
+		static_assert(!BufferId, "Failed to find specialized Get for your BufferId");
+	}
+
+	template<>
+	TInputBuffer& Get<ENetworkSimBufferTypeId::Input>()
+	{
+		return Input;
+	}
+
+	template<>
+	TSyncBuffer& Get<ENetworkSimBufferTypeId::Sync>()
+	{
+		return Sync;
+	}
+
+	template<>
+	TAuxBuffer& Get<ENetworkSimBufferTypeId::Aux>()
+	{
+		return Aux;
+	}
+
+	template<>
+	TDebugBuffer& Get<ENetworkSimBufferTypeId::Debug>()
+	{
+		return Debug;
+	}
+};
+
+
+// ---------------------------------------------------------------------------------------------------------------------
+//	TSimulationTicker
+//	This holds the data for where we are in ticking the sim. Such as how much time we've ticked it, how much we are allowed to tick it.
+//	It also holds input processing info: the last keyframe of the input buffer we processed and the last GFrameNumber we accepted input.
+// ---------------------------------------------------------------------------------------------------------------------
+
+template<typename TBufferTypes>
+struct TSimulationTickInfo
+{
+	using FSimulationTime = TSimulationTimeKeeper<TBufferTypes>;
+
+	int32 LastProcessedInputKeyframe;	// Tracks input keyframe that we last processed
+	uint32 LastLocalInputGFrameNumber;	// Tracks the last time we accepted local input via GetNextInputForWrite. Used to guard against accidental input latency due to ordering of input/sim ticking.
+
+	FSimulationTime MaxSimulationTime;			// Total max time we have been "given" to process. We cannot process more simulation time than this: doing so would be speed hacking.
+	FSimulationTime ProcessedSimulationTime;	// How much time we've actually processed. The only way to increment this is to process user commands or receive authoritative state from the network.
+
+	using TInputCmd = typename TBufferTypes::TInputCmd;
+	TInputCmd* GetNextInputForWrite(TNetworkSimBufferContainer<TBufferTypes>& Buffers)
+	{
+		ensure(LastLocalInputGFrameNumber != GFrameNumber);
+		LastLocalInputGFrameNumber = GFrameNumber;
+
+		TInputCmd* Next = nullptr;
+		if (Buffers.Input.GetHeadKeyframe() == LastProcessedInputKeyframe)
+		{
+			// Only return a cmd if we have processed the last one. This is a bit heavy handed but is a good practice to start with. We want buffering of input to be very explicit, never something that accidentally happens.
+			Next = Buffers.Input.GetWriteNext();
+			*Next = TInputCmd();
+		}
+		return Next;
+	}
+};
\ No newline at end of file
diff --git a/Engine/Plugins/Runtime/NetworkPrediction/readme.txt b/Engine/Plugins/Runtime/NetworkPrediction/readme.txt
--- a/Engine/Plugins/Runtime/NetworkPrediction/readme.txt
+++ b/Engine/Plugins/Runtime/NetworkPrediction/readme.txt
@@ -34,7 +34,7 @@ Code to look at:
 -MockNetworkSimulation.h: Good place to start to see a bare minimum example of a predicted network simulation. Both the simulation and actor component live in this header.
 -NetworkPredictionExtrasFlyingPawn.h: Starting point to see how a pawn using flying movement is setup. In NetworkPredictionExtras.
 -FlyingMovement.h: The actual implementation of the flying movement system. In NetworkPrediction.
--NetworkSimulationModelTemplates.h: this is the lowest level guts of the system. This is what every simulation ultimately uses to stay in sync.
+-NetworkSimulationModel.h: this is the lowest level guts of the system. This is what every simulation ultimately uses to stay in sync.
  
 Extremely high level technical overview:
 -Concept of "NetworkedSimulationModel": defining a tightly constrained (gameplay) simulation that can be run predictively on clients and authoritatively on server. E.g., a "Movement System" is a NetworkedSimulationModel.
@@ -95,6 +95,13 @@ Smoothing: Taking the output of the simulation and applying an additional layer
 // Release notes
 // ----------------------------------------------------------------------------------------------------------
  
+Update (8-9-19)
+-Big refactor underway within the NetworkSimModel. Some files are renamed or moved as well. Still a bit to do but this was a good commit point.
+-Mainly to facilitate clean implementation of interpolation/extrapolation/forward predict option for simulated proxy simulations: this was exposing weaknesses in the templated implementations.
+-Working towards improving "how and when simulation is allowed to advanced" (see TNetworkSimulationTickInfo)
+-Incomplete, but thinking about global management of active simulations: making sure reconcile and Tick are called in the right places. See notes in NetworkSimulationGlobalManager.h
+
+ 
 Update (8-2-19)
 -We want to take a look non player controlled simulations, such as doors, elevators, "pushers" before going deeper on ability system or movement.
 -Basic ParametricMoverment system checked in. This is also in an incomplete state. ("Pushing" the flying movement component does not work and won't be the focus for now until we are looking for closely at movement specifically).
diff --git a/Engine/Plugins/Runtime/NetworkPredictionExtras/Source/NetworkPredictionExtras/Private/MockNetworkSimulation.cpp b/Engine/Plugins/Runtime/NetworkPredictionExtras/Source/NetworkPredictionExtras/Private/MockNetworkSimulation.cpp
--- a/Engine/Plugins/Runtime/NetworkPredictionExtras/Source/NetworkPredictionExtras/Private/MockNetworkSimulation.cpp
+++ b/Engine/Plugins/Runtime/NetworkPredictionExtras/Source/NetworkPredictionExtras/Private/MockNetworkSimulation.cpp
@@ -176,6 +176,8 @@ void UMockNetworkSimulationComponent::TickComponent(float DeltaTime, enum ELevel
 	// -------------------------------------
 	if (NetworkSim)
 	{
+		PreTickSimulation(DeltaTime); // Fixme
+
 		// Check if we should trip a mispredict. (Note how its not possible to do this inside the Update function!)
 		if (OwnerRole == ROLE_Authority && MockNetworkSimCVars::RequestMispredict)
 		{
diff --git a/Engine/Plugins/Runtime/NetworkPredictionExtras/Source/NetworkPredictionExtras/Private/NetworkPredictionExtrasFlyingPawn.cpp b/Engine/Plugins/Runtime/NetworkPredictionExtras/Source/NetworkPredictionExtras/Private/NetworkPredictionExtrasFlyingPawn.cpp
--- a/Engine/Plugins/Runtime/NetworkPredictionExtras/Source/NetworkPredictionExtras/Private/NetworkPredictionExtrasFlyingPawn.cpp
+++ b/Engine/Plugins/Runtime/NetworkPredictionExtras/Source/NetworkPredictionExtras/Private/NetworkPredictionExtrasFlyingPawn.cpp
@@ -41,7 +41,7 @@ ANetworkPredictionExtrasFlyingPawn::ANetworkPredictionExtrasFlyingPawn()
 		//  3. Tick this actor (ANetworkPredictionExtrasFlyingPawn::Tick). This wants the latest movement state (computed in step #2).
 		//
 		// The ::SetLocallyControlledPreTick isn't strictly necessary. A player controller or some other object could tick first, submit the input, then let the sim tick and this actor tick.
-		// But for this simple example, its clearer to contain everything here. The lambda is just a helper to save you from dealing with tick graph prereq headaches.
+		// But for this simple example, its clearer to contain everything here. The delegate is just a helper to save you from dealing with tick graph prereq headaches.
 
 		FlyingMovementComponent->SetLocallyControlledPreTick(UNetworkPredictionComponent::FPreSimTickDelegate::CreateUObject(this, &ThisClass::GenerateLocalInput));
 
diff --git a/Engine/Plugins/Runtime/NetworkPredictionExtras/Source/NetworkPredictionExtras/Public/MockNetworkSimulation.h b/Engine/Plugins/Runtime/NetworkPredictionExtras/Source/NetworkPredictionExtras/Public/MockNetworkSimulation.h
--- a/Engine/Plugins/Runtime/NetworkPredictionExtras/Source/NetworkPredictionExtras/Public/MockNetworkSimulation.h
+++ b/Engine/Plugins/Runtime/NetworkPredictionExtras/Source/NetworkPredictionExtras/Public/MockNetworkSimulation.h
@@ -14,7 +14,7 @@
 #include "Misc/OutputDevice.h"
 #include "Misc/CoreDelegates.h"
 
-#include "NetworkSimulationModelTemplates.h"
+#include "NetworkSimulationModel.h"
 #include "NetworkPredictionComponent.h"
 
 #include "MockNetworkSimulation.generated.h"
@@ -116,8 +116,10 @@ struct FMockAuxState
 	}
 };
 
+using TMockNetworkSimulationBufferTypes = TNetworkSimBufferTypes<FMockInputCmd, FMockSyncState, FMockAuxState>;
+
 // Actual definition of our network simulation.
-class FMockNetworkSimulation : public TNetworkedSimulationModel<FMockNetworkSimulation, FMockInputCmd, FMockSyncState, FMockAuxState>
+class FMockNetworkSimulation : public TNetworkedSimulationModel<FMockNetworkSimulation, TMockNetworkSimulationBufferTypes>
 {
 public:
 
