Continued WIP on NetworkPrediction
-Refactoring netsimmodel's internal ticking state
-Reworking template implementations
-Setting us up for better seamless fixed tick support

Note: this is still a WIP and breaks some things that previous worked. Committing now to unblock mac builds.

#rnx

diff --git a/Engine/Plugins/Runtime/NetworkPrediction/Source/NetworkPrediction/Private/Movement/FlyingMovementComponent.cpp b/Engine/Plugins/Runtime/NetworkPrediction/Source/NetworkPrediction/Private/Movement/FlyingMovementComponent.cpp
--- a/Engine/Plugins/Runtime/NetworkPrediction/Source/NetworkPrediction/Private/Movement/FlyingMovementComponent.cpp
+++ b/Engine/Plugins/Runtime/NetworkPrediction/Source/NetworkPrediction/Private/Movement/FlyingMovementComponent.cpp
@@ -133,24 +133,6 @@ void UFlyingMovementComponent::TickComponent(float DeltaTime, enum ELevelTick Ti
 	}
 }
 
-FlyingMovement::FInputCmd* UFlyingMovementComponent::GetNextClientInputCmdForWrite(float DeltaTime)
-{
-	// This is unfortunate but we must check here to initialize the InputCmdBuffer.
-	// -Tick alone is not appropriate since the PC may tick first and generate new user cmds (before our first tick)
-	// -Startup functions (BeginPlay, PostInitializeComponent etc) are not appropriate because Net Role is not known.
-	// -PostNetReceive could work but would be inefficient since we would be checking each simulated proxy too, ever PostNetReceive.
-	//
-	// This works out best because its only the locally controlled actor that will call into this.
-	// Its just unfortunate it will only pass the first time but we must keep checking.
-	CheckOwnerRoleChange();
-
-	if (NetworkSim.IsValid())
-	{
-		return NetworkSim->GetNextInputForWrite(DeltaTime);
-	}
-	return nullptr;
-}
-
 // ----------------------------------------------------------------------------------------------------------
 //	Movement System Driver
 //
@@ -180,6 +162,12 @@ void UFlyingMovementComponent::PreSimSync(const FlyingMovement::FMoveState& Sync
 	}
 }
 
+void UFlyingMovementComponent::ProduceInput(const FlyingMovement::TSimTime& SimFrameTime, FlyingMovement::FInputCmd& Cmd)
+{
+	// This isn't ideal. It probably makes sense for the component to do all the input binding rather.
+	ProduceInputDelegate.ExecuteIfBound(SimFrameTime, Cmd);
+}
+
 void UFlyingMovementComponent::FinalizeFrame(const FlyingMovement::FMoveState& SyncState)
 {
 	PreSimSync(SyncState);
diff --git a/Engine/Plugins/Runtime/NetworkPrediction/Source/NetworkPrediction/Private/Movement/FlyingMovementSimulation.cpp b/Engine/Plugins/Runtime/NetworkPrediction/Source/NetworkPrediction/Private/Movement/FlyingMovementSimulation.cpp
--- a/Engine/Plugins/Runtime/NetworkPrediction/Source/NetworkPrediction/Private/Movement/FlyingMovementSimulation.cpp
+++ b/Engine/Plugins/Runtime/NetworkPrediction/Source/NetworkPrediction/Private/Movement/FlyingMovementSimulation.cpp
@@ -152,7 +152,7 @@ namespace FlyingMovement
 		return 0.f;
 	}
 
-	void FMovementSystem::Update(IMovementDriver* Driver, const FInputCmd& InputCmd, const FMoveState& InputState, FMoveState& OutputState, const FAuxState& AuxState)
+	void FMovementSystem::Update(IMovementDriver* Driver, const TSimTime& DeltaTimeMS, const FInputCmd& InputCmd, const FMoveState& InputState, FMoveState& OutputState, const FAuxState& AuxState)
 	{
 		// ----------------------------------
 		// AuxState todo
@@ -163,7 +163,7 @@ namespace FlyingMovement
 
 		// ----------------------------------------------------------------
 
-		const float DeltaSeconds = InputCmd.FrameDeltaTime;
+		const float DeltaSeconds = DeltaTimeMS.ToRealTimeSeconds();
 		OutputState = InputState;
 
 		IBaseMovementDriver& BaseMovementDriver = Driver->GetBaseMovementDriver();
@@ -175,7 +175,7 @@ namespace FlyingMovement
 		//	In this simulation, the rotation update isn't allowed to "fail". We don't expect the collision query to be able to fail the rotational update.
 		// --------------------------------------------------------------
 
-		OutputState.Rotation += (InputCmd.RotationInput * InputCmd.FrameDeltaTime);
+		OutputState.Rotation += (InputCmd.RotationInput * DeltaSeconds);
 		OutputState.Rotation.Normalize();
 
 		const FQuat OutputQuat = OutputState.Rotation.Quaternion();
diff --git a/Engine/Plugins/Runtime/NetworkPrediction/Source/NetworkPrediction/Private/Movement/ParametricMovement.cpp b/Engine/Plugins/Runtime/NetworkPrediction/Source/NetworkPrediction/Private/Movement/ParametricMovement.cpp
--- a/Engine/Plugins/Runtime/NetworkPrediction/Source/NetworkPrediction/Private/Movement/ParametricMovement.cpp
+++ b/Engine/Plugins/Runtime/NetworkPrediction/Source/NetworkPrediction/Private/Movement/ParametricMovement.cpp
@@ -31,14 +31,14 @@ static FAutoConsoleVariableRef CVarForceNetUpdate(TEXT("parametricmover.ForceNet
 
 namespace ParametricMovement
 {
-	void FMovementSystem::Update(IMovementDriver* Driver, const FInputCmd& InputCmd, const FMoveState& InputState, FMoveState& OutputState, const FAuxState& AuxState)
+	void FMovementSystem::Update(IMovementDriver* Driver, const TSimTime& SimeTimeDeltaMS, const FInputCmd& InputCmd, const FMoveState& InputState, FMoveState& OutputState, const FAuxState& AuxState)
 	{
 		IBaseMovementDriver& BaseMovementDriver = Driver->GetBaseMovementDriver();
-		const float DeltaSeconds = InputCmd.FrameDeltaTime;
+		const float DeltaSeconds = SimeTimeDeltaMS.ToRealTimeSeconds();
 
 		// Advance parametric time. This won't always be linear: we could loop, rewind/bounce, etc
 		const float InputPlayRate = InputCmd.PlayRate.Get(InputState.PlayRate); // Returns InputCmds playrate if set, else returns previous state's playrate		
-		Driver->AdvanceParametricTime(InputState.Position, InputPlayRate, OutputState.Position, OutputState.PlayRate, InputCmd.FrameDeltaTime);
+		Driver->AdvanceParametricTime(InputState.Position, InputPlayRate, OutputState.Position, OutputState.PlayRate, DeltaSeconds);
 
 		// We have our time that we should be at. We just need to move primitive component to that position.
 		// Again, note that we expect this cannot fail. We move like this so that it can push things, but we don't expect failure.
@@ -205,15 +205,6 @@ void UParametricMovementComponent::TickSimulation(float DeltaTimeSeconds)
 
 	if (NetworkSim && OwnerRole != ROLE_None)
 	{
-		if (OwnerRole == ROLE_Authority)
-		{
-			if (ParametricMovement::FInputCmd* NextCmd = NetworkSim->GetNextInputForWrite(DeltaTimeSeconds))
-			{
-				NextCmd->PlayRate = PendingPlayRate;
-				PendingPlayRate.Reset();
-			}
-		}
-
 		ParametricMovement::FMovementSystem::FTickParameters Parameters;
 		Parameters.Role = OwnerRole;
 		Parameters.LocalDeltaTimeSeconds = DeltaTimeSeconds;
@@ -227,4 +218,10 @@ void UParametricMovementComponent::TickSimulation(float DeltaTimeSeconds)
 	{
 		GetOwner()->ForceNetUpdate();
 	}
+}
+
+void UParametricMovementComponent::ProduceInput(const ParametricMovement::TSimTime& SimFrameTime, ParametricMovement::FInputCmd& Cmd)
+{
+	Cmd.PlayRate = PendingPlayRate;
+	PendingPlayRate.Reset();
 }
\ No newline at end of file
diff --git a/Engine/Plugins/Runtime/NetworkPrediction/Source/NetworkPrediction/Private/NetworkPredictionComponent.cpp b/Engine/Plugins/Runtime/NetworkPrediction/Source/NetworkPrediction/Private/NetworkPredictionComponent.cpp
--- a/Engine/Plugins/Runtime/NetworkPrediction/Source/NetworkPrediction/Private/NetworkPredictionComponent.cpp
+++ b/Engine/Plugins/Runtime/NetworkPrediction/Source/NetworkPrediction/Private/NetworkPredictionComponent.cpp
@@ -1,6 +1,5 @@
 // Copyright 1998-2019 Epic Games, Inc. All Rights Reserved.
 
-
 #include "NetworkPredictionComponent.h"
 #include "Net/UnrealNetwork.h"
 #include "GameFramework/Actor.h"
@@ -96,10 +95,6 @@ bool UNetworkPredictionComponent::IsLocallyControlled()
 void UNetworkPredictionComponent::PreTickSimulation(float DeltaTime)
 {
 	CheckOwnerRoleChange();
-	if (PreTickLocallyControlledSim.IsBound() && IsLocallyControlled())
-	{
-		PreTickLocallyControlledSim.Execute(DeltaTime);
-	}
 }
 
 void UNetworkPredictionComponent::CheckOwnerRoleChange()
diff --git a/Engine/Plugins/Runtime/NetworkPrediction/Source/NetworkPrediction/Private/NetworkPredictionTests.cpp b/Engine/Plugins/Runtime/NetworkPrediction/Source/NetworkPrediction/Private/NetworkPredictionTests.cpp
new file mode 100644
--- /dev/null
+++ b/Engine/Plugins/Runtime/NetworkPrediction/Source/NetworkPrediction/Private/NetworkPredictionTests.cpp
@@ -0,0 +1,38 @@
+// Copyright 1998-2019 Epic Games, Inc. All Rights Reserved.
+
+#include "NetworkSimulationModel.h"
+
+// ------------------------------------------------------------------------------------------------------------
+//	Testing Fixed vs Variable tick rate struct declration
+// ------------------------------------------------------------------------------------------------------------
+
+struct FDummy { };
+using FDummyBufferTypes = TNetworkSimBufferTypes<FDummy, FDummy, FDummy>;
+
+template<typename TTickSettings>
+void TickSettingsTest(TTickSettings& Settings)
+{
+	using TBufferTypes = TInternalBufferTypes<FDummyBufferTypes, TTickSettings>;
+
+	TSimulationTickState<TTickSettings>	TickState;
+
+	float RealTime = FMath::Rand();
+	TickState.GiveSimulationTime(RealTime);
+
+}
+
+void NetSimulationTypesTest()
+{
+	TNetworkSimTickSettings<> DefaultTickSettings;
+	TickSettingsTest(DefaultTickSettings);
+
+
+	TNetworkSimTickSettings<0, 60> VariableTickSettingsWithMax;
+	TickSettingsTest(VariableTickSettingsWithMax);
+
+	TNetworkSimTickSettings<20> FixedTickSettings;
+	TickSettingsTest(FixedTickSettings);
+
+	//TNetworkSimTickSettings<20, 60> StaticAssertSettings;
+	//TickSettingsTest(StaticAssertSettings);
+}
diff --git a/Engine/Plugins/Runtime/NetworkPrediction/Source/NetworkPrediction/Public/Movement/FlyingMovement.h b/Engine/Plugins/Runtime/NetworkPrediction/Source/NetworkPrediction/Public/Movement/FlyingMovement.h
--- a/Engine/Plugins/Runtime/NetworkPrediction/Source/NetworkPrediction/Public/Movement/FlyingMovement.h
+++ b/Engine/Plugins/Runtime/NetworkPrediction/Source/NetworkPrediction/Public/Movement/FlyingMovement.h
@@ -30,17 +30,12 @@ namespace FlyingMovement
 	// State the client generates
 	struct FInputCmd
 	{
-
-		// Client's FrameTime. This is essential for a non-fixed step simulation. The server will run the movement simulation for this client at this rate.
-		float FrameDeltaTime = 0.f;
-
 		// Input: "pure" input for this frame. At this level, frame time has not been accounted for. (E.g., "move straight" would be (1,0,0) regardless of frame time)
 		FRotator RotationInput;
 		FVector MovementInput;
 
 		void NetSerialize(const FNetSerializeParams& P)
 		{
-			P.Ar << FrameDeltaTime;
 			P.Ar << RotationInput;
 			P.Ar << MovementInput;
 		}
@@ -53,7 +48,6 @@ namespace FlyingMovement
 			}
 			else if (P.Context == EStandardLoggingContext::Full)
 			{
-				P.Ar->Logf(TEXT("FrameDeltaTime: %.4f"), FrameDeltaTime);
 				P.Ar->Logf(TEXT("Movement: %s"), *MovementInput.ToString());
 				P.Ar->Logf(TEXT("ViewRot: %s"), *RotationInput.ToString());
 			}
@@ -131,12 +125,15 @@ namespace FlyingMovement
 		};
 
 		/** Main update function */
-		static void Update(IMovementDriver* Driver, const FInputCmd& InputCmd, const FMoveState& InputState, FMoveState& OutputState, const FAuxState& AuxState);
+		static void Update(IMovementDriver* Driver, const TSimTime& DeltaTimeMS, const FInputCmd& InputCmd, const FMoveState& InputState, FMoveState& OutputState, const FAuxState& AuxState);
 
 		/** Dev tool to force simple mispredict */
 		static bool ForceMispredict;
 	};
 
+	// Simulation time used by the movement system
+	using TSimTime = FMovementSystem::TSimTime;
+
 	// general tolerance value for rotation checks
 	static const float ROTATOR_TOLERANCE = (1e-3);
 
@@ -161,15 +158,16 @@ public:
 	
 	virtual void TickComponent(float DeltaTime, enum ELevelTick TickType, FActorComponentTickFunction *ThisTickFunction) override;
 
-	// Gets writable copy of FClientInputCmd for this frame. Should only be called once per frame by the local player.
-	FlyingMovement::FInputCmd* GetNextClientInputCmdForWrite(float DeltaTime);
-
 	IBaseMovementDriver& GetBaseMovementDriver() override final { return *static_cast<IBaseMovementDriver*>(this); }
 
 	void InitSyncState(FlyingMovement::FMoveState& OutSyncState) const override;
 	void PreSimSync(const FlyingMovement::FMoveState& SyncState) override;
+	void ProduceInput(const FlyingMovement::TSimTime& SimFrameTime, FlyingMovement::FInputCmd& Cmd) override;
 	void FinalizeFrame(const FlyingMovement::FMoveState& SyncState) override;
 
+	DECLARE_DELEGATE_TwoParams(FProduceFlyingInput, const FlyingMovement::TSimTime& /*SimFrameTime*/, FlyingMovement::FInputCmd& /*Cmd*/)
+	FProduceFlyingInput ProduceInputDelegate;
+
 protected:
 
 	// Network Prediction
diff --git a/Engine/Plugins/Runtime/NetworkPrediction/Source/NetworkPrediction/Public/Movement/ParametricMovement.h b/Engine/Plugins/Runtime/NetworkPrediction/Source/NetworkPrediction/Public/Movement/ParametricMovement.h
--- a/Engine/Plugins/Runtime/NetworkPrediction/Source/NetworkPrediction/Public/Movement/ParametricMovement.h
+++ b/Engine/Plugins/Runtime/NetworkPrediction/Source/NetworkPrediction/Public/Movement/ParametricMovement.h
@@ -14,16 +14,11 @@ namespace ParametricMovement
 	// State the client generates
 	struct FInputCmd
 	{
-		// Client's FrameTime. This is essential for a non-fixed step simulation. The server will run the movement simulation for this client at this rate.
-		float FrameDeltaTime = 0.f;
-
 		// Input Playrate. This being set can be thought of "telling the simulation what its new playrate should be"
 		TOptional<float> PlayRate;
 
 		void NetSerialize(const FNetSerializeParams& P)
 		{
-			// FIXME: quantization and effecient packing of this data needs to be looked at/generalized in a nice way
-			P.Ar << FrameDeltaTime;
 			P.Ar << PlayRate;
 		}
 
@@ -35,7 +30,6 @@ namespace ParametricMovement
 			}
 			else if (P.Context == EStandardLoggingContext::Full)
 			{
-				P.Ar->Logf(TEXT("FrameDeltaTime: %.4f"), FrameDeltaTime);
 				if (PlayRate.IsSet())
 				{
 					P.Ar->Logf(TEXT("PlayRate: %.2f"), PlayRate.GetValue());
@@ -94,7 +88,7 @@ namespace ParametricMovement
 	using TMovementBufferTypes = TNetworkSimBufferTypes<FInputCmd, FMoveState, FAuxState>;
 
 	// Actual definition of our network simulation.
-	class FMovementSystem : public TNetworkedSimulationModel<FMovementSystem, TMovementBufferTypes>
+	class FMovementSystem : public TNetworkedSimulationModel<FMovementSystem, TMovementBufferTypes, TNetworkSimTickSettings<30>>
 	{
 	public:
 
@@ -115,12 +109,14 @@ namespace ParametricMovement
 		};
 
 		/** Main update function */
-		static void Update(IMovementDriver* Driver, const FInputCmd& InputCmd, const FMoveState& InputState, FMoveState& OutputState, const FAuxState& AuxState);
+		static void Update(IMovementDriver* Driver, const TSimTime& SimeTimeDeltaMS, const FInputCmd& InputCmd, const FMoveState& InputState, FMoveState& OutputState, const FAuxState& AuxState);
 
 		/** Dev tool to force simple mispredict */
 		static bool ForceMispredict;
 	};
 
+	using TSimTime = FMovementSystem::TSimTime;
+
 } // End namespace
 
 // Needed to trick UHT into letting UMockNetworkSimulationComponent implement. UHT cannot parse the ::
@@ -152,6 +148,7 @@ class NETWORKPREDICTION_API UParametricMovementComponent : public UBaseMovementC
 	// Base TNetworkModelSimulation driver
 	void InitSyncState(ParametricMovement::FMoveState& OutSyncState) const override;
 	void FinalizeFrame(const ParametricMovement::FMoveState& SyncState) override;
+	void ProduceInput(const ParametricMovement::TSimTime& SimFrameTime, ParametricMovement::FInputCmd& Cmd);
 
 	// Base Movement Driver
 	IBaseMovementDriver& GetBaseMovementDriver() override final { return *static_cast<IBaseMovementDriver*>(this); }
@@ -168,7 +165,7 @@ protected:
 	void InitializeForNetworkRole(ENetRole Role) override;
 	FNetworkSimulationModelInitParameters GetSimulationInitParameters(ENetRole Role) override;
 
-	TUniquePtr<ParametricMovement::FMovementSystem> NetworkSim;
+	TUniquePtr<ParametricMovement::FMovementSystem> NetworkSim;	
 
 	// ------------------------------------------------------------------------
 	// Temp Parametric movement example
diff --git a/Engine/Plugins/Runtime/NetworkPrediction/Source/NetworkPrediction/Public/NetworkPredictionComponent.h b/Engine/Plugins/Runtime/NetworkPrediction/Source/NetworkPrediction/Public/NetworkPredictionComponent.h
--- a/Engine/Plugins/Runtime/NetworkPrediction/Source/NetworkPrediction/Public/NetworkPredictionComponent.h
+++ b/Engine/Plugins/Runtime/NetworkPrediction/Source/NetworkPrediction/Public/NetworkPredictionComponent.h
@@ -35,13 +35,6 @@ public:
 	virtual void InitializeComponent() override;
 	virtual void PreReplication(IRepChangedPropertyTracker & ChangedPropertyTracker) override;
 	virtual void PreNetReceive() override;
-
-	// PreSimulationTick callback.
-	// Called right in ::TickComponent for autonomous proxies. This is optional and intended to be a convencience for generating local input.
-	// If you don't use this function, you should manually be submitting client input prior to simulation tick (via GetNextClientInputCmdForWrite).
-	// (It is important to submit input prior to ticking the sim. If you get it backwards, you introduce local latency and cause problems with variable time steps).
-	DECLARE_DELEGATE_OneParam(FPreSimTickDelegate, float /* DeltaSeconds */)
-	void SetLocallyControlledPreTick(const FPreSimTickDelegate& PreSimTick) { PreTickLocallyControlledSim = PreSimTick; }
 	
 	virtual void Reconcile() { }
 	virtual void TickSimulation(float DeltaTimeSeconds) { }
@@ -77,8 +70,6 @@ protected:
 
 private:
 
-	FPreSimTickDelegate PreTickLocallyControlledSim;
-
 	UPROPERTY(Replicated, transient)
 	FReplicationProxy ReplicationProxy_Autonomous;
 
diff --git a/Engine/Plugins/Runtime/NetworkPrediction/Source/NetworkPrediction/Public/NetworkSimulationModel.h b/Engine/Plugins/Runtime/NetworkPrediction/Source/NetworkPrediction/Public/NetworkSimulationModel.h
--- a/Engine/Plugins/Runtime/NetworkPrediction/Source/NetworkPrediction/Public/NetworkSimulationModel.h
+++ b/Engine/Plugins/Runtime/NetworkPrediction/Source/NetworkPrediction/Public/NetworkSimulationModel.h
@@ -4,397 +4,7 @@
 #include "NetworkPredictionTypes.h"
 #include "NetworkSimulationModelBuffer.h"
 #include "NetworkSimulationModelTypes.h"
-
-// ----------------------------------------------------------------------------------------------------------------------------------------------------------------------
-//	CVars and compile time constants
-// ----------------------------------------------------------------------------------------------------------------------------------------------------------------------
-
-namespace NetworkSimulationModelCVars
-{
-static int32 EnableLocalPrediction = 1;
-static FAutoConsoleVariableRef CVarEnableLocalPrediction(TEXT("ns.EnableLocalPrediction"),
-	EnableLocalPrediction, TEXT("Toggle local prediction. Meant for debugging."), ECVF_Default);
-
-static int32 ForceReconcile = 0;
-static FAutoConsoleVariableRef CVarForceReconcile(TEXT("ns.ForceReconcile"),
-	ForceReconcile, TEXT("Forces reconcile even if state does not differ. E.g, force resimulation after every netupdate."), ECVF_Default);
-
-static int32 MaxInputCmdsFrame = 1;
-static FAutoConsoleVariableRef CVarMaxInputCmdsFrame(TEXT("ns.MaxInputCmdsFrame"),
-	MaxInputCmdsFrame, TEXT("Max cap on how many input cmds can be processed in a frame."), ECVF_Default);
-}
-
-#ifndef NETSIM_NETCONSTANT_NUM_BITS_KEYFRAME
-	#define NETSIM_NETCONSTANT_NUM_BITS_KEYFRAME 8	// Allows you to override this setting via UBT, but access via cleaner FActorMotionNetworkingConstants::NUM_BITS_KEYFRAME
-#endif
-
-struct FNetworkSimulationSerialization
-{
-	// How many bits we use to encode the key frame number for buffers.
-	// Client Keyframes are stored locally as 32 bit integers, but we use a smaller # of bits to NetSerialize.
-	// Frames are only relatively relevant: the absolute value doesn't really matter. We just need to detect newer/older.
-	enum { NUM_BITS_KEYFRAME = NETSIM_NETCONSTANT_NUM_BITS_KEYFRAME };		
-
-	// Abs max value we encode into the bit writer
-	enum { MAX_KEYFRAME_WRITE = 1 << NUM_BITS_KEYFRAME };
-
-	// This is the threshold at which we would wrap around and incorrectly assign a keyframe on the receiving side.
-	// E.g, If there are KEYFRAME_ERROR_THRESHOLD keyframes that do not make it across from sender->receiver, the
-	// receiver will have incorrect local values. With 8 bits, this works out to be 128 frames or about 2 seconds at 60fps.
-	enum { KEYFRAME_ERROR_THRESHOLD = MAX_KEYFRAME_WRITE / 2};
-
-	// Helper to serialize the int32 HeadKeyframe. Returns the unpacked value (this will be same as input in the save path)
-	static int32 SerializeKeyframe(FArchive& Ar, int32 LocalHeadKeyframe)
-	{
-		if (Ar.IsSaving())
-		{
-			((FNetBitWriter&)Ar).WriteIntWrapped( LocalHeadKeyframe, MAX_KEYFRAME_WRITE );
-			return LocalHeadKeyframe;
-		}
-		
-		return MakeRelative(((FNetBitReader&)Ar).ReadInt( MAX_KEYFRAME_WRITE ), LocalHeadKeyframe, MAX_KEYFRAME_WRITE );
-	}
-};
-
-
-// ----------------------------------------------------------------------------------------------------------------------------------------------------------------------
-//	Templated Replicators/Reconciliars
-// ----------------------------------------------------------------------------------------------------------------------------------------------------------------------
-
-/** Replicates only the latest element. Does not synchronize keyframe */
-template<typename TBufferTypes, ENetworkSimBufferTypeId BufferId>
-struct TReplicator_LatestOnly
-{
-	using TState = typename TBufferTypes::template select_type<BufferId>::type;
-
-	int32 GetProxyDirtyCount(TNetworkSimBufferContainer<TBufferTypes>& Buffers) const
-	{ 
-		return Buffers.template Get<BufferId>().GetDirtyCount(); 
-	}
-
-	void NetSerialize(const FNetSerializeParams& P, TNetworkSimBufferContainer<TBufferTypes>& Buffers, const TSimulationTickInfo<TBufferTypes>& TickInfo)
-	{
-		FArchive& Ar = P.Ar;
-
-		auto& Buffer = Buffers.template Get<BufferId>();
-		TState* State = nullptr;
-		
-		if (Ar.IsSaving())
-		{
-			State = Buffer.GetElementFromHead(0);
-			check(State); // We should not be here if the buffer is empty. Want to avoid serializing an "empty" flag at the top here.
-		}
-		else
-		{
-			State = Buffer.GetWriteNext();
-		}
-
-		State->NetSerialize(Ar);
-	}
-};
-
-template<typename TBufferTypes>
-struct TReplicator_SimulatedExtrapolatedReconciliar
-{
-	using TSyncState = typename TBufferTypes::TSyncState;
-	
-	int32 GetProxyDirtyCount(TNetworkSimBufferContainer<TBufferTypes>& Buffers) const
-	{
-		return Buffers.Sync.GetDirtyCount(); 
-	}
-
-	void NetSerialize(const FNetSerializeParams& P, TNetworkSimBufferContainer<TBufferTypes>& Buffers, const TSimulationTickInfo<TBufferTypes>& TickInfo)
-	{
-		FArchive& Ar = P.Ar;
-		TSyncState* State = nullptr;
-		if (Ar.IsSaving())
-		{
-			State = Buffers.Sync.GetElementFromHead(0);
-			check(State); // We should not be here if the buffer is empty. Want to avoid serializing an "empty" flag at the top here.
-			SerializedTimeKeeper = TickInfo.ProcessedSimulationTime;
-		}
-		else
-		{
-			State = Buffers.Sync.GetWriteNext();
-			bPendingReconciliation = true;
-		}
-
-		SerializedTimeKeeper.NetSerialize(Ar);
-		State->NetSerialize(Ar);
-	}
-
-	template<typename T, typename TDriver>
-	void Reconcile(TDriver* Driver, TNetworkSimBufferContainer<TBufferTypes>& Buffers, TSimulationTickInfo<TBufferTypes>& TickInfo)
-	{
-		if (!bPendingReconciliation)
-			return;
-
-		bPendingReconciliation = false;
-
-		// This is still being worked out. We are almost there but not quite.
-		// Processing input commands is pretty rigid right now. Since we just added a new SyncState keyframe in ::NetSerialize,
-		// we need to insert a fake input command so the extrapolation code in the NetSimModel can continue to work.
-		// This can be improved.
-
-		{
-			// Generate a new, fake, command since we just added a new sync state to head (via GetWriteNext)
-			TickInfo.GetNextInputForWrite(Buffers);
-			// Gotta fake this too
-			TickInfo.LastLocalInputGFrameNumber--; 
-			// Set our LastProcessedInputKeyframe to fake that we handled it
-			TickInfo.LastProcessedInputKeyframe = Buffers.Sync.GetHeadKeyframe();
-		}
-	}
-
-private:
-	
-	TSimulationTimeKeeper<TBufferTypes> SerializedTimeKeeper;
-	bool bPendingReconciliation = false;
-};
-
-/** Replicates the latest sequence of N elements. N is dynamic, not compiled in. */
-template<typename TBufferTypes, ENetworkSimBufferTypeId BufferId>
-struct TReplicator_DynamicSequence
-{
-	int32 GetProxyDirtyCount(TNetworkSimBufferContainer<TBufferTypes>& Buffers) const 
-	{
-		return Buffers.template Get<BufferId>().GetDirtyCount(); 
-	}
-	
-	void NetSerialize(const FNetSerializeParams& P, TNetworkSimBufferContainer<TBufferTypes>& Buffers, const TSimulationTickInfo<TBufferTypes>& TickInfo)
-	{
-		auto& Buffer = Buffers.template Get<BufferId>();
-		FArchive& Ar = P.Ar;
-		uint8 SerializedNumElements = FMath::Min<uint8>(MaxNumElements, Buffer.GetNumValidElements());
-		Ar << SerializedNumElements;
-
-		const int32 HeadKeyframe = FNetworkSimulationSerialization::SerializeKeyframe(Ar, Buffer.GetHeadKeyframe());
-		const int32 StartingKeyframe = FMath::Max(0, HeadKeyframe - SerializedNumElements + 1);
-
-		if (Ar.IsLoading())
-		{
-			const int32 PrevHead = Buffer.GetHeadKeyframe();
-			if (PrevHead < StartingKeyframe && PrevHead >= 0)
-			{
-				// There is a gap in the stream. In some cases, we want this to be a "fault" and bubble up. We may want to synthesize state or maybe we just skip ahead.
-				UE_LOG(LogNetworkSim, Warning, TEXT("Fault: gap in received buffer. PrevHead: %d. Received: %d-%d. Reseting previous buffer contents"), PrevHead, StartingKeyframe, HeadKeyframe);
-			}
-
-			Buffer.ResetNextHeadKeyframe(StartingKeyframe);
-		}
-
-		for (int32 Keyframe = StartingKeyframe; Keyframe <= HeadKeyframe; ++Keyframe)
-		{
-			// This, as is, is bad. The intention is that these functions serialize multiple items in some delta compressed fashion.
-			// As is, we are just serializing the elements individually.
-			auto* Cmd = Ar.IsLoading() ? Buffer.GetWriteNext() : Buffer.FindElementByKeyframe(Keyframe);
-			Cmd->NetSerialize(P);
-		}
-	}
-
-private:
-	
-	uint8 MaxNumElements = 3;
-};
-
-/** Replicates TSyncState and does basic reconciliation. */
-template<typename TBufferTypes>
-struct TReplicator_BasicReconciliar
-{
-	using TInputCmd = typename TBufferTypes::TInputCmd;
-	using TSyncState = typename TBufferTypes::TSyncState;
-	using TAuxState = typename TBufferTypes::TAuxState;
-
-	int32 GetLastSerializedKeyframe() const { return LastSerializedKeyframe; }
-	bool IsReconcileFaultDetected() const { return bReconcileFaultDetected; }
-	const TSimulationTimeKeeper<TBufferTypes>& GetLastSerializedSimulationTimeKeeper() const { return SerializedTimeKeeper; }
-
-	int32 GetProxyDirtyCount(TNetworkSimBufferContainer<TBufferTypes>& Buffers) const
-	{
-		return Buffers.Sync.GetDirtyCount();
-	}
-
-	// --------------------------------------------------------------------
-	//	NetSerialize
-	// --------------------------------------------------------------------
-	void NetSerialize(const FNetSerializeParams& P, TNetworkSimBufferContainer<TBufferTypes>& Buffers, const TSimulationTickInfo<TBufferTypes>& TickInfo)
-	{
-		FArchive& Ar = P.Ar;
-
-		const int32 SerializedHeadKeyframe = FNetworkSimulationSerialization::SerializeKeyframe(Ar, Buffers.Sync.GetHeadKeyframe());
-		TSyncState* SerializedState = nullptr;
-
-		// Serialize total simulation time. This isn't really necessary since we have the keyframe above. 
-		SerializedTimeKeeper = TickInfo.ProcessedSimulationTime;
-		SerializedTimeKeeper.NetSerialize(P);
-
-		if (Ar.IsSaving())
-		{
-			// Server serialize the latest state
-			SerializedState = Buffers.Sync.GetElementFromHead(0);
-		}
-		else
-		{
-			// Its possible we process two packets in between ticks. This isn't really a problem but should be rare.
-			if (bPendingReconciliation)
-			{
-				UE_LOG(LogNetworkSim, Warning, TEXT("bPendingReconciliation while in ::NetSerialize. LastSerializedKeyframe: %d. New SerializedHeadKeyframe: %d."), LastSerializedKeyframe, SerializedHeadKeyframe);
-			}
-
-			// Lazy init the reconciliation buffer. We don't need this on the server/writing side. (Fixme, with templated approach we could easily specialize this to an inline element instead of buffer?)
-			if (ReconciliationBuffer.GetMaxNumElements() == 0)
-			{
-				ReconciliationBuffer.SetBufferSize(1);
-			}
-
-			// Reset the head to match the server. This is not critical, we just need a place to store this, but it makes the code in the reconciliation a bit nicer since the reconcilated frame number = head frame.
-			ReconciliationBuffer.ResetNextHeadKeyframe(SerializedHeadKeyframe);
-			SerializedState = ReconciliationBuffer.GetWriteNext();
-		}
-
-		SerializedState->NetSerialize(Ar);
-
-		if (Ar.IsLoading())
-		{
-			bReconcileFaultDetected = false;
-			bPendingReconciliation = false;
-
-			// The state the client predicted that corresponds to the state the server just serialized to us
-			if (TSyncState* ClientExistingState = Buffers.Sync.FindElementByKeyframe(SerializedHeadKeyframe))
-			{
-				if (ClientExistingState->ShouldReconcile(*SerializedState) || (NetworkSimulationModelCVars::ForceReconcile > 0))
-				{
-					UE_CLOG(!Buffers.Input.IsValidKeyframe(SerializedHeadKeyframe), LogNetworkSim, Error, TEXT("::NetSerialize: Client InputBuffer does not contain data for frame %d. {%s} {%s}"), SerializedHeadKeyframe, *Buffers.Input.GetBasicDebugStr(), *Buffers.Sync.GetBasicDebugStr());
-					bPendingReconciliation =  true;
-				}
-			}
-			else
-			{
-				// We don't have corresponding local state. There are two cases:
-				if (NetworkSimulationModelCVars::EnableLocalPrediction) // Fixme: this is awkward. Expected if player prediction is disabled but coupling like this feels bad.
-				{
-					UE_LOG(LogNetworkSim, Warning, TEXT("::NetSerialize Fault: SyncBuffer does not contain data for frame %d. [%d-%d]"), SerializedHeadKeyframe, Buffers.Sync.GetTailKeyframe(), Buffers.Sync.GetHeadKeyframe());
-				}
-
-				// Case 1: the serialized state is older than what we've kept in our buffer. A bigger buffer would solve this! (at the price of more resimulated frames to recover when this happens)
-				if (SerializedHeadKeyframe < Buffers.Sync.GetTailKeyframe())
-				{
-					// This is a reconcile fault and we just need to chill. We'll stop sending user commands until the cmds in flight flush through the system and we catch back up.
-					bReconcileFaultDetected = true;
-				}
-				else
-				{
-					// Case 2: We've received a newer frame than what we've processed locally. This could happen if we are buffering our inputs locally (but still sending to the server).
-					// Since this doesn't require resimulating, we can just set the authoritative state here and continue on. Ticking logic will need to detect this gap and skip
-					// our LastProcessedInputKeyframe ahead.
-
-					// However, this keyframe better be a valid input. Otherwise how did the server generate it?
-					checkf(SerializedHeadKeyframe <= Buffers.Input.GetHeadKeyframe(), TEXT("Received newer motionstate that doesn't correspond to valid input cmd. SerializedHeadKeyframe: %d. {%s} {%s}"),
-						SerializedHeadKeyframe, *Buffers.Input.GetBasicDebugStr(), *Buffers.Sync.GetBasicDebugStr());
-
-					Buffers.Sync.ResetNextHeadKeyframe(SerializedHeadKeyframe);
-					TSyncState* ClientMotionState = Buffers.Sync.GetWriteNext();
-					*ClientMotionState = *SerializedState;
-				}
-			}
-		}
-
-		LastSerializedKeyframe = SerializedHeadKeyframe;
-	}
-
-	// --------------------------------------------------------------------
-	//	Reconcile
-	// --------------------------------------------------------------------
-	template<typename T, typename TDriver>
-	void Reconcile(TDriver* Driver, TNetworkSimBufferContainer<TBufferTypes>& Buffers, TSimulationTickInfo<TBufferTypes>& TickInfo)
-	{
-		if (bPendingReconciliation == false)
-		{
-			return;
-		}
-		bPendingReconciliation = false;
-		check(ReconciliationBuffer.GetNumValidElements() == 1); // this function assumes only 1 element
-
-		const int32 ReconciliationKeyframe = ReconciliationBuffer.GetHeadKeyframe();
-		TSyncState* ServerState = ReconciliationBuffer.GetElementFromHead(0);
-
-		TInputCmd* ClientInputCmd = Buffers.Input.FindElementByKeyframe( ReconciliationKeyframe );
-		if (ClientInputCmd == nullptr)
-		{
-			// Fault: no historic command for this frame to reconcile with.
-			UE_LOG(LogNetworkSim, Error, TEXT("Client InputBuffer does not contain data for frame %d. {%s} {%s}"), ReconciliationKeyframe, *Buffers.Input.GetBasicDebugStr(), *Buffers.Sync.GetBasicDebugStr());
-			return;
-		}
-
-		TSyncState* ClientSyncState = Buffers.Sync.FindElementByKeyframe( ReconciliationKeyframe );
-
-		// Should not be here if the client doesn't have a valid state. See ::NetSerialize
-		checkf(ClientSyncState != nullptr, TEXT("SyncBuffer does not contain data for frame %d. %s"), ReconciliationKeyframe, *Buffers.Sync.GetBasicDebugStr());
-
-		// -------------------------------------------------------------------------------------------------------------------------
-		// Resimulate
-		// -------------------------------------------------------------------------------------------------------------------------
-
-		ServerState->VisualLog( FVisualLoggingParameters(EVisualLoggingContext::LastConfirmed, ReconciliationKeyframe, EVisualLoggingLifetime::Persistent), Driver, Driver);
-		ClientSyncState->VisualLog( FVisualLoggingParameters(EVisualLoggingContext::FirstMispredicted, ReconciliationKeyframe, EVisualLoggingLifetime::Persistent), Driver, Driver);
-
-		TSyncState* LatestClient = Buffers.Sync.GetElementFromHead(0);
-
-		// Copy authoritative state over client state (FIXME: we may want to store this off historically somewhere? Or will VLog be enough for debugging?)
-		*ClientSyncState = *ServerState;
-
-		// Set the canonical simulation time to what we received (we will advance it as we resimulate)
-		TickInfo.ProcessedSimulationTime = SerializedTimeKeeper;
-		TickInfo.LastProcessedInputKeyframe = ReconciliationKeyframe;
-		
-		// Resimulate up to our latest SyncBuffer frame. Note that we may have unprocessed user commands in the command buffer at this point.
-		const int32 LatestKeyframe = Buffers.Sync.GetHeadKeyframe();		
-		for (int32 Keyframe = ReconciliationKeyframe+1; Keyframe <= LatestKeyframe; ++Keyframe)
-		{
-			// Keyframe is the frame we are resimulating right now.
-			TInputCmd* ResimulateCmd	= Buffers.Input.FindElementByKeyframe(Keyframe);
-			TSyncState* PrevMotionState = Buffers.Sync.FindElementByKeyframe(Keyframe - 1);
-			TSyncState* NextMotionState = Buffers.Sync.FindElementByKeyframe(Keyframe);
-			
-			check(ResimulateCmd);
-			check(PrevMotionState);
-			check(NextMotionState);
-
-			// TEMP (Aux buffer not fully plumbed through the system yet)
-			// ------------------------------------------------------
-			TAuxState TempAuxState;
-			TAuxState* AuxState = Buffers.Aux.FindElementByKeyframe(Keyframe);
-			if (AuxState == nullptr)
-			{
-				AuxState = &TempAuxState;
-			}
-			// ------------------------------------------------------
-
-			// Log out the Mispredicted state that we are about to overwrite.
-			NextMotionState->VisualLog( FVisualLoggingParameters(Keyframe == LatestKeyframe ? EVisualLoggingContext::LastMispredicted : EVisualLoggingContext::OtherMispredicted, Keyframe, EVisualLoggingLifetime::Persistent), Driver, Driver);
-
-			// Do the actual update
-			T::Update(Driver, *ResimulateCmd, *PrevMotionState, *NextMotionState, *AuxState);
-			
-			// Update TickInfo
-			TickInfo.ProcessedSimulationTime.AccumulateTimeFromInputCmd(*ResimulateCmd);
-			TickInfo.LastProcessedInputKeyframe = Keyframe;
-
-			// Log out the newly predicted state that we got.
-			NextMotionState->VisualLog( FVisualLoggingParameters(Keyframe == LatestKeyframe ? EVisualLoggingContext::LastPredicted : EVisualLoggingContext::OtherPredicted, Keyframe, EVisualLoggingLifetime::Persistent), Driver, Driver);
-		}
-	}
-
-private:
-	
-	TReplicationBuffer<TSyncState> ReconciliationBuffer;
-	
-	TSimulationTimeKeeper<TBufferTypes> SerializedTimeKeeper; // last serialized time keeper
-
-	int32 LastSerializedKeyframe = -1;
-	bool bPendingReconciliation = false;	// Reconciliation is pending: we need to reconcile state from the server that differs from the locally predicted state
-	bool bReconcileFaultDetected = false;	// A fault was detected: we received state from the server that we are unable to reconcile with locally predicted state
-};
+#include "NetworkSimulationModelReplicators.h"
 
 // ----------------------------------------------------------------------------------------------------------------------------------------------------------------------
 //	TNetworkedSimulationModel
@@ -412,26 +22,34 @@ private:
 
 template <
 	typename T,
-	typename TBufferTypes,
-	typename TRepProxyServerRPC = TReplicator_DynamicSequence<TBufferTypes, ENetworkSimBufferTypeId::Input>,
-	typename TRepProxyAutonomous = TReplicator_BasicReconciliar<TBufferTypes>,
-	typename TRepProxySimulated = TReplicator_SimulatedExtrapolatedReconciliar<TBufferTypes>,
-	typename TRepProxyReplay = TReplicator_DynamicSequence<TBufferTypes, ENetworkSimBufferTypeId::Sync>,
-	typename TRepProxyDebug = TReplicator_DynamicSequence<TBufferTypes, ENetworkSimBufferTypeId::Debug>
+	typename TUserBufferTypes,
+	typename InTTickSettings=TNetworkSimTickSettings<>,
+	typename TRepProxyServerRPC =	TReplicator_DynamicSequence	<TInternalBufferTypes<TUserBufferTypes, InTTickSettings>,	InTTickSettings, ENetworkSimBufferTypeId::Input>,
+	typename TRepProxyAutonomous =	TReplicator_BasicReconciliar <TInternalBufferTypes<TUserBufferTypes, InTTickSettings>,	InTTickSettings>,
+	typename TRepProxySimulated =	TReplicator_SimulatedExtrapolatedReconciliar <TInternalBufferTypes<TUserBufferTypes, InTTickSettings>, InTTickSettings>,
+	typename TRepProxyReplay =		TReplicator_DynamicSequence	<TInternalBufferTypes<TUserBufferTypes, InTTickSettings>,	InTTickSettings, ENetworkSimBufferTypeId::Sync>,
+	typename TRepProxyDebug =		TReplicator_DynamicSequence	<TInternalBufferTypes<TUserBufferTypes, InTTickSettings>,	InTTickSettings, ENetworkSimBufferTypeId::Debug>
 >
 class TNetworkedSimulationModel : public IReplicationProxy
 {
 public:
 
+	using TBufferTypes = TInternalBufferTypes<TUserBufferTypes, InTTickSettings>;
+	using TTickSettings = InTTickSettings;
+
 	using TInputCmd = typename TBufferTypes::TInputCmd;
 	using TSyncState = typename TBufferTypes::TSyncState;
 	using TAuxState = typename TBufferTypes::TAuxState;
 	using TDebugState = typename TBufferTypes::TDebugState;
 
+	using TSimTime = TNetworkSimTime<TTickSettings>;
+	using TRealTime = typename TTickSettings::TRealTime;
+
 	class IDriver
 	{
 	public:
 		virtual void InitSyncState(TSyncState& OutSyncState) const = 0;	// Called to create initial value of the sync state.
+		virtual void ProduceInput(const TSimTime&, typename TUserBufferTypes::TInputCmd&) = 0; // Called when the sim is ready to process new local input
 		virtual void FinalizeFrame(const TSyncState& SyncState) = 0; // Called from the Network Sim at the end of the sim frame when there is new sync data.
 	};
 
@@ -452,52 +70,51 @@ public:
 			DebugState->ProcessedKeyframes.Reset();
 		}
 
-		// How many commands we are allowed to process right now. This will need to be built out a bit more to handle shifting network conditions etc. For now, just a cvar.
-		int32 AllowCmds = NetworkSimulationModelCVars::MaxInputCmdsFrame;
+		// --------------------------------------------------------------------------------------------------------------------------
+		//	Reconcile
+		//	This will eventually be called outside the Tick loop, only after processing a network bunch
+		//	Reconcile is about "making things right" after a network update. We are not processing "more" simulation yet.
+		// --------------------------------------------------------------------------------------------------------------------------
+		switch (Parameters.Role)
+		{
+			case ROLE_Authority:
+				RepProxy_ServerRPC.template Reconcile<T, TDriver>(Driver, Buffers, TickInfo);
+			break;
 
-		switch (Parameters.Role) {
+			case ROLE_AutonomousProxy:
+				RepProxy_Autonomous.template Reconcile<T, TDriver>(Driver, Buffers, TickInfo);
+			break;
 
-		// -------------------------------------------------------------------------------------------------------------------------------------------------
-		//														Reconciliation
-		// -------------------------------------------------------------------------------------------------------------------------------------------------
-		case ROLE_AutonomousProxy:
+			case ROLE_SimulatedProxy:
+				RepProxy_Simulated.template Reconcile<T, TDriver>(Driver, Buffers, TickInfo);
+			break;
+		}
+
+		// ----------------------------------------------------------------------------------------------------------------
+		//	PreSimTick
+		//	This is the beginning of a new frame. PreSimTick will decide if we should take Parameters.LocalDeltaTimeSeconds
+		//	and advance the simulation or not. It will also generate new local input if necessary.
+		// ----------------------------------------------------------------------------------------------------------------
+		switch (Parameters.Role)
 		{
-			// Client: don't allow buffering and ensure continuous stream of inputs. This will need to be built out a bit more. The main points are:
-			// -Don't buffer input locally (unless we explicitly opt into it as a lag hiding measurement). Don't want to accidentally introduce application latency.
-			// -Don't send server input cmds that we didn't actually process. Current Input buffer replication is rather hard coded to send continuous frames and doesn't handle gaps.
-
-			// Recovering from this is possible but messy. We could immediately set things up so that we process the latest user cmd.
-			// The tricky part is making sure we the server does the same thing. We would need to set the tail of the input cmd stream so
-			// that we don't replicate any cmds that were skipped. Then server side it would have to detect this and know to jump ahead.
-			// At this point in development, the complexity this adds to the overall system is not worth the benefit. The better way to
-			// fix this for now is to closely guard what we allow into the client's InputCmdBuffer
-
-			// You can be at head (no new cmds to process) or right before the head (1 new cmd to process)
-			checkf(Buffers.Input.GetHeadKeyframe() - TickInfo.LastProcessedInputKeyframe <= 1, TEXT("Client Input Processing is out of sync. LastProcessedInputKeyframe: %d. InputBuffer: [%d - %d]"), 
-				TickInfo.LastProcessedInputKeyframe, Buffers.Input.GetTailKeyframe(), Buffers.Input.GetHeadKeyframe());
-
-			// We should be processing input from this GFrameNumber. If we aren't, then our DeltaTime is probably off (input cmd's delta time wont match the actual delta time of the frame we sim on which will cause artifacts).
-			// We also want to avoid accidentally introducing input latency as a matter of principle (input latency bad)
-			// We *can* support local input buffering in two ways (neither currently implemented):
-			//	-Fixed time steps (delta times will be constant so buffering wont matter)
-			//	-Variable time steps with "interpolation layer": we will need to be able to absorb fluctuations in frame rate so there will be render frames that doesn't process any new input and
-			//		render frames that have to process multiple inputs. (E.g, the goal would be to always "see" an update every render frame, even if we can't actually tick a new sim frame every render frame. Interpolation would be needed to accomplish this)
-			// If either of the above ideas are implemented, it will be done as an optional/templated parameter to the network sim class.
-			checkf(TickInfo.LastLocalInputGFrameNumber == 0 || GFrameNumber == TickInfo.LastLocalInputGFrameNumber, TEXT("TNetworkedSimulationModel running on stale input. You may have an ordering issue between TNetworkedSimulationModel::GetNextInputForWrite and TNetworkedSimulationModel::Tick. "
-				"GFrameNumber: %d LastLocalInputGFrameNumber: %d"), GFrameNumber, TickInfo.LastLocalInputGFrameNumber);
+			case ROLE_Authority:
+				RepProxy_ServerRPC.template PreSimTick<T, TDriver>(Driver, Buffers, TickInfo, Parameters.LocalDeltaTimeSeconds);
+			break;
 
-		
-			// Note that is is important for the NetworkSimModel to control *when* the reconcile happens. The replicator objects themselves should never do
-			// reconciliation inside a NetSerialze call. Tick may not be the permanent place for this, but it can't be in NetSerialize.
-			RepProxy_Autonomous.template Reconcile<T, TDriver>(Driver, Buffers, TickInfo);
+			case ROLE_AutonomousProxy:
+				RepProxy_Autonomous.template PreSimTick<T, TDriver>(Driver, Buffers, TickInfo, Parameters.LocalDeltaTimeSeconds);
+			break;
 
-			// Don't process new user commands if we are dealing with a reconciliation fault. That means we weren't able to correlate our predicted position with what the
-			// server last sent us. If we are hitting this, it means we are waiting for the in flight commands to flush through the system. The proxy rep object will handle recovering.
-			if (RepProxy_Autonomous.IsReconcileFaultDetected())
-			{
-				AllowCmds = 0;
-			}
+			case ROLE_SimulatedProxy:
+				RepProxy_Simulated.template PreSimTick<T, TDriver>(Driver, Buffers, TickInfo, Parameters.LocalDeltaTimeSeconds);
+			break;
+		}		
 
+		switch (Parameters.Role) {
+
+		
+		case ROLE_AutonomousProxy:
+		{
 			// Check SyncBuffer being ahead of processed Keyframes. This would happen in cases where we are either not predicting or are buffering our input locally
 			// while sending latest cmds to the server. Essentially, we got the authoritative motion state from the server before we ran the simulation locally.
 			if (Buffers.Sync.GetHeadKeyframe() > TickInfo.LastProcessedInputKeyframe)
@@ -510,13 +127,6 @@ public:
 
 				TickInfo.LastProcessedInputKeyframe = Buffers.Sync.GetHeadKeyframe();
 			}
-
-			if (NetworkSimulationModelCVars::EnableLocalPrediction == 0)
-			{
-				AllowCmds = 0; // Don't process any commands this frame
-				TickInfo.LastProcessedInputKeyframe = Buffers.Input.GetHeadKeyframe(); // Increment so we can accept a new command next frame
-			}
-
 			break;
 		}
 
@@ -534,79 +144,91 @@ public:
 		case ROLE_SimulatedProxy:
 		{
 			RepProxy_Simulated.template Reconcile<T, TDriver>(Driver, Buffers, TickInfo);
-
-			// Extrapolation: this still needs work  (needs to be a setting/state somewhere and we need to better handle the relationship with NetSerialize and Reconcile)
-			if (TInputCmd* SynthesizedCmd = GetNextInputForWrite(Parameters.LocalDeltaTimeSeconds))
-			{
-
-			}
 			break;
 		}
 
 		} // end switch
 
+		// -------------------------------------------------------------------------------------------------------------------------------------------------
+		//	PreInput processing
+		//	We are just about to start mechanically processing input. Lets just make sure we have a valid sync state to start with
+		// -------------------------------------------------------------------------------------------------------------------------------------------------
+
+		// The SyncedState buffer needs to be in sync here:
+		//	-We want it to have a SyncedState, but it may not on the first frame through (thats ok).
+		//  -Its HeadKeyframe should be one behind the Keyframe we are about to process.
+		//
+		// Note, InputCmds start @ Keyframe=1. The first SyncedState that Update produces will go in KeyFrame=1.
+		// (E.g, InputCmd @ keyframe=X is used to generate MotionState @ keyframe=X)
+		// This means that SyncedState @ keyframe=0 is always created here via InitSyncState.
+		// This also means that we never actually process InputCmd @ keyframe=0. Which is why LastProcessedInputKeyframe is initialized to 0 ("already processed")
+		// and the buffer has an empty element inserted in InitLocalInputBuffer.
+
+		if (Buffers.Sync.GetHeadKeyframe() != TickInfo.LastProcessedInputKeyframe)
+		{
+			if (TickInfo.LastProcessedInputKeyframe != 0)
+			{
+				// This shouldn't happen, but is not fatal. We are reseting the sync state buffer.
+				UE_LOG(LogNetworkSim, Warning, TEXT("Break in SyncState continuity. LastProcessedInputKeyframe: %d. SyncBuffer.GetHeadKeyframe(): %d. Role=%d"), TickInfo.LastProcessedInputKeyframe, Buffers.Sync.GetHeadKeyframe(), (int32)Parameters.Role);
+			}
+
+			// We need an initial/current state. Get this from the sim driver
+			Buffers.Sync.ResetNextHeadKeyframe(TickInfo.LastProcessedInputKeyframe);
+			TSyncState* StartingState = Buffers.Sync.GetWriteNext();
+			Driver->InitSyncState(*StartingState);
+		}
+
 		// -------------------------------------------------------------------------------------------------------------------------------------------------
 		//														Input Processing
 		// -------------------------------------------------------------------------------------------------------------------------------------------------
 
-		int32 NumProcessed = 0;
-		while (AllowCmds-- > 0)
+		while(true)
 		{
 			const int32 Keyframe = TickInfo.LastProcessedInputKeyframe+1;
+			if (Keyframe > TickInfo.MaxAllowedInputKeyframe)
+			{
+				break;
+			}
+
 			if (TInputCmd* NextCmd = Buffers.Input.FindElementByKeyframe(Keyframe))
 			{
-				// The SyncedState buffer needs to be in sync here:
-				//	-We want it to have a SyncedState, but it may not on the first frame through (thats ok).
-				//  -Its HeadKeyframe should be one behind the Keyframe we are about to process.
-				//
-				// Note, InputCmds start @ Keyframe=1. The first SyncedState that Update produces will go in KeyFrame=1.
-				// (E.g, InputCmd @ keyframe=X is used to generate MotionState @ keyframe=X)
-				// This means that SyncedState @ keyframe=0 is always created here via InitSyncState.
-				// This also means that we never actually process InputCmd @ keyframe=0. Which is why LastProcessedInputKeyframe is initialized to 0 ("already processed")
-				// and the buffer has an empty element inserted in InitLocalInputBuffer.
-
-				if (Buffers.Sync.GetHeadKeyframe() != TickInfo.LastProcessedInputKeyframe)
+				// We have an unprocessed command, do we have enough allotted simulation time to process it?
+				if (TickInfo.GetRemaningAllowedSimulationTime() >= NextCmd->GetFrameDeltaTime())
 				{
-					if (TickInfo.LastProcessedInputKeyframe != 0)
+					// -------------------------------------------------------------------------------------------------
+					//	The core process input command and call ::Update block!
+					// -------------------------------------------------------------------------------------------------
+					TSyncState* PrevSyncState = Buffers.Sync.FindElementByKeyframe(TickInfo.LastProcessedInputKeyframe);
+					TSyncState* NextSyncState = Buffers.Sync.GetWriteNext();
+
+					check(PrevSyncState != nullptr);
+					check(NextSyncState != nullptr);
+					check(Buffers.Sync.GetHeadKeyframe() == Keyframe);
+				
+					if (DebugState)
 					{
-						// This shouldn't happen, but is not fatal. We are reseting the sync state buffer.
-						UE_LOG(LogNetworkSim, Warning, TEXT("Break in SyncState continuity. LastProcessedInputKeyframe: %d. SyncBuffer.GetHeadKeyframe(): %d. Role=%d"), TickInfo.LastProcessedInputKeyframe, Buffers.Sync.GetHeadKeyframe(), (int32)Parameters.Role);
+						DebugState->ProcessedKeyframes.Add(Keyframe);
 					}
+					
+					TAuxState AuxState; // Temp: aux buffer not implemented yet
 
-					// We need an initial/current state. Get this from the sim driver
-					Buffers.Sync.ResetNextHeadKeyframe(TickInfo.LastProcessedInputKeyframe);
-					TSyncState* StartingState = Buffers.Sync.GetWriteNext();
-					Driver->InitSyncState(*StartingState);
+					T::Update(Driver, NextCmd->GetFrameDeltaTime(), *NextCmd, *PrevSyncState, *NextSyncState, AuxState);
 					
+					TickInfo.TotalProcessedSimulationTime += NextCmd->GetFrameDeltaTime();
+					TickInfo.LastProcessedInputKeyframe = Keyframe;
 				}
-
-				TSyncState* PrevSyncState = Buffers.Sync.FindElementByKeyframe(TickInfo.LastProcessedInputKeyframe);
-				TSyncState* NextSyncState = Buffers.Sync.GetWriteNext();
-
-				check(PrevSyncState != nullptr);
-				check(NextSyncState != nullptr);
-				check(Buffers.Sync.GetHeadKeyframe() == Keyframe);
-				
-				if (DebugState)
+				else
 				{
-					DebugState->ProcessedKeyframes.Add(Keyframe);
+					break;
 				}
-
-				TAuxState AuxState; // Temp: aux buffer not implemented yet
-
-				T::Update(Driver, *NextCmd, *PrevSyncState, *NextSyncState, AuxState);
-				TickInfo.ProcessedSimulationTime.AccumulateTimeFromInputCmd(*NextCmd);
-				NumProcessed++;
-
-				TickInfo.LastProcessedInputKeyframe = Keyframe;
 			}
 			else
 			{
-				// Not a warning in itself anymore. Though the concept of "haven't processed cmds in x MS" is something we want to define.
 				break;
 			}
 		}
 
+
 		// FIXME: this needs to be sorted out. We really want to check if there is new sync state and then call this here.
 		// Call into the driver to sync to the latest state if we processed any input
 		//if (NumProcessed > 0)
@@ -667,7 +289,7 @@ public:
 	{
 		// You are not allowed to change the simulations TickInfo while NetSerializing. Since all rep proxies are templated
 		// and there is not interface/base class, we are enforcing the constness at the call site.
-		const TSimulationTickInfo<TBufferTypes>& ConstTickInfo = const_cast<const TSimulationTickInfo<TBufferTypes>&>(TickInfo);
+		const TSimulationTickState<TTickSettings>& ConstTickInfo = const_cast<const TSimulationTickState<TTickSettings>&>(TickInfo);
 
 		switch(Target)
 		{
@@ -714,16 +336,6 @@ public:
 			return 0;
 		};
 	}
-	
-	TInputCmd* GetNextInputForWrite(float DeltaTime)
-	{
-		if (TInputCmd* Next = TickInfo.GetNextInputForWrite(Buffers))
-		{
-			Next->FrameDeltaTime = DeltaTime;
-			return Next;
-		}
-		return nullptr;
-	}
 
 	void InitLocalInputBuffer()
 	{
@@ -735,7 +347,7 @@ public:
 		TickInfo.LastLocalInputGFrameNumber = 0;
 	}	
 	
-	TSimulationTickInfo<TBufferTypes> TickInfo;	// Manages simulation time and what inputs we are processed
+	TSimulationTickState<TTickSettings> TickInfo;	// Manages simulation time and what inputs we are processed
 
 	TNetworkSimBufferContainer<TBufferTypes> Buffers;
 
diff --git a/Engine/Plugins/Runtime/NetworkPrediction/Source/NetworkPrediction/Public/NetworkSimulationModelDebugger.h b/Engine/Plugins/Runtime/NetworkPrediction/Source/NetworkPrediction/Public/NetworkSimulationModelDebugger.h
--- a/Engine/Plugins/Runtime/NetworkPrediction/Source/NetworkPrediction/Public/NetworkSimulationModelDebugger.h
+++ b/Engine/Plugins/Runtime/NetworkPrediction/Source/NetworkPrediction/Public/NetworkSimulationModelDebugger.h
@@ -329,6 +329,8 @@ private:
 template <typename TNetSimModel, typename TDriver>
 struct TNetworkSimulationModelDebugger : public INetworkSimulationModelDebugger
 {
+	using TSimTime = typename TNetSimModel::TSimTime;
+
 	TNetworkSimulationModelDebugger(TNetSimModel* InNetSim, TDriver* InDriver, AActor* OwningActor, FString InDebugName)
 	{
 		NetworkSim = InNetSim;
@@ -365,16 +367,17 @@ struct TNetworkSimulationModelDebugger : public INetworkSimulationModelDebugger
 			const int32 LastSerializedKeyframe = NetworkSim->RepProxy_Autonomous.GetLastSerializedKeyframe();
 
 			// Calc how much predicted time we have processed. Note that we use the motionstate buffer to iterate but the MS is on the input cmd. (if we are buffering cmds, don't want to count them)
-			float PredictedMS = 0.f;
+			
+			TSimTime PredictedMS;
 			for (int32 PredKeyrame = LastSerializedKeyframe+1; PredKeyrame <= NetworkSim->Buffers.Sync.GetHeadKeyframe(); ++PredKeyrame)
 			{
 				if (auto* Cmd = NetworkSim->Buffers.Input.FindElementByKeyframe(PredKeyrame))
 				{
-					PredictedMS += Cmd->FrameDeltaTime * 1000.f;
+					PredictedMS += Cmd->GetFrameDeltaTime();
 				}
 			}
 
-			FString ConfirmedFrameStr = FString::Printf(TEXT("LastConfirmedFrame: %d. Prediction: %d Frames, %.2f MS"), LastSerializedKeyframe, NetworkSim->Buffers.Sync.GetHeadKeyframe() - LastSerializedKeyframe, PredictedMS);
+			FString ConfirmedFrameStr = FString::Printf(TEXT("LastConfirmedFrame: %d. Prediction: %d Frames, %s MS"), LastSerializedKeyframe, NetworkSim->Buffers.Sync.GetHeadKeyframe() - LastSerializedKeyframe, *PredictedMS.ToString());
 			if (FaultDetected)
 			{
 				ConfirmedFrameStr += TEXT(" RECONCILE FAULT DETECTED!");
@@ -383,8 +386,8 @@ struct TNetworkSimulationModelDebugger : public INetworkSimulationModelDebugger
 
 			Out.Emit(*ConfirmedFrameStr, Color);
 
-			FString SimulationTimeString = FString::Printf(TEXT("Local SimulationTime: %s. SerialisedSimulationTime: %s. Difference MS: %s"), *NetworkSim->TickInfo.ProcessedSimulationTime.ToString(),
-				*NetworkSim->RepProxy_Autonomous.GetLastSerializedSimulationTimeKeeper().ToString(), *(NetworkSim->TickInfo.ProcessedSimulationTime - NetworkSim->RepProxy_Autonomous.GetLastSerializedSimulationTimeKeeper()).ToString());
+			FString SimulationTimeString = FString::Printf(TEXT("Local SimulationTime: %s. SerialisedSimulationTime: %s. Difference MS: %s"), *NetworkSim->TickInfo.TotalProcessedSimulationTime.ToString(),
+				*NetworkSim->RepProxy_Autonomous.GetLastSerializedSimTime().ToString(), *(NetworkSim->TickInfo.TotalProcessedSimulationTime - NetworkSim->RepProxy_Autonomous.GetLastSerializedSimTime()).ToString());
 			Out.Emit(*SimulationTimeString, Color);
 			
 		}
@@ -433,10 +436,10 @@ struct TNetworkSimulationModelDebugger : public INetworkSimulationModelDebugger
 			float MinHeight = TextSizing.DrawYL;
 			static float MinHeightMS = 1/60.f * 1000.f;	// 60hz frame is drawn at MinHeight.
 
-			auto CalcHeight = [&](float MS)
+			auto CalcHeight = [&](TSimTime MS)
 			{
 				const float Ratio = MinHeight / MinHeightMS;
-				return MS * Ratio;
+				return MS.ToRealTimeSeconds() * Ratio;
 			};
 
 			for (auto It = DebugBuffer->CreateIterator(); It; ++It)
@@ -452,7 +455,7 @@ struct TNetworkSimulationModelDebugger : public INetworkSimulationModelDebugger
 				Out.EmitQuad(FVector2D( ScreenX, ScreenY), FVector2D( ScreenWidth, ServerHeight), Color);
 				Out.EmitText(FVector2D( ScreenX, ScreenY), FColor::Black, FString::Printf(TEXT("%.2f"), (DebugState->LocalDeltaTimeSeconds * 1000.f)));
 
-				float ClientSimTime = 0.f;
+				TSimTime ClientSimTime;
 				float ClientX = ScreenX;
 				float ClientY = ScreenY - ClientOffsetY;
 				for (int32 Keyframe : DebugState->ProcessedKeyframes)
@@ -460,10 +463,10 @@ struct TNetworkSimulationModelDebugger : public INetworkSimulationModelDebugger
 					auto* Cmd = InputBuffer.FindElementByKeyframe(Keyframe);
 					if (Cmd)
 					{
-						ClientSimTime += Cmd->FrameDeltaTime * 1000.f;
+						ClientSimTime += Cmd->GetFrameDeltaTime();
 						
 						float ClientSizeX = MinWidth;
-						float ClientSizeY =  CalcHeight(Cmd->FrameDeltaTime * 1000.f);
+						float ClientSizeY =  CalcHeight(Cmd->GetFrameDeltaTime());
 
 						FVector2D ScreenPos(ClientX, ClientY - ClientSizeY);
 						Out.EmitQuad(ScreenPos, FVector2D( ClientSizeX, ClientSizeY), FColor::Blue);
@@ -477,7 +480,7 @@ struct TNetworkSimulationModelDebugger : public INetworkSimulationModelDebugger
 					if (auto* Cmd = InputBuffer.FindElementByKeyframe(Keyframe))
 					{
 						float ClientSizeX = MinWidth;
-						float ClientSizeY = CalcHeight(Cmd->FrameDeltaTime * 1000.f);
+						float ClientSizeY = CalcHeight(Cmd->GetFrameDeltaTime());
 						FVector2D ScreenPos(ClientX, ClientY - ClientSizeY);
 						Out.EmitQuad(ScreenPos, FVector2D( ClientSizeX, ClientSizeY), FColor::Red);
 						Out.EmitText(ScreenPos, FColor::White, LexToString(Keyframe));
diff --git a/Engine/Plugins/Runtime/NetworkPrediction/Source/NetworkPrediction/Public/NetworkSimulationModelReplicators.h b/Engine/Plugins/Runtime/NetworkPrediction/Source/NetworkPrediction/Public/NetworkSimulationModelReplicators.h
new file mode 100644
--- /dev/null
+++ b/Engine/Plugins/Runtime/NetworkPrediction/Source/NetworkPrediction/Public/NetworkSimulationModelReplicators.h
@@ -0,0 +1,467 @@
+// Copyright 1998-2019 Epic Games, Inc. All Rights Reserved.
+
+#pragma once
+
+// ----------------------------------------------------------------------------------------------------------------------------------------------------------------------
+//	CVars and compile time constants
+// ----------------------------------------------------------------------------------------------------------------------------------------------------------------------
+
+namespace NetworkSimulationModelCVars
+{
+static int32 EnableLocalPrediction = 1;
+static FAutoConsoleVariableRef CVarEnableLocalPrediction(TEXT("ns.EnableLocalPrediction"),
+	EnableLocalPrediction, TEXT("Toggle local prediction. Meant for debugging."), ECVF_Default);
+
+static int32 ForceReconcile = 0;
+static FAutoConsoleVariableRef CVarForceReconcile(TEXT("ns.ForceReconcile"),
+	ForceReconcile, TEXT("Forces reconcile even if state does not differ. E.g, force resimulation after every netupdate."), ECVF_Default);
+}
+
+#ifndef NETSIM_NETCONSTANT_NUM_BITS_KEYFRAME
+	#define NETSIM_NETCONSTANT_NUM_BITS_KEYFRAME 8	// Allows you to override this setting via UBT, but access via cleaner FActorMotionNetworkingConstants::NUM_BITS_KEYFRAME
+#endif
+
+struct FNetworkSimulationSerialization
+{
+	// How many bits we use to encode the key frame number for buffers.
+	// Client Keyframes are stored locally as 32 bit integers, but we use a smaller # of bits to NetSerialize.
+	// Frames are only relatively relevant: the absolute value doesn't really matter. We just need to detect newer/older.
+	enum { NUM_BITS_KEYFRAME = NETSIM_NETCONSTANT_NUM_BITS_KEYFRAME };		
+
+	// Abs max value we encode into the bit writer
+	enum { MAX_KEYFRAME_WRITE = 1 << NUM_BITS_KEYFRAME };
+
+	// This is the threshold at which we would wrap around and incorrectly assign a keyframe on the receiving side.
+	// E.g, If there are KEYFRAME_ERROR_THRESHOLD keyframes that do not make it across from sender->receiver, the
+	// receiver will have incorrect local values. With 8 bits, this works out to be 128 frames or about 2 seconds at 60fps.
+	enum { KEYFRAME_ERROR_THRESHOLD = MAX_KEYFRAME_WRITE / 2};
+
+	// Helper to serialize the int32 HeadKeyframe. Returns the unpacked value (this will be same as input in the save path)
+	static int32 SerializeKeyframe(FArchive& Ar, int32 LocalHeadKeyframe)
+	{
+		if (Ar.IsSaving())
+		{
+			((FNetBitWriter&)Ar).WriteIntWrapped( LocalHeadKeyframe, MAX_KEYFRAME_WRITE );
+			return LocalHeadKeyframe;
+		}
+		
+		return MakeRelative(((FNetBitReader&)Ar).ReadInt( MAX_KEYFRAME_WRITE ), LocalHeadKeyframe, MAX_KEYFRAME_WRITE );
+	}
+};
+
+// ----------------------------------------------------------------------------------------------------------------------------------------------------------------------
+//	Templated Replicators/Reconciliars
+// ----------------------------------------------------------------------------------------------------------------------------------------------------------------------
+
+template<typename TBufferTypes, typename TTickSettings>
+struct TReplicatorBase
+{
+	template<typename T, typename TDriver>
+	void Reconcile(TDriver* Driver, TNetworkSimBufferContainer<TBufferTypes>& Buffers, TSimulationTickState<TTickSettings>& TickInfo) { }
+
+	template<typename T, typename TDriver>
+	void PreSimTick(TDriver* Driver, TNetworkSimBufferContainer<TBufferTypes>& Buffers, TSimulationTickState<TTickSettings>& TickInfo, typename TTickSettings::TRealTime RealDeltaTimeSeconds) { }
+};
+
+/** Replicates only the latest element. Does not synchronize keyframe */
+template<typename TBufferTypes, typename TTickSettings, ENetworkSimBufferTypeId BufferId, typename TBase=TReplicatorBase<TBufferTypes, TTickSettings>>
+struct TReplicator_LatestOnly : public TBase
+{
+	using TState = typename TBufferTypes::template select_type<BufferId>::type;
+
+	int32 GetProxyDirtyCount(TNetworkSimBufferContainer<TBufferTypes>& Buffers) const
+	{ 
+		return Buffers.template Get<BufferId>().GetDirtyCount(); 
+	}
+
+	void NetSerialize(const FNetSerializeParams& P, TNetworkSimBufferContainer<TBufferTypes>& Buffers, const TSimulationTickState<TTickSettings>& TickInfo)
+	{
+		FArchive& Ar = P.Ar;
+
+		auto& Buffer = Buffers.template Get<BufferId>();
+		TState* State = nullptr;
+		
+		if (Ar.IsSaving())
+		{
+			State = Buffer.GetElementFromHead(0);
+			check(State); // We should not be here if the buffer is empty. Want to avoid serializing an "empty" flag at the top here.
+		}
+		else
+		{
+			State = Buffer.GetWriteNext();
+		}
+
+		State->NetSerialize(Ar);
+	}
+};
+
+template<typename TBufferTypes, typename TTickSettings, typename TBase=TReplicatorBase<TBufferTypes, TTickSettings>, bool IsFixed=TTickSettings::FixedStepMS==0>
+struct TReplicator_SimulatedExtrapolatedReconciliar : public TBase
+{
+	using TInputCmd = typename TBufferTypes::TInputCmd;
+	using TSyncState = typename TBufferTypes::TSyncState;
+	
+	int32 GetProxyDirtyCount(TNetworkSimBufferContainer<TBufferTypes>& Buffers) const
+	{
+		return Buffers.Sync.GetDirtyCount(); 
+	}
+
+	void NetSerialize(const FNetSerializeParams& P, TNetworkSimBufferContainer<TBufferTypes>& Buffers, const TSimulationTickState<TTickSettings>& TickInfo)
+	{
+		FArchive& Ar = P.Ar;
+		TSyncState* State = nullptr;
+		if (Ar.IsSaving())
+		{
+			State = Buffers.Sync.GetElementFromHead(0);
+			check(State); // We should not be here if the buffer is empty. Want to avoid serializing an "empty" flag at the top here.
+			SerializedTime = TickInfo.TotalProcessedSimulationTime;
+		}
+		else
+		{
+			State = Buffers.Sync.GetWriteNext();
+			bPendingReconciliation = true;
+		}
+
+		SerializedTime.NetSerialize(Ar);
+		State->NetSerialize(Ar);
+	}
+
+	template<typename T, typename TDriver>
+	void Reconcile(TDriver* Driver, TNetworkSimBufferContainer<TBufferTypes>& Buffers, TSimulationTickState<TTickSettings>& TickInfo)
+	{
+		if (!bPendingReconciliation)
+			return;
+
+		bPendingReconciliation = false;
+
+		// This is still being worked out. We are almost there but not quite.
+		// Processing input commands is pretty rigid right now. Since we just added a new SyncState keyframe in ::NetSerialize,
+		// we need to insert a fake input command so the extrapolation code in the NetSimModel can continue to work.
+		// This can be improved.
+
+		{
+			// Generate a new, fake, command since we just added a new sync state to head (via GetWriteNext)
+			TickInfo.GetNextInputForWrite(Buffers);
+			// Gotta fake this too
+			TickInfo.LastLocalInputGFrameNumber--; 
+			// Set our LastProcessedInputKeyframe to fake that we handled it
+			TickInfo.LastProcessedInputKeyframe = Buffers.Sync.GetHeadKeyframe();
+		}
+	}
+
+	template<typename T, typename TDriver>
+	void PreSimTick(TDriver* Driver, TNetworkSimBufferContainer<TBufferTypes>& Buffers, TSimulationTickState<TTickSettings>& TickInfo, typename TTickSettings::TRealTime RealDeltaTimeSeconds)
+	{
+		// Extrapolation: still needs to be a setting/state somewhere that we can change from actor code
+		if (TInputCmd* SynthesizedCmd = TickInfo.GetNextInputForWrite(Buffers))
+		{
+			// We will create a single cmd to simulate the remaining time
+			auto TimeToSimulate = TickInfo.GetRemaningAllowedSimulationTime();
+			//SynthesizedCmd->FrameDeltaTime;
+		}
+	}
+
+private:
+	
+	TNetworkSimTime<TTickSettings> SerializedTime;
+	bool bPendingReconciliation = false;
+};
+
+template<typename TBufferTypes, typename TTickSettings, typename TBase>
+struct TReplicator_SimulatedExtrapolatedReconciliar<TBufferTypes, TTickSettings, TBase, true> : TReplicator_SimulatedExtrapolatedReconciliar<TBufferTypes, TTickSettings, TBase, false>
+{
+	using TInputCmd = typename TBufferTypes::TInputCmd;
+	using TSyncState = typename TBufferTypes::TSyncState;
+
+	template<typename T, typename TDriver>
+	void PreSimTick(TDriver* Driver, TNetworkSimBufferContainer<TBufferTypes>& Buffers, TSimulationTickState<TTickSettings>& TickInfo, typename TTickSettings::TRealTime RealDeltaTimeSeconds)
+	{
+		// Extrapolation: still needs to be a setting/state somewhere that we can change from actor code
+		if (TInputCmd* SynthesizedCmd = TickInfo.GetNextInputForWrite(Buffers))
+		{
+			// We will create a single cmd to simulate the remaining time
+			auto TimeToSimulate = TickInfo.GetRemaningAllowedSimulationTime();
+			SynthesizedCmd->GetFrameDeltaTime();
+		}
+	}
+};
+
+/** Replicates the latest sequence of N elements. N is dynamic, not compiled in. */
+template<typename TBufferTypes, typename TTickSettings, ENetworkSimBufferTypeId BufferId, typename TBase=TReplicatorBase<TBufferTypes, TTickSettings>>
+struct TReplicator_DynamicSequence : public TBase
+{
+	int32 GetProxyDirtyCount(TNetworkSimBufferContainer<TBufferTypes>& Buffers) const 
+	{
+		return Buffers.template Get<BufferId>().GetDirtyCount(); 
+	}
+	
+	void NetSerialize(const FNetSerializeParams& P, TNetworkSimBufferContainer<TBufferTypes>& Buffers, const TSimulationTickState<TTickSettings>& TickInfo)
+	{
+		auto& Buffer = Buffers.template Get<BufferId>();
+		FArchive& Ar = P.Ar;
+		uint8 SerializedNumElements = FMath::Min<uint8>(MaxNumElements, Buffer.GetNumValidElements());
+		Ar << SerializedNumElements;
+
+		const int32 HeadKeyframe = FNetworkSimulationSerialization::SerializeKeyframe(Ar, Buffer.GetHeadKeyframe());
+		const int32 StartingKeyframe = FMath::Max(0, HeadKeyframe - SerializedNumElements + 1);
+
+		if (Ar.IsLoading())
+		{
+			const int32 PrevHead = Buffer.GetHeadKeyframe();
+			if (PrevHead < StartingKeyframe && PrevHead >= 0)
+			{
+				// There is a gap in the stream. In some cases, we want this to be a "fault" and bubble up. We may want to synthesize state or maybe we just skip ahead.
+				UE_LOG(LogNetworkSim, Warning, TEXT("Fault: gap in received buffer. PrevHead: %d. Received: %d-%d. Reseting previous buffer contents"), PrevHead, StartingKeyframe, HeadKeyframe);
+			}
+
+			Buffer.ResetNextHeadKeyframe(StartingKeyframe);
+		}
+
+		for (int32 Keyframe = StartingKeyframe; Keyframe <= HeadKeyframe; ++Keyframe)
+		{
+			// This, as is, is bad. The intention is that these functions serialize multiple items in some delta compressed fashion.
+			// As is, we are just serializing the elements individually.
+			auto* Cmd = Ar.IsLoading() ? Buffer.GetWriteNext() : Buffer.FindElementByKeyframe(Keyframe);
+			Cmd->NetSerialize(P);
+		}
+	}
+
+private:
+	
+	uint8 MaxNumElements = 3;
+};
+
+/** Replicates TSyncState and does basic reconciliation. */
+template<typename TBufferTypes, typename TTickSettings, typename TBase=TReplicatorBase<TBufferTypes, TTickSettings>>
+struct TReplicator_BasicReconciliar : public TBase
+{
+	using TInputCmd = typename TBufferTypes::TInputCmd;
+	using TSyncState = typename TBufferTypes::TSyncState;
+	using TAuxState = typename TBufferTypes::TAuxState;
+
+	int32 GetLastSerializedKeyframe() const { return LastSerializedKeyframe; }
+	bool IsReconcileFaultDetected() const { return bReconcileFaultDetected; }
+	const TNetworkSimTime<TTickSettings>& GetLastSerializedSimTime() const { return SerializedTime; }
+
+	int32 GetProxyDirtyCount(TNetworkSimBufferContainer<TBufferTypes>& Buffers) const
+	{
+		return Buffers.Sync.GetDirtyCount();
+	}
+
+	// --------------------------------------------------------------------
+	//	PreSimTick
+	// --------------------------------------------------------------------
+	template<typename T, typename TDriver>
+	void PreSimTick(TDriver* Driver, TNetworkSimBufferContainer<TBufferTypes>& Buffers, TSimulationTickState<TTickSettings>& TickInfo, typename TTickSettings::TRealTime RealDeltaTimeSeconds)
+	{
+		// If we have a reconcile fault, we cannot continue on with the simulation until it clears itself out. This effectively drops the input time and does not sample new inputs
+		if (bReconcileFaultDetected)
+		{
+			return;
+		}
+
+		// Accumulate real time into the simulation. This will results in more Allowed Simulation Time (or will accumulate internally until we are ready for the next tick)
+		TickInfo.GiveSimulationTime(RealDeltaTimeSeconds);
+
+		TNetworkSimTime<TTickSettings> DeltaSimTime = TickInfo.GetRemaningAllowedSimulationTime();
+		if (DeltaSimTime.ToRealTimeSeconds() > 0)
+		{
+			// Get next input in buffer, clear it, set delta time, then send it to driver to sample local input
+			if (auto* InputCmd = Buffers.Input.GetWriteNext())
+			{
+				*InputCmd = typename TBufferTypes::TInputCmd();
+				InputCmd->SetFrameDeltaTime(DeltaSimTime);
+				Driver->ProduceInput(DeltaSimTime, *InputCmd);
+			}
+
+			// If we are predicting, then we are allowed to process the new command
+			if (NetworkSimulationModelCVars::EnableLocalPrediction)
+			{
+				TickInfo.MaxAllowedInputKeyframe = Buffers.Input.GetHeadKeyframe();
+			}
+		}
+	}
+
+	// --------------------------------------------------------------------
+	//	NetSerialize
+	// --------------------------------------------------------------------
+	void NetSerialize(const FNetSerializeParams& P, TNetworkSimBufferContainer<TBufferTypes>& Buffers, const TSimulationTickState<TTickSettings>& TickInfo)
+	{
+		FArchive& Ar = P.Ar;
+
+		const int32 SerializedHeadKeyframe = FNetworkSimulationSerialization::SerializeKeyframe(Ar, Buffers.Sync.GetHeadKeyframe());
+		TSyncState* SerializedState = nullptr;
+
+		// Serialize total simulation time. This isn't really necessary since we have the keyframe above. 
+		SerializedTime = TickInfo.TotalProcessedSimulationTime;
+		SerializedTime.NetSerialize(P.Ar);
+
+		if (Ar.IsSaving())
+		{
+			// Server serialize the latest state
+			SerializedState = Buffers.Sync.GetElementFromHead(0);
+		}
+		else
+		{
+			// Its possible we process two packets in between ticks. This isn't really a problem but should be rare.
+			if (bPendingReconciliation)
+			{
+				UE_LOG(LogNetworkSim, Warning, TEXT("bPendingReconciliation while in ::NetSerialize. LastSerializedKeyframe: %d. New SerializedHeadKeyframe: %d."), LastSerializedKeyframe, SerializedHeadKeyframe);
+			}
+
+			// Lazy init the reconciliation buffer. We don't need this on the server/writing side. (Fixme, with templated approach we could easily specialize this to an inline element instead of buffer?)
+			if (ReconciliationBuffer.GetMaxNumElements() == 0)
+			{
+				ReconciliationBuffer.SetBufferSize(1);
+			}
+
+			// Reset the head to match the server. This is not critical, we just need a place to store this, but it makes the code in the reconciliation a bit nicer since the reconcilated frame number = head frame.
+			ReconciliationBuffer.ResetNextHeadKeyframe(SerializedHeadKeyframe);
+			SerializedState = ReconciliationBuffer.GetWriteNext();
+		}
+
+		SerializedState->NetSerialize(Ar);
+
+		if (Ar.IsLoading())
+		{
+			bReconcileFaultDetected = false;
+			bPendingReconciliation = false;
+
+			// The state the client predicted that corresponds to the state the server just serialized to us
+			if (TSyncState* ClientExistingState = Buffers.Sync.FindElementByKeyframe(SerializedHeadKeyframe))
+			{
+				if (ClientExistingState->ShouldReconcile(*SerializedState) || (NetworkSimulationModelCVars::ForceReconcile > 0))
+				{
+					UE_CLOG(!Buffers.Input.IsValidKeyframe(SerializedHeadKeyframe), LogNetworkSim, Error, TEXT("::NetSerialize: Client InputBuffer does not contain data for frame %d. {%s} {%s}"), SerializedHeadKeyframe, *Buffers.Input.GetBasicDebugStr(), *Buffers.Sync.GetBasicDebugStr());
+					bPendingReconciliation =  true;
+				}
+			}
+			else
+			{
+				// We don't have corresponding local state. There are two cases:
+				if (NetworkSimulationModelCVars::EnableLocalPrediction) // Fixme: this is awkward. Expected if player prediction is disabled but coupling like this feels bad.
+				{
+					UE_LOG(LogNetworkSim, Warning, TEXT("::NetSerialize Fault: SyncBuffer does not contain data for frame %d. [%d-%d]"), SerializedHeadKeyframe, Buffers.Sync.GetTailKeyframe(), Buffers.Sync.GetHeadKeyframe());
+				}
+
+				// Case 1: the serialized state is older than what we've kept in our buffer. A bigger buffer would solve this! (at the price of more resimulated frames to recover when this happens)
+				if (SerializedHeadKeyframe < Buffers.Sync.GetTailKeyframe())
+				{
+					// This is a reconcile fault and we just need to chill. We'll stop sending user commands until the cmds in flight flush through the system and we catch back up.
+					bReconcileFaultDetected = true;
+				}
+				else
+				{
+					// Case 2: We've received a newer frame than what we've processed locally. This could happen if we are buffering our inputs locally (but still sending to the server).
+					// Since this doesn't require resimulating, we can just set the authoritative state here and continue on. 
+					
+					
+
+					// However, this keyframe better be a valid input. Otherwise how did the server generate it?
+					checkf(SerializedHeadKeyframe <= Buffers.Input.GetHeadKeyframe(), TEXT("Received newer motionstate that doesn't correspond to valid input cmd. SerializedHeadKeyframe: %d. {%s} {%s}"),
+						SerializedHeadKeyframe, *Buffers.Input.GetBasicDebugStr(), *Buffers.Sync.GetBasicDebugStr());
+
+					Buffers.Sync.ResetNextHeadKeyframe(SerializedHeadKeyframe);
+					TSyncState* ClientMotionState = Buffers.Sync.GetWriteNext();
+					*ClientMotionState = *SerializedState;
+
+					//!!TickInfo.LastProcessedInputKeyframe = SerializedHeadKeyframe;
+				}
+			}
+		}
+
+		LastSerializedKeyframe = SerializedHeadKeyframe;
+	}
+
+	// --------------------------------------------------------------------
+	//	Reconcile
+	// --------------------------------------------------------------------
+	template<typename T, typename TDriver>
+	void Reconcile(TDriver* Driver, TNetworkSimBufferContainer<TBufferTypes>& Buffers, TSimulationTickState<TTickSettings>& TickInfo)
+	{
+		if (bPendingReconciliation == false)
+		{
+			return;
+		}
+		bPendingReconciliation = false;
+		check(ReconciliationBuffer.GetNumValidElements() == 1); // this function assumes only 1 element
+
+		const int32 ReconciliationKeyframe = ReconciliationBuffer.GetHeadKeyframe();
+		TSyncState* ServerState = ReconciliationBuffer.GetElementFromHead(0);
+
+		TInputCmd* ClientInputCmd = Buffers.Input.FindElementByKeyframe( ReconciliationKeyframe );
+		if (ClientInputCmd == nullptr)
+		{
+			// Fault: no historic command for this frame to reconcile with.
+			UE_LOG(LogNetworkSim, Error, TEXT("Client InputBuffer does not contain data for frame %d. {%s} {%s}"), ReconciliationKeyframe, *Buffers.Input.GetBasicDebugStr(), *Buffers.Sync.GetBasicDebugStr());
+			return;
+		}
+
+		TSyncState* ClientSyncState = Buffers.Sync.FindElementByKeyframe( ReconciliationKeyframe );
+
+		// Should not be here if the client doesn't have a valid state. See ::NetSerialize
+		checkf(ClientSyncState != nullptr, TEXT("SyncBuffer does not contain data for frame %d. %s"), ReconciliationKeyframe, *Buffers.Sync.GetBasicDebugStr());
+
+		// -------------------------------------------------------------------------------------------------------------------------
+		// Resimulate
+		// -------------------------------------------------------------------------------------------------------------------------
+
+		ServerState->VisualLog( FVisualLoggingParameters(EVisualLoggingContext::LastConfirmed, ReconciliationKeyframe, EVisualLoggingLifetime::Persistent), Driver, Driver);
+		ClientSyncState->VisualLog( FVisualLoggingParameters(EVisualLoggingContext::FirstMispredicted, ReconciliationKeyframe, EVisualLoggingLifetime::Persistent), Driver, Driver);
+
+		TSyncState* LatestClient = Buffers.Sync.GetElementFromHead(0);
+
+		// Copy authoritative state over client state (FIXME: we may want to store this off historically somewhere? Or will VLog be enough for debugging?)
+		*ClientSyncState = *ServerState;
+
+		// Set the canonical simulation time to what we received (we will advance it as we resimulate)
+		TickInfo.TotalProcessedSimulationTime = SerializedTime;
+		TickInfo.LastProcessedInputKeyframe = ReconciliationKeyframe;
+		
+		// Resimulate up to our latest SyncBuffer frame. Note that we may have unprocessed user commands in the command buffer at this point.
+		const int32 LatestKeyframe = Buffers.Sync.GetHeadKeyframe();		
+		for (int32 Keyframe = ReconciliationKeyframe+1; Keyframe <= LatestKeyframe; ++Keyframe)
+		{
+			// Keyframe is the frame we are resimulating right now.
+			TInputCmd* ResimulateCmd	= Buffers.Input.FindElementByKeyframe(Keyframe);
+			TSyncState* PrevMotionState = Buffers.Sync.FindElementByKeyframe(Keyframe - 1);
+			TSyncState* NextMotionState = Buffers.Sync.FindElementByKeyframe(Keyframe);
+			
+			check(ResimulateCmd);
+			check(PrevMotionState);
+			check(NextMotionState);
+
+			// TEMP (Aux buffer not fully plumbed through the system yet)
+			// ------------------------------------------------------
+			TAuxState TempAuxState;
+			TAuxState* AuxState = Buffers.Aux.FindElementByKeyframe(Keyframe);
+			if (AuxState == nullptr)
+			{
+				AuxState = &TempAuxState;
+			}
+			// ------------------------------------------------------
+
+			// Log out the Mispredicted state that we are about to overwrite.
+			NextMotionState->VisualLog( FVisualLoggingParameters(Keyframe == LatestKeyframe ? EVisualLoggingContext::LastMispredicted : EVisualLoggingContext::OtherMispredicted, Keyframe, EVisualLoggingLifetime::Persistent), Driver, Driver);
+
+			// Do the actual update
+			T::Update(Driver, ResimulateCmd->GetFrameDeltaTime(), *ResimulateCmd, *PrevMotionState, *NextMotionState, *AuxState);
+			
+			// Update TickInfo
+			TickInfo.TotalProcessedSimulationTime += ResimulateCmd->GetFrameDeltaTime();
+			TickInfo.LastProcessedInputKeyframe = Keyframe;
+
+			// Log out the newly predicted state that we got.
+			NextMotionState->VisualLog( FVisualLoggingParameters(Keyframe == LatestKeyframe ? EVisualLoggingContext::LastPredicted : EVisualLoggingContext::OtherPredicted, Keyframe, EVisualLoggingLifetime::Persistent), Driver, Driver);
+		}
+	}
+
+private:
+	
+	TReplicationBuffer<TSyncState> ReconciliationBuffer;
+	TNetworkSimTime<TTickSettings> SerializedTime; // last serialized time keeper
+
+	int32 LastSerializedKeyframe = -1;
+	bool bPendingReconciliation = false;	// Reconciliation is pending: we need to reconcile state from the server that differs from the locally predicted state
+	bool bReconcileFaultDetected = false;	// A fault was detected: we received state from the server that we are unable to reconcile with locally predicted state
+};
\ No newline at end of file
diff --git a/Engine/Plugins/Runtime/NetworkPrediction/Source/NetworkPrediction/Public/NetworkSimulationModelTemplates.h b/Engine/Plugins/Runtime/NetworkPrediction/Source/NetworkPrediction/Public/NetworkSimulationModelTemplates.h
deleted file mode 100644
diff --git a/Engine/Plugins/Runtime/NetworkPrediction/Source/NetworkPrediction/Public/NetworkSimulationModelTypes.h b/Engine/Plugins/Runtime/NetworkPrediction/Source/NetworkPrediction/Public/NetworkSimulationModelTypes.h
--- a/Engine/Plugins/Runtime/NetworkPrediction/Source/NetworkPrediction/Public/NetworkSimulationModelTypes.h
+++ b/Engine/Plugins/Runtime/NetworkPrediction/Source/NetworkPrediction/Public/NetworkSimulationModelTypes.h
@@ -17,6 +17,37 @@ enum ENetworkSimBufferTypeId
 	Debug
 };
 
+// Helper needed to specialize TNetworkSimBufferTypes::Get (must be done outside of templated struct)
+template<typename TBufferTypes, ENetworkSimBufferTypeId BufferId> 
+struct TSelectTypeHelper
+{
+	using type = void;
+};
+
+template<typename TBufferTypes>
+struct TSelectTypeHelper<TBufferTypes, Input>
+{
+	using type = typename TBufferTypes::TInputCmd;
+};
+
+template<typename TBufferTypes>
+struct TSelectTypeHelper<TBufferTypes, Sync>
+{
+	using type = typename TBufferTypes::TSyncState;
+};
+
+template<typename TBufferTypes>
+struct TSelectTypeHelper<TBufferTypes, Aux>
+{
+	using type = typename TBufferTypes::TAuxState;
+};
+
+template<typename TBufferTypes>
+struct TSelectTypeHelper<TBufferTypes, Debug>
+{
+	using type = typename TBufferTypes::TDebugState;
+};
+
 // A collection of the system's buffer types. This allows us to collapse the 4 types into a single type to use a template argument elsewhere.
 template<typename InInputCmd, typename InSyncState, typename InAuxState, typename InDebugState = FNetSimProcessedFrameDebugInfo>
 struct TNetworkSimBufferTypes
@@ -28,121 +59,62 @@ struct TNetworkSimBufferTypes
 	using TDebugState = InDebugState;
 
 	// Template access via ENetworkSimBufferTypeId when "which buffer" is parameterized
-	template<ENetworkSimBufferTypeId>
+	template<ENetworkSimBufferTypeId Id>
 	struct select_type
 	{
-		using type = void;
-	};
-
-	template<>
-	struct select_type<ENetworkSimBufferTypeId::Input>
-	{
-		using type = TInputCmd;
-	};
-
-	template<>
-	struct select_type<ENetworkSimBufferTypeId::Sync>
-	{
-		using type = TSyncState;
-	};
-
-	template<>
-	struct select_type<ENetworkSimBufferTypeId::Aux>
-	{
-		using type = TAuxState;
-	};
-
-	template<>
-	struct select_type<ENetworkSimBufferTypeId::Debug>
-	{
-		using type = TDebugState;
+		using type = typename TSelectTypeHelper< TNetworkSimBufferTypes<TInputCmd, TSyncState, TAuxState, TDebugState>, Id >::type;
 	};
 };
 
 // ---------------------------------------------------------------------------------------------------------------------
-//	Ticking and Input processing helpers. We want to support fixed tick systems (but variable is most important).
-//	Whether your system is fixed or not depends on your input cmd type.
+//	TNetworkSimBufferContainer
+//	Container for the actual replicated buffers that the system uses.
+//	Has compile time accessors for retrieving the buffers based on templated enum value.
 // ---------------------------------------------------------------------------------------------------------------------
 
-// Basic struct trait for whether the input cmd is fixed or not. By default you are not fixed, you must specialize this template for your type and set value=true
-template<typename T>
-struct TNetworkSimInput_is_fixedtick
+// Helper struct for enum-based access. This has to be done in an outside struct because we cant specialize inside TNetworkSimBufferContainer on all compilers
+template<typename TContainer, ENetworkSimBufferTypeId BufferId>
+struct TBufferGetterHelper
 {
-	static const bool value = false;
+	static typename TContainer::template select_buffer_type<BufferId>::type& Get(TContainer& Container)
+	{
+		static_assert(!BufferId, "Failed to find specialized Get for your BufferId");
+	}
 };
 
-// Variable tick simulation time. For now this is a double.
-template<bool IsFixedTick=false>
-struct TSimulationTime
+template<typename TContainer>
+struct TBufferGetterHelper<TContainer, ENetworkSimBufferTypeId::Input>
 {
-	TSimulationTime(double InTime=0.f) : Time(InTime) { }
-
-	template<typename T>
-	void AccumulateTimeFromInputCmd(const T& InputCmd)
+	static typename TContainer::TInputBuffer& Get(TContainer& Container)
 	{
-		Time += InputCmd.FrameDeltaTime;
+		return Container.Input;
 	}
-	
-	void NetSerialize(const FNetSerializeParams& P)
-	{
-		// Temp: we can devise a system like keyframe replication to make this more efficient
-		P.Ar << Time;
-	}
-
-	double Time = 0.f;
 };
-
-// Fixed tick simulation time. Just a frame counter.
-template<>
-struct TSimulationTime<true>
+template<typename TContainer>
+struct TBufferGetterHelper<TContainer, ENetworkSimBufferTypeId::Sync>
 {
-	TSimulationTime(int32 InTime=0) : Time(InTime) { }
-
-	template<typename T>
-	void AccumulateTimeFromInputCmd(const T& InputCmd)
+	static typename TContainer::TSyncBuffer& Get(TContainer& Container)
 	{
-		Time++;
+		return Container.Sync;
 	}
-
-	void NetSerialize(const FNetSerializeParams& P)
+};
+template<typename TContainer>
+struct TBufferGetterHelper<TContainer, ENetworkSimBufferTypeId::Aux>
+{
+	static typename TContainer::TAuxBuffer& Get(TContainer& Container)
 	{
-		// Temp: we can devise a system like keyframe replication to make this more efficient
-		P.Ar << Time;
+		return Container.Aux;
 	}
-
-	int32 Time = 0;
 };
-
-// Actual struct to use to track time. Parameterized by your NetworkSimBufferTypes and uses the input cmd to determine if you are fixed or not.
-template<typename TBufferTypes, typename TParent = TSimulationTime< TNetworkSimInput_is_fixedtick<typename TBufferTypes::TInputCmd>::value >>
-struct TSimulationTimeKeeper : public TParent
+template<typename TContainer>
+struct TBufferGetterHelper<TContainer, ENetworkSimBufferTypeId::Debug>
 {
-	using T = TSimulationTimeKeeper<TBufferTypes>;
-	using TParent::TParent;
-
-	FString ToString() const { return LexToString(this->Time); }
-
-	T& operator+= (const T &rhs) { this->Time += rhs.Time; return(*this); }
-	T& operator-= (const T &rhs) { this->Time -= rhs.Time; return(*this); }
-	
-	T operator+ (const T &rhs) { return T(this->Time - rhs.Time); }
-	T operator- (const T &rhs) { return T(this->Time - rhs.Time); }
-
-	bool operator<  (const T &rhs) const { return(this->Time < rhs.Time); }
-	bool operator<= (const T &rhs) const { return(this->Time <= rhs.Time); }
-	bool operator>  (const T &rhs) const { return(this->Time > rhs.Time); }
-	bool operator>= (const T &rhs) const { return(this->Time >= rhs.Time); }
-	bool operator== (const T &rhs) const { return(this->Time == rhs.Time); }
-	bool operator!= (const T &rhs) const { return(this->Time != rhs.Time); }
+	static typename TContainer::TDebugBuffer& Get(TContainer& Container)
+	{
+		return Container.Debug;
+	}
 };
 
-
-// ---------------------------------------------------------------------------------------------------------------------
-//	TNetworkSimBufferContainer
-//	Container for the actual replicated buffers that the system uses.
-//	Has compile time accessors for retrieving the buffers based on templated enum value.
-// ---------------------------------------------------------------------------------------------------------------------
-
 template<typename T>
 struct TNetworkSimBufferContainer
 {
@@ -176,65 +148,199 @@ struct TNetworkSimBufferContainer
 	template<ENetworkSimBufferTypeId BufferId>
 	typename select_buffer_type<BufferId>::type& Get()
 	{
-		static_assert(!BufferId, "Failed to find specialized Get for your BufferId");
+		return TBufferGetterHelper<TNetworkSimBufferContainer<T>, BufferId>::Get(*this);
 	}
+};
 
-	template<>
-	TInputBuffer& Get<ENetworkSimBufferTypeId::Input>()
-	{
-		return Input;
-	}
+// ----------------------------------------------------------------------------------------------------------------------------------------------
+//
+// ----------------------------------------------------------------------------------------------------------------------------------------------
+
+// Holds all settings for a network sim related to ticking
+template<uint32 InFixedStepMS=0, uint32 InMaxStepMS=0, typename TUnderlyingSimTimeType=uint32, typename TUnderlyingRealTimeType=float, int32 InRealToSimFactor=1000>
+struct TNetworkSimTickSettings
+{
+	using TUnderlingSimTime = TUnderlyingSimTimeType;	// Underlying type of time. This type shouldn't be used directly (use TNetworkSimTime)
+	using TRealTime = TUnderlyingRealTimeType;	// This is the "final" real time type. We do not wrap it in anything else.
 
-	template<>
-	TSyncBuffer& Get<ENetworkSimBufferTypeId::Sync>()
+	static_assert(!(InFixedStepMS!=0 && InMaxStepMS != 0), "MaxStepMS is only applicable when using variable step (FixedStepMS == 0)");
+
+	enum 
 	{
-		return Sync;
-	}
+		MaxStepMS = InMaxStepMS,				// Max step. Only applicable to variable time step.
+		FixedStepMS = InFixedStepMS,			// Fixed step. If 0, then we are "variable time step"
+		RealToSimFactor = InRealToSimFactor		// Factor to go from RealTime (always seconds) to SimTime (MSec by default with factor of 1000)
+	};
+
+	// Typed accessors
+	static constexpr TUnderlingSimTime GetMaxStepMS() { return static_cast<TUnderlingSimTime>(InMaxStepMS); }
+	static constexpr TUnderlingSimTime GetFixedStepMS() { return static_cast<TUnderlingSimTime>(InFixedStepMS); }
+	static constexpr TRealTime GetRealToSimFactor() { return static_cast<TRealTime>(InRealToSimFactor); }
+};
 
-	template<>
-	TAuxBuffer& Get<ENetworkSimBufferTypeId::Aux>()
+// Actual time value. This by default will store time in MSec (ultimately determined by TickSettings::RealToSimFactor)
+template<typename TickSettings>
+struct TNetworkSimTime
+{
+	using TTime = typename TickSettings::TUnderlingSimTime;
+	using TRealTime = typename TickSettings::TRealTime;
+
+	TNetworkSimTime() { }
+	TNetworkSimTime(const TTime& InTime) : Time(InTime) { }
+	TNetworkSimTime(const TRealTime& InRealTime) : Time(InRealTime * TickSettings::RealToSimFactor) { }
+
+	TRealTime ToRealTimeSeconds() const { return (Time / TickSettings::RealToSimFactor); }
+	FString ToString() const { return LexToString(this->Time); }
+
+	// FIXME
+	void NetSerialize(FArchive& Ar) { Ar << Time; }
+
+	using T = TNetworkSimTime;
+	T& operator+= (const T &rhs) { this->Time += rhs.Time; return(*this); }
+	T& operator-= (const T &rhs) { this->Time -= rhs.Time; return(*this); }
+	
+	T operator+ (const T &rhs) const { return T(this->Time - rhs.Time); }
+	T operator- (const T &rhs) const { return T(this->Time - rhs.Time); }
+
+	bool operator<  (const T &rhs) const { return(this->Time < rhs.Time); }
+	bool operator<= (const T &rhs) const { return(this->Time <= rhs.Time); }
+	bool operator>  (const T &rhs) const { return(this->Time > rhs.Time); }
+	bool operator>= (const T &rhs) const { return(this->Time >= rhs.Time); }
+	bool operator== (const T &rhs) const { return(this->Time == rhs.Time); }
+	bool operator!= (const T &rhs) const { return(this->Time != rhs.Time); }
+
+	TTime Time = 0;
+};
+
+// ----------------------------------------------------------------------------------------------------------------------------------------------
+//
+// ----------------------------------------------------------------------------------------------------------------------------------------------
+
+// Helper for accumulating real time into sim time based on TickSettings
+template<typename TickSettings, bool IsFixedTick=(TickSettings::FixedStepMS!=0)>
+struct TRealTimeAccumulator
+{
+	using TRealTime = typename TickSettings::TRealTime;
+	void Accumulate(TNetworkSimTime<TickSettings>& NetworkSimTime, TRealTime RealTimeSeconds)
 	{
-		return Aux;
+		// Non fixed step: just accumulate the time directly
+		// Note that MaxStepMS enforcement does NOT belong here. Dropping time due to MaxStepMS would just make the sim run slower. MaxStepMS is used at the input processing level.
+		NetworkSimTime += RealTimeSeconds;
 	}
+};
+
+// Specialized version of FixedTicking. This accumulates real time that spills over into NetworkSimTime as it crosses the FixStep threshold
+template<typename TickSettings>
+struct TRealTimeAccumulator<TickSettings, true>
+{
+	using TRealTime = typename TickSettings::TRealTime;
+	const TRealTime RealTimeFixedStep = TNetworkSimTime<TickSettings>(TickSettings::GetFixedStepMS() / TickSettings::GetRealToSimFactor()).ToRealTimeSeconds();
 
-	template<>
-	TDebugBuffer& Get<ENetworkSimBufferTypeId::Debug>()
+	void Accumulate(TNetworkSimTime<TickSettings>& NetworkSimTime, TRealTime RealTimeSeconds)
 	{
-		return Debug;
+		AccumulatedTime += RealTimeSeconds;
+		if (AccumulatedTime > RealTimeFixedStep)
+		{
+			const int32 NumFrames = AccumulatedTime / RealTimeFixedStep;
+			AccumulatedTime -= NumFrames * RealTimeFixedStep;
+				
+			if (FMath::Abs<TRealTime>(AccumulatedTime) < SMALL_NUMBER)
+			{
+				AccumulatedTime = TRealTime(0.f);
+			}
+
+			NetworkSimTime += TNetworkSimTime<TickSettings>(NumFrames * TickSettings::GetFixedStepMS());
+		}
 	}
+	TRealTime AccumulatedTime;
 };
 
+// ----------------------------------------------------------------------------------------------------------------------------------------------
+//
+// ----------------------------------------------------------------------------------------------------------------------------------------------
 
-// ---------------------------------------------------------------------------------------------------------------------
-//	TSimulationTicker
-//	This holds the data for where we are in ticking the sim. Such as how much time we've ticked it, how much we are allowed to tick it.
-//	It also holds input processing info: the last keyframe of the input buffer we processed and the last GFrameNumber we accepted input.
-// ---------------------------------------------------------------------------------------------------------------------
-
-template<typename TBufferTypes>
-struct TSimulationTickInfo
+// Holds active state for simulation ticking: what inputs we have processed, how time we have simulated, how much we are allowed to, etc.
+template<typename TickSettings=TNetworkSimTickSettings<>>
+struct TSimulationTickState
 {
-	using FSimulationTime = TSimulationTimeKeeper<TBufferTypes>;
+	using TSettings = TickSettings;
+
+	int32 LastProcessedInputKeyframe;	// The last input keyframe that we processed
+	int32 MaxAllowedInputKeyframe;		// The max input keyframe that we are allowed to process (e.g, don't process input past this keyframe yet)
 
-	int32 LastProcessedInputKeyframe;	// Tracks input keyframe that we last processed
-	uint32 LastLocalInputGFrameNumber;	// Tracks the last time we accepted local input via GetNextInputForWrite. Used to guard against accidental input latency due to ordering of input/sim ticking.
 
-	FSimulationTime MaxSimulationTime;			// Total max time we have been "given" to process. We cannot process more simulation time than this: doing so would be speed hacking.
-	FSimulationTime ProcessedSimulationTime;	// How much time we've actually processed. The only way to increment this is to process user commands or receive authoritative state from the network.
+	uint32 LastLocalInputGFrameNumber;	// Tracks the last time we accepted local input via GetNextInputForWrite. Used to guard against accidental input latency due to ordering of input/sim ticking.
 
-	using TInputCmd = typename TBufferTypes::TInputCmd;
-	TInputCmd* GetNextInputForWrite(TNetworkSimBufferContainer<TBufferTypes>& Buffers)
+	TNetworkSimTime<TSettings> TotalAllowedSimulationTime;	// Total time we have been "given" to process. We cannot process more simulation time than this: doing so would be speed hacking.
+	TNetworkSimTime<TSettings> TotalProcessedSimulationTime;	// How much time we've actually processed. The only way to increment this is to process user commands or receive authoritative state from the network.
+	
+	template<typename TBufferTypes>
+	typename TBufferTypes::TInputCmd* GetNextInputForWrite(TNetworkSimBufferContainer<TBufferTypes>& Buffers)
 	{
+		// Not really necessary anymore since input command is requested by the sim now instead of pushed into it
 		ensure(LastLocalInputGFrameNumber != GFrameNumber);
 		LastLocalInputGFrameNumber = GFrameNumber;
 
-		TInputCmd* Next = nullptr;
+		typename TBufferTypes::TInputCmd* Next = nullptr;
 		if (Buffers.Input.GetHeadKeyframe() == LastProcessedInputKeyframe)
 		{
 			// Only return a cmd if we have processed the last one. This is a bit heavy handed but is a good practice to start with. We want buffering of input to be very explicit, never something that accidentally happens.
 			Next = Buffers.Input.GetWriteNext();
-			*Next = TInputCmd();
+			*Next = typename TBufferTypes::TInputCmd();
 		}
 		return Next;
 	}
+
+	// "Grants" allowed simulation time to this tick state. That is, we are now allowed to advance the simulation by this amount the next time the sim ticks.
+	// Note the input is RealTime in SECONDS. This is what the rest of the engine uses when dealing with float delta time.
+	void GiveSimulationTime(float RealTimeSeconds)
+	{
+		RealtimeAccumulator.Accumulate(TotalAllowedSimulationTime, RealTimeSeconds);
+	}
+
+	TNetworkSimTime<TSettings> GetRemaningAllowedSimulationTime() const
+	{
+		return TotalAllowedSimulationTime - TotalProcessedSimulationTime;
+	}
+
+private:
+
+	TRealTimeAccumulator<TSettings>	RealtimeAccumulator;
+};
+
+// ----------------------------------------------------------------------------------------------------------------------------------------------
+//
+// ----------------------------------------------------------------------------------------------------------------------------------------------
+
+// Wraps an input command (BaseType) in a NetworkSimulation time. 
+template<typename BaseCmdType, typename TickSettings, bool IsFixedTick=(TickSettings::FixedStepMS!=0)>
+struct TFrameCmd : public BaseCmdType
+{
+	TNetworkSimTime<TickSettings> GetFrameDeltaTime() const { return FrameDeltaTime; }
+	void SetFrameDeltaTime(const TNetworkSimTime<TickSettings>& InTime) { FrameDeltaTime = InTime; }
+
+private:
+	TNetworkSimTime<TickSettings> FrameDeltaTime;
+};
+
+// Fixed tick specialization
+template<typename BaseCmdType, typename TickSettings>
+struct TFrameCmd<BaseCmdType, TickSettings, true> : public BaseCmdType
+{
+	constexpr TNetworkSimTime<TickSettings> GetFrameDeltaTime() const { return TNetworkSimTime<TickSettings>(TickSettings::GetFixedStepMS()); }
+	void SetFrameDeltaTime(const TNetworkSimTime<TickSettings>& InTime) { }
+};
+
+// Helper to turn user supplied buffer types into the "real" buffer types: the InputCmd struct is wrapped in TFrameCmd
+template<typename TUserBufferTypes, typename TTickSettings>
+struct TInternalBufferTypes : TNetworkSimBufferTypes< 
+	
+	// InputCmds are wrapped in TFrameCmd, which will store an explicit sim delta time if we are not a fixed tick sim
+	TFrameCmd< typename TUserBufferTypes::TInputCmd , TTickSettings>,
+
+	typename TUserBufferTypes::TSyncState,	// SyncState Passes through
+	typename TUserBufferTypes::TAuxState,	// Auxstate passes through
+	typename TUserBufferTypes::TDebugState	// Debugstate passes through
+>
+{
 };
\ No newline at end of file
diff --git a/Engine/Plugins/Runtime/NetworkPredictionExtras/Source/NetworkPredictionExtras/Private/MockNetworkSimulation.cpp b/Engine/Plugins/Runtime/NetworkPredictionExtras/Source/NetworkPredictionExtras/Private/MockNetworkSimulation.cpp
--- a/Engine/Plugins/Runtime/NetworkPredictionExtras/Source/NetworkPredictionExtras/Private/MockNetworkSimulation.cpp
+++ b/Engine/Plugins/Runtime/NetworkPredictionExtras/Source/NetworkPredictionExtras/Private/MockNetworkSimulation.cpp
@@ -59,9 +59,9 @@ static FAutoConsoleVariableRef CVarBindAutomatically(TEXT("mns.BindAutomatically
 // ============================================================================================
 
 static bool ForceMispredict = false;
-void FMockNetworkSimulation::Update(IMockNetworkSimulationDriver* Driver, const FMockInputCmd& InputCmd, const FMockSyncState& InputState, FMockSyncState& OutputState, const FMockAuxState& AuxState)
+void FMockNetworkSimulation::Update(IMockNetworkSimulationDriver* Driver, const TSimTime& DeltaTimeMS, const FMockInputCmd& InputCmd, const FMockSyncState& InputState, FMockSyncState& OutputState, const FMockAuxState& AuxState)
 {
-	OutputState.Total = InputState.Total + (InputCmd.InputValue * AuxState.Multiplier * InputCmd.FrameDeltaTime);
+	OutputState.Total = InputState.Total + (InputCmd.InputValue * AuxState.Multiplier * DeltaTimeMS.ToRealTimeSeconds());
 
 	// Dev hack to force mispredict
 	if (ForceMispredict)
@@ -150,27 +150,7 @@ void UMockNetworkSimulationComponent::TickComponent(float DeltaTime, enum ELevel
 	CheckOwnerRoleChange();
 
 	const ENetRole OwnerRole = GetOwnerRole();
-
-	// -------------------------------------------------------
-	//	Local Input
-	//		-For the mock sim, just create the input here.
-	//		-Real implementation, the input will probably be set outside of tick, from a PC, InputComponent, etc.
-	// -------------------------------------------------------
-	if (OwnerRole == ROLE_AutonomousProxy)
-	{
-		if (FMockInputCmd* Input = NetworkSim->GetNextInputForWrite(DeltaTime))
-		{
-			if (MockNetworkSimCVars::DoLocalInput)
-			{
-				Input->InputValue = FMath::FRand() * 10.f;
-			}
-			else
-			{
-				Input->InputValue = 0.f;
-			}
-		}
-	}
-
+	
 	// -------------------------------------
 	// Tick the network sim
 	// -------------------------------------
@@ -204,6 +184,18 @@ void UMockNetworkSimulationComponent::TickComponent(float DeltaTime, enum ELevel
 	DrawDebugString( GetWorld(), GetOwner()->GetActorLocation() + FVector(0.f,0.f,100.f), *LexToString(MockValue), nullptr, FColor::White, 0.00001f );
 }
 
+void UMockNetworkSimulationComponent::ProduceInput(const FMockNetworkSimulation::TSimTime& SimFrameTime, FMockInputCmd& Cmd)
+{
+	if (MockNetworkSimCVars::DoLocalInput)
+	{
+		Cmd.InputValue = FMath::FRand() * 10.f;
+	}
+	else
+	{
+		Cmd.InputValue = 0.f;
+	}
+}
+
 void UMockNetworkSimulationComponent::InitSyncState(FMockSyncState& OutSyncState) const
 {
 	OutSyncState.Total = MockValue;
diff --git a/Engine/Plugins/Runtime/NetworkPredictionExtras/Source/NetworkPredictionExtras/Private/NetworkPredictionExtrasFlyingPawn.cpp b/Engine/Plugins/Runtime/NetworkPredictionExtras/Source/NetworkPredictionExtras/Private/NetworkPredictionExtrasFlyingPawn.cpp
--- a/Engine/Plugins/Runtime/NetworkPredictionExtras/Source/NetworkPredictionExtras/Private/NetworkPredictionExtrasFlyingPawn.cpp
+++ b/Engine/Plugins/Runtime/NetworkPredictionExtras/Source/NetworkPredictionExtras/Private/NetworkPredictionExtrasFlyingPawn.cpp
@@ -35,15 +35,7 @@ ANetworkPredictionExtrasFlyingPawn::ANetworkPredictionExtrasFlyingPawn()
 	
 	if (HasAnyFlags(RF_ClassDefaultObject) == false)
 	{
-		// Setup callback for generating local input. There is some nuance here. This is the order we want things to happen for a locally controlled actor:
-		//	1. Create local input (the lambda below)
-		//	2. Tick the network sim (called by UFlyingMovementComponent::TickComponent). This wants the latest input (important the user cmd has the same DeltaTime as the current frame!)
-		//  3. Tick this actor (ANetworkPredictionExtrasFlyingPawn::Tick). This wants the latest movement state (computed in step #2).
-		//
-		// The ::SetLocallyControlledPreTick isn't strictly necessary. A player controller or some other object could tick first, submit the input, then let the sim tick and this actor tick.
-		// But for this simple example, its clearer to contain everything here. The delegate is just a helper to save you from dealing with tick graph prereq headaches.
-
-		FlyingMovementComponent->SetLocallyControlledPreTick(UNetworkPredictionComponent::FPreSimTickDelegate::CreateUObject(this, &ThisClass::GenerateLocalInput));
+		FlyingMovementComponent->ProduceInputDelegate.BindUObject(this, &ANetworkPredictionExtrasFlyingPawn::ProduceInput);
 
 		// Binds 0 and 9 to the debug hud commands. This is just a convenience for the extras plugin. Real projects should bind this themselves
 		if (FlyingPawnCVars::BindAutomatically > 0)
@@ -143,7 +135,7 @@ void ANetworkPredictionExtrasFlyingPawn::Tick( float DeltaSeconds)
 	// Do whatever you want here. By now we have the latest movement state and latest input processed.
 }
 
-void ANetworkPredictionExtrasFlyingPawn::GenerateLocalInput(float DeltaSeconds)
+void ANetworkPredictionExtrasFlyingPawn::ProduceInput(const FlyingMovement::TSimTime& SimFrameTime, FlyingMovement::FInputCmd& Cmd)
 {
 	// Generate user commands. Called right before the flying movement simulation will tick (for a locally controlled pawn)
 	// This isn't meant to be the best way of doing a camera system. It is just meant to show a couple of ways it may be done
@@ -158,164 +150,164 @@ void ANetworkPredictionExtrasFlyingPawn::GenerateLocalInput(float DeltaSeconds)
 	// targeting systems to happen /outside/ of the system, i.e, here. But I can think of scenarios where that may not be ideal too.
 
 	check(Controller);
-	if (FlyingMovement::FInputCmd* NextCmd = FlyingMovementComponent->GetNextClientInputCmdForWrite(DeltaSeconds))
+	
+	if (USpringArmComponent* SpringComp = FindComponentByClass<USpringArmComponent>())
 	{
-		if (USpringArmComponent* SpringComp = FindComponentByClass<USpringArmComponent>())
-		{
-			// This is not best practice: do not search for component every frame
-			SpringComp->bUsePawnControlRotation = true;
-		}
+		// This is not best practice: do not search for component every frame
+		SpringComp->bUsePawnControlRotation = true;
+	}
 
-		// Simple input scaling. A real game will probably map this to an acceleration curve
-		static float LookRateYaw = 150.f;
-		static float LookRatePitch = 150.f;
+	// Simple input scaling. A real game will probably map this to an acceleration curve
+	static float LookRateYaw = 150.f;
+	static float LookRatePitch = 150.f;
 
-		static float ControllerLookRateYaw = 1.5f;
-		static float ControllerLookRatePitch = 1.5f;
+	static float ControllerLookRateYaw = 1.5f;
+	static float ControllerLookRatePitch = 1.5f;
 
-		// Zero out input structs in case each path doesnt set each member. This is really all we are filling out here.
-		NextCmd->MovementInput = FVector::ZeroVector;
-		NextCmd->RotationInput = FRotator::ZeroRotator;
+	// Zero out input structs in case each path doesnt set each member. This is really all we are filling out here.
+	Cmd.MovementInput = FVector::ZeroVector;
+	Cmd.RotationInput = FRotator::ZeroRotator;
 
-		switch (FlyingPawnCVars::CameraStyle)
+	const float DeltaTimeSeconds = SimFrameTime.ToRealTimeSeconds();
+
+	switch (FlyingPawnCVars::CameraStyle)
+	{
+		case 0:
 		{
-			case 0:
+			// Fixed camera
+			if (USpringArmComponent* SpringComp = FindComponentByClass<USpringArmComponent>())
 			{
-				// Fixed camera
-				if (USpringArmComponent* SpringComp = FindComponentByClass<USpringArmComponent>())
-				{
-					// Only this camera mode has to set this to false
-					SpringComp->bUsePawnControlRotation = false;
-				}
+				// Only this camera mode has to set this to false
+				SpringComp->bUsePawnControlRotation = false;
+			}
 
-				NextCmd->RotationInput.Yaw = CachedLookInput.X * LookRateYaw;
-				NextCmd->RotationInput.Pitch = CachedLookInput.Y * LookRatePitch;
-				NextCmd->RotationInput.Roll = 0;
+			Cmd.RotationInput.Yaw = CachedLookInput.X * LookRateYaw;
+			Cmd.RotationInput.Pitch = CachedLookInput.Y * LookRatePitch;
+			Cmd.RotationInput.Roll = 0;
 					
-				NextCmd->MovementInput = CachedMoveInput;
-				break;
-			}
-			case 1:
+			Cmd.MovementInput = CachedMoveInput;
+			break;
+		}
+		case 1:
+		{
+			// Free camera Restricted 2D movement on XY plane.
+			APlayerController* PC = Cast<APlayerController>(Controller);
+			if (ensure(PC)) // Requires player controller for now
 			{
-				// Free camera Restricted 2D movement on XY plane.
-				APlayerController* PC = Cast<APlayerController>(Controller);
-				if (ensure(PC)) // Requires player controller for now
-				{
-					// Camera yaw rotation
-					PC->AddYawInput(CachedLookInput.X * ControllerLookRateYaw );
-					PC->AddPitchInput(CachedLookInput.Y * ControllerLookRatePitch );
+				// Camera yaw rotation
+				PC->AddYawInput(CachedLookInput.X * ControllerLookRateYaw );
+				PC->AddPitchInput(CachedLookInput.Y * ControllerLookRatePitch );
 
-					static float RotationMagMin = (1e-3);
+				static float RotationMagMin = (1e-3);
 						
-					float RotationInputYaw = 0.f;
-					if (CachedMoveInput.Size() >= RotationMagMin)
-					{
-						FVector DesiredMovementDir = PC->GetControlRotation().RotateVector(CachedMoveInput);
+				float RotationInputYaw = 0.f;
+				if (CachedMoveInput.Size() >= RotationMagMin)
+				{
+					FVector DesiredMovementDir = PC->GetControlRotation().RotateVector(CachedMoveInput);
 
-						// 2D xy movement, relative to camera
-						{
-							FVector DesiredMoveDir2D = FVector(DesiredMovementDir.X, DesiredMovementDir.Y, 0.f);
-							DesiredMoveDir2D.Normalize();
+					// 2D xy movement, relative to camera
+					{
+						FVector DesiredMoveDir2D = FVector(DesiredMovementDir.X, DesiredMovementDir.Y, 0.f);
+						DesiredMoveDir2D.Normalize();
 
-							const float DesiredYaw = DesiredMoveDir2D.Rotation().Yaw;
-							const float DeltaYaw = DesiredYaw - GetActorRotation().Yaw;
-							NextCmd->RotationInput.Yaw = DeltaYaw / NextCmd->FrameDeltaTime;
-						}
+						const float DesiredYaw = DesiredMoveDir2D.Rotation().Yaw;
+						const float DeltaYaw = DesiredYaw - GetActorRotation().Yaw;
+						Cmd.RotationInput.Yaw = DeltaYaw / DeltaTimeSeconds;
 					}
-
-					NextCmd->MovementInput = FVector(FMath::Clamp<float>( CachedMoveInput.Size2D(), -1.f, 1.f ), 0.0f, CachedMoveInput.Z);
 				}
-				break;
+
+				Cmd.MovementInput = FVector(FMath::Clamp<float>( CachedMoveInput.Size2D(), -1.f, 1.f ), 0.0f, CachedMoveInput.Z);
 			}
-			case 2:
+			break;
+		}
+		case 2:
+		{
+			// Free camera on yaw and pitch, camera-relative movement.
+			APlayerController* PC = Cast<APlayerController>(Controller);
+			if (ensure(PC)) // Requires player controller for now
 			{
-				// Free camera on yaw and pitch, camera-relative movement.
-				APlayerController* PC = Cast<APlayerController>(Controller);
-				if (ensure(PC)) // Requires player controller for now
-				{
-					// Camera yaw rotation
-					PC->AddYawInput(CachedLookInput.X * ControllerLookRateYaw );
-					PC->AddPitchInput(CachedLookInput.Y * ControllerLookRatePitch );
+				// Camera yaw rotation
+				PC->AddYawInput(CachedLookInput.X * ControllerLookRateYaw );
+				PC->AddPitchInput(CachedLookInput.Y * ControllerLookRatePitch );
 
-					// Rotational movement: orientate us towards our camera-relative desired velocity (unless we are upside down, then flip it)
-					static float RotationMagMin = (1e-3);
-					const float MoveInputMag = CachedMoveInput.Size();
+				// Rotational movement: orientate us towards our camera-relative desired velocity (unless we are upside down, then flip it)
+				static float RotationMagMin = (1e-3);
+				const float MoveInputMag = CachedMoveInput.Size();
 						
-					float RotationInputYaw = 0.f;
-					float RotationInputPitch = 0.f;
-					float RotationInputRoll = 0.f;
+				float RotationInputYaw = 0.f;
+				float RotationInputPitch = 0.f;
+				float RotationInputRoll = 0.f;
 
-					if (MoveInputMag >= RotationMagMin)
-					{								
-						FVector DesiredMovementDir = PC->GetControlRotation().RotateVector(CachedMoveInput);
+				if (MoveInputMag >= RotationMagMin)
+				{								
+					FVector DesiredMovementDir = PC->GetControlRotation().RotateVector(CachedMoveInput);
 
-						const float DesiredYaw = DesiredMovementDir.Rotation().Yaw;
-						const float DeltaYaw = DesiredYaw - GetActorRotation().Yaw;
+					const float DesiredYaw = DesiredMovementDir.Rotation().Yaw;
+					const float DeltaYaw = DesiredYaw - GetActorRotation().Yaw;
 
-						const float DesiredPitch = DesiredMovementDir.Rotation().Pitch;
-						const float DeltaPitch = DesiredPitch - GetActorRotation().Pitch;
+					const float DesiredPitch = DesiredMovementDir.Rotation().Pitch;
+					const float DeltaPitch = DesiredPitch - GetActorRotation().Pitch;
 
-						const float DesiredRoll = DesiredMovementDir.Rotation().Roll;
-						const float DeltaRoll = DesiredRoll - GetActorRotation().Roll;
+					const float DesiredRoll = DesiredMovementDir.Rotation().Roll;
+					const float DeltaRoll = DesiredRoll - GetActorRotation().Roll;
 
-						// Kind of gross but because we want "instant" turning we must factor in delta time so that it gets factored out inside the simulation
-						RotationInputYaw = DeltaYaw / NextCmd->FrameDeltaTime;
-						RotationInputPitch = DeltaPitch / NextCmd->FrameDeltaTime;
-						RotationInputRoll = DeltaRoll / NextCmd->FrameDeltaTime;
+					// Kind of gross but because we want "instant" turning we must factor in delta time so that it gets factored out inside the simulation
+					RotationInputYaw = DeltaYaw / DeltaTimeSeconds;
+					RotationInputPitch = DeltaPitch / DeltaTimeSeconds;
+					RotationInputRoll = DeltaRoll / DeltaTimeSeconds;
 								
-					}
+				}
 
-					NextCmd->RotationInput.Yaw = RotationInputYaw;
-					NextCmd->RotationInput.Pitch = RotationInputPitch;
-					NextCmd->RotationInput.Roll = RotationInputRoll;
+				Cmd.RotationInput.Yaw = RotationInputYaw;
+				Cmd.RotationInput.Pitch = RotationInputPitch;
+				Cmd.RotationInput.Roll = RotationInputRoll;
 						
-					NextCmd->MovementInput = FVector(FMath::Clamp<float>( MoveInputMag, -1.f, 1.f ), 0.0f, 0.0f); // Not the best way but simple
-				}
-				break;
+				Cmd.MovementInput = FVector(FMath::Clamp<float>( MoveInputMag, -1.f, 1.f ), 0.0f, 0.0f); // Not the best way but simple
 			}
-			case 3:
+			break;
+		}
+		case 3:
+		{
+			// Free camera on the yaw, camera-relative motion
+			APlayerController* PC = Cast<APlayerController>(Controller);
+			if (ensure(PC)) // Requires player controller for now
 			{
-				// Free camera on the yaw, camera-relative motion
-				APlayerController* PC = Cast<APlayerController>(Controller);
-				if (ensure(PC)) // Requires player controller for now
-				{
-					// Camera yaw rotation
-					PC->AddYawInput(CachedLookInput.X * ControllerLookRateYaw );
+				// Camera yaw rotation
+				PC->AddYawInput(CachedLookInput.X * ControllerLookRateYaw );
 
-					// Rotational movement: orientate us towards our camera-relative desired velocity (unless we are upside down, then flip it)
-					static float RotationMagMin = (1e-3);
-					const float MoveInputMag = CachedMoveInput.Size2D();
+				// Rotational movement: orientate us towards our camera-relative desired velocity (unless we are upside down, then flip it)
+				static float RotationMagMin = (1e-3);
+				const float MoveInputMag = CachedMoveInput.Size2D();
 						
-					float RotationInputYaw = 0.f;
-					if (MoveInputMag >= RotationMagMin)
+				float RotationInputYaw = 0.f;
+				if (MoveInputMag >= RotationMagMin)
+				{
+					FVector DesiredMovementDir = PC->GetControlRotation().RotateVector(CachedMoveInput);
+
+					const bool bIsUpsideDown = FVector::DotProduct(FVector(0.f, 0.f, 1.f), GetActorQuat().GetUpVector() ) < 0.f;
+					if (bIsUpsideDown)
 					{
-						FVector DesiredMovementDir = PC->GetControlRotation().RotateVector(CachedMoveInput);
+						DesiredMovementDir *= -1.f;
+					}							
 
-						const bool bIsUpsideDown = FVector::DotProduct(FVector(0.f, 0.f, 1.f), GetActorQuat().GetUpVector() ) < 0.f;
-						if (bIsUpsideDown)
-						{
-							DesiredMovementDir *= -1.f;
-						}							
+					const float DesiredYaw = DesiredMovementDir.Rotation().Yaw;
+					const float DeltaYaw = DesiredYaw - GetActorRotation().Yaw;
 
-						const float DesiredYaw = DesiredMovementDir.Rotation().Yaw;
-						const float DeltaYaw = DesiredYaw - GetActorRotation().Yaw;
-
-						// Kind of gross but because we want "instant" turning we must factor in delta time so that it gets factored out inside the simulation
-						RotationInputYaw = DeltaYaw / NextCmd->FrameDeltaTime;
-					}
+					// Kind of gross but because we want "instant" turning we must factor in delta time so that it gets factored out inside the simulation
+					RotationInputYaw = DeltaYaw / DeltaTimeSeconds;
+				}
 
-					NextCmd->RotationInput.Yaw = RotationInputYaw;
-					NextCmd->RotationInput.Pitch = CachedLookInput.Y * LookRatePitch; // Just pitch like normal
-					NextCmd->RotationInput.Roll = 0;
+				Cmd.RotationInput.Yaw = RotationInputYaw;
+				Cmd.RotationInput.Pitch = CachedLookInput.Y * LookRatePitch; // Just pitch like normal
+				Cmd.RotationInput.Roll = 0;
 						
-					NextCmd->MovementInput = FVector(FMath::Clamp<float>( MoveInputMag, -1.f, 1.f ), 0.0f, CachedMoveInput.Z); // Not the best way but simple
-				}
-				break;
+				Cmd.MovementInput = FVector(FMath::Clamp<float>( MoveInputMag, -1.f, 1.f ), 0.0f, CachedMoveInput.Z); // Not the best way but simple
 			}
+			break;
 		}
+	}
 				
 
-		CachedMoveInput = FVector::ZeroVector;
-		CachedLookInput = FVector2D::ZeroVector;
-	}
+	CachedMoveInput = FVector::ZeroVector;
+	CachedLookInput = FVector2D::ZeroVector;
 }
\ No newline at end of file
diff --git a/Engine/Plugins/Runtime/NetworkPredictionExtras/Source/NetworkPredictionExtras/Public/MockNetworkSimulation.h b/Engine/Plugins/Runtime/NetworkPredictionExtras/Source/NetworkPredictionExtras/Public/MockNetworkSimulation.h
--- a/Engine/Plugins/Runtime/NetworkPredictionExtras/Source/NetworkPredictionExtras/Public/MockNetworkSimulation.h
+++ b/Engine/Plugins/Runtime/NetworkPredictionExtras/Source/NetworkPredictionExtras/Public/MockNetworkSimulation.h
@@ -51,12 +51,10 @@ class IMockNetworkSimulationDriver;
 // State the client generates
 struct FMockInputCmd
 {
-	float FrameDeltaTime=0.f;
 	float InputValue=0;
 
 	void NetSerialize(const FNetSerializeParams& P)
 	{
-		P.Ar << FrameDeltaTime;
 		P.Ar << InputValue;
 	}
 
@@ -64,7 +62,6 @@ struct FMockInputCmd
 	{
 		if (P.Context == EStandardLoggingContext::Full)
 		{
-			P.Ar->Logf(TEXT("FrameDeltaTime: %4f"), FrameDeltaTime);
 			P.Ar->Logf(TEXT("InputValue: %4f"), InputValue);
 		}
 		else
@@ -133,7 +130,7 @@ public:
 	};
 
 	/** Main update function */
-	static void Update(IMockNetworkSimulationDriver* Driver, const FMockInputCmd& InputCmd, const FMockSyncState& InputState, FMockSyncState& OutputState, const FMockAuxState& AuxState);
+	static void Update(IMockNetworkSimulationDriver* Driver, const TSimTime& DeltaTimeMS, const FMockInputCmd& InputCmd, const FMockSyncState& InputState, FMockSyncState& OutputState, const FMockAuxState& AuxState);
 };
 
 // Needed to trick UHT into letting UMockNetworkSimulationComponent implement. UHT cannot parse the ::
@@ -164,8 +161,11 @@ private:
 	TUniquePtr<FMockNetworkSimulation> NetworkSim;
 	
 public:
+
 	void InitSyncState(FMockSyncState& OutSyncState) const override;
 	void FinalizeFrame(const FMockSyncState& SyncState) override;
+	void ProduceInput(const FMockNetworkSimulation::TSimTime& SimFrameTime, FMockInputCmd& Cmd);
+
 	virtual UWorld* GetDriverWorld() const override final { return GetWorld(); }
 	virtual UObject* GetVLogOwner() const override final;
 	virtual FTransform GetDebugWorldTransform() const override final;
diff --git a/Engine/Plugins/Runtime/NetworkPredictionExtras/Source/NetworkPredictionExtras/Public/NetworkPredictionExtrasFlyingPawn.h b/Engine/Plugins/Runtime/NetworkPredictionExtras/Source/NetworkPredictionExtras/Public/NetworkPredictionExtrasFlyingPawn.h
--- a/Engine/Plugins/Runtime/NetworkPredictionExtras/Source/NetworkPredictionExtras/Public/NetworkPredictionExtrasFlyingPawn.h
+++ b/Engine/Plugins/Runtime/NetworkPredictionExtras/Source/NetworkPredictionExtras/Public/NetworkPredictionExtrasFlyingPawn.h
@@ -2,6 +2,7 @@
 #pragma once
 
 #include "GameFramework/Pawn.h"
+#include "Movement/FlyingMovement.h"
 #include "NetworkPredictionExtrasFlyingPawn.generated.h"
 
 class UInputComponent;
@@ -44,7 +45,7 @@ class NETWORKPREDICTIONEXTRAS_API ANetworkPredictionExtrasFlyingPawn : public AP
 
 private:
 
-	void GenerateLocalInput(float DeltaSeconds);
+	void ProduceInput(const FlyingMovement::TSimTime& SimFrameTime, FlyingMovement::FInputCmd& Cmd);
 
 	FVector CachedMoveInput;
 	FVector2D CachedLookInput;
