NetworkPrediction: continue wip. Fixed up some issues with new time keeping structures. Server side seems to be working again but auto and simulated proxies still have some issues.

#rnx

diff --git a/Engine/Plugins/Runtime/NetworkPrediction/Source/NetworkPrediction/Private/Movement/FlyingMovementComponent.cpp b/Engine/Plugins/Runtime/NetworkPrediction/Source/NetworkPrediction/Private/Movement/FlyingMovementComponent.cpp
--- a/Engine/Plugins/Runtime/NetworkPrediction/Source/NetworkPrediction/Private/Movement/FlyingMovementComponent.cpp
+++ b/Engine/Plugins/Runtime/NetworkPrediction/Source/NetworkPrediction/Private/Movement/FlyingMovementComponent.cpp
@@ -114,8 +114,9 @@ void UFlyingMovementComponent::TickComponent(float DeltaTime, enum ELevelTick Ti
 		}
 
 		FlyingMovement::FMovementSystem::FTickParameters Parameters;
-		Parameters.Role = OwnerRole;
 		Parameters.LocalDeltaTimeSeconds = DeltaTime;
+		Parameters.Role = OwnerRole;
+		Parameters.bIsRemotelyControlled = (OwnerRole == ROLE_Authority && GetOwner()->GetNetConnection() != nullptr);
 
 		// Tick the core network sim, this will consume input and generate new sync state
 		NetworkSim->Tick((IFlyingMovementDriver*)this, Parameters);
diff --git a/Engine/Plugins/Runtime/NetworkPrediction/Source/NetworkPrediction/Private/Movement/ParametricMovement.cpp b/Engine/Plugins/Runtime/NetworkPrediction/Source/NetworkPrediction/Private/Movement/ParametricMovement.cpp
--- a/Engine/Plugins/Runtime/NetworkPrediction/Source/NetworkPrediction/Private/Movement/ParametricMovement.cpp
+++ b/Engine/Plugins/Runtime/NetworkPrediction/Source/NetworkPrediction/Private/Movement/ParametricMovement.cpp
@@ -206,8 +206,9 @@ void UParametricMovementComponent::TickSimulation(float DeltaTimeSeconds)
 	if (NetworkSim && OwnerRole != ROLE_None)
 	{
 		ParametricMovement::FMovementSystem::FTickParameters Parameters;
-		Parameters.Role = OwnerRole;
 		Parameters.LocalDeltaTimeSeconds = DeltaTimeSeconds;
+		Parameters.Role = OwnerRole;
+		Parameters.bIsRemotelyControlled = (OwnerRole == ROLE_Authority && GetOwner()->GetNetConnection() != nullptr);
 
 		// Tick the core network sim, this will consume input and generate new sync state
 		NetworkSim->Tick((IParametricMovementDriver*)this, Parameters);
diff --git a/Engine/Plugins/Runtime/NetworkPrediction/Source/NetworkPrediction/Public/Movement/FlyingMovement.h b/Engine/Plugins/Runtime/NetworkPrediction/Source/NetworkPrediction/Public/Movement/FlyingMovement.h
--- a/Engine/Plugins/Runtime/NetworkPrediction/Source/NetworkPrediction/Public/Movement/FlyingMovement.h
+++ b/Engine/Plugins/Runtime/NetworkPrediction/Source/NetworkPrediction/Public/Movement/FlyingMovement.h
@@ -107,7 +107,7 @@ namespace FlyingMovement
 	using TMovementBufferTypes = TNetworkSimBufferTypes<FInputCmd, FMoveState, FAuxState>;
 
 	// Actual definition of our network simulation.
-	class FMovementSystem : public TNetworkedSimulationModel<FMovementSystem, TMovementBufferTypes>
+	class FMovementSystem : public TNetworkedSimulationModel<FMovementSystem, TMovementBufferTypes, TNetworkSimTickSettings<0>>
 	{
 	public:
 
diff --git a/Engine/Plugins/Runtime/NetworkPrediction/Source/NetworkPrediction/Public/Movement/ParametricMovement.h b/Engine/Plugins/Runtime/NetworkPrediction/Source/NetworkPrediction/Public/Movement/ParametricMovement.h
--- a/Engine/Plugins/Runtime/NetworkPrediction/Source/NetworkPrediction/Public/Movement/ParametricMovement.h
+++ b/Engine/Plugins/Runtime/NetworkPrediction/Source/NetworkPrediction/Public/Movement/ParametricMovement.h
@@ -88,7 +88,7 @@ namespace ParametricMovement
 	using TMovementBufferTypes = TNetworkSimBufferTypes<FInputCmd, FMoveState, FAuxState>;
 
 	// Actual definition of our network simulation.
-	class FMovementSystem : public TNetworkedSimulationModel<FMovementSystem, TMovementBufferTypes, TNetworkSimTickSettings<30>>
+	class FMovementSystem : public TNetworkedSimulationModel<FMovementSystem, TMovementBufferTypes, TNetworkSimTickSettings<0>>
 	{
 	public:
 
diff --git a/Engine/Plugins/Runtime/NetworkPrediction/Source/NetworkPrediction/Public/NetworkSimulationModel.h b/Engine/Plugins/Runtime/NetworkPrediction/Source/NetworkPrediction/Public/NetworkSimulationModel.h
--- a/Engine/Plugins/Runtime/NetworkPrediction/Source/NetworkPrediction/Public/NetworkSimulationModel.h
+++ b/Engine/Plugins/Runtime/NetworkPrediction/Source/NetworkPrediction/Public/NetworkSimulationModel.h
@@ -55,8 +55,14 @@ public:
 
 	struct FTickParameters
 	{
-		ENetRole Role;
-		float LocalDeltaTimeSeconds;
+		// Owner's role. Necessary to know which proxy we should be forwarding functions in tick to
+		ENetRole Role = ROLE_None;
+
+		// (Only used when ROLE_Authorty: Whether we are expecting a remote client to control. e.g, there is an autoproxy client driving this)
+		bool bIsRemotelyControlled = false;
+
+		// Local fraem delta time seconds. Just passed in from ::Tick
+		typename TTickSettings::TRealTime LocalDeltaTimeSeconds = 0.f;
 	};
 
 	template<typename TDriver>
@@ -108,28 +114,36 @@ public:
 			case ROLE_SimulatedProxy:
 				RepProxy_Simulated.template PreSimTick<T, TDriver>(Driver, Buffers, TickInfo, Parameters.LocalDeltaTimeSeconds);
 			break;
-		}		
+		}
 
-		switch (Parameters.Role) {
 
-		
-		case ROLE_AutonomousProxy:
+		if (Parameters.Role == ROLE_Authority)
 		{
-			// Check SyncBuffer being ahead of processed Keyframes. This would happen in cases where we are either not predicting or are buffering our input locally
-			// while sending latest cmds to the server. Essentially, we got the authoritative motion state from the server before we ran the simulation locally.
-			if (Buffers.Sync.GetHeadKeyframe() > TickInfo.LastProcessedInputKeyframe)
-			{
-				checkf(Buffers.Input.IsValidKeyframe(Buffers.Sync.GetHeadKeyframe()), TEXT("MotionState and InputCmd buffers are out of system. LastProcessedInputKeyframe: %d {%s} vs {%s}"),
-					TickInfo.LastProcessedInputKeyframe, *Buffers.Sync.GetBasicDebugStr(), *Buffers.Input.GetBasicDebugStr());
-
-				UE_LOG(LogNetworkSim, Warning, TEXT("Skipping local input frames because we have newer data in SyncBuffer. LastProcessedInputKeyframe: %d. {%s} {%s}"),
-					TickInfo.LastProcessedInputKeyframe, *Buffers.Sync.GetBasicDebugStr(), *Buffers.Input.GetBasicDebugStr());
+			// Always give simulation time on the authority
+			TickInfo.GiveSimulationTime(Parameters.LocalDeltaTimeSeconds);
 
-				TickInfo.LastProcessedInputKeyframe = Buffers.Sync.GetHeadKeyframe();
+			TNetworkSimTime<TTickSettings> DeltaSimTime = TickInfo.GetRemaningAllowedSimulationTime();
+			if (DeltaSimTime.ToRealTimeSeconds() > 0)
+			{
+				if (Parameters.bIsRemotelyControlled)
+				{
+					TickInfo.MaxAllowedInputKeyframe = Buffers.Input.GetHeadKeyframe();
+				}
+				else
+				{
+					if (TInputCmd* InputCmd = Buffers.Input.GetWriteNext())
+					{
+						*InputCmd = TInputCmd();
+						InputCmd->SetFrameDeltaTime(DeltaSimTime);
+						Driver->ProduceInput(DeltaSimTime, *InputCmd);
+						TickInfo.MaxAllowedInputKeyframe++;
+					}
+				}
 			}
-			break;
 		}
 
+		switch (Parameters.Role) {
+
 		case ROLE_Authority:
 		{
 			if ( TickInfo.LastProcessedInputKeyframe+1 < Buffers.Input.GetTailKeyframe() )
@@ -141,12 +155,6 @@ public:
 			break;
 		}
 
-		case ROLE_SimulatedProxy:
-		{
-			RepProxy_Simulated.template Reconcile<T, TDriver>(Driver, Buffers, TickInfo);
-			break;
-		}
-
 		} // end switch
 
 		// -------------------------------------------------------------------------------------------------------------------------------------------------
diff --git a/Engine/Plugins/Runtime/NetworkPrediction/Source/NetworkPrediction/Public/NetworkSimulationModelDebugger.h b/Engine/Plugins/Runtime/NetworkPrediction/Source/NetworkPrediction/Public/NetworkSimulationModelDebugger.h
--- a/Engine/Plugins/Runtime/NetworkPrediction/Source/NetworkPrediction/Public/NetworkSimulationModelDebugger.h
+++ b/Engine/Plugins/Runtime/NetworkPrediction/Source/NetworkPrediction/Public/NetworkSimulationModelDebugger.h
@@ -434,12 +434,12 @@ struct TNetworkSimulationModelDebugger : public INetworkSimulationModelDebugger
 
 			float MinWidth = TextSizing.DrawXL;
 			float MinHeight = TextSizing.DrawYL;
-			static float MinHeightMS = 1/60.f * 1000.f;	// 60hz frame is drawn at MinHeight.
+			static float MinHeightSeconds = 1/60.f;	// 60hz frame is drawn at MinHeight.
 
-			auto CalcHeight = [&](TSimTime MS)
+			auto CalcHeight = [&](float RealTimeSeconds)
 			{
-				const float Ratio = MinHeight / MinHeightMS;
-				return MS.ToRealTimeSeconds() * Ratio;
+				const float Ratio = MinHeight / MinHeightSeconds;
+				return RealTimeSeconds * Ratio;
 			};
 
 			for (auto It = DebugBuffer->CreateIterator(); It; ++It)
@@ -450,7 +450,7 @@ struct TNetworkSimulationModelDebugger : public INetworkSimulationModelDebugger
 				float ScreenWidth = MinWidth;
 				float FramePct = FMath::Clamp<float>((DebugState->LocalDeltaTimeSeconds - LocalFrameTimeRed) / ( LocalFrameTimeGreen - LocalFrameTimeRed ), 0.f, 1.f);
 				FColor Color = FColor::MakeRedToGreenColorFromScalar( FramePct );
-				float ServerHeight = CalcHeight(DebugState->LocalDeltaTimeSeconds * 1000.f);
+				float ServerHeight = CalcHeight(DebugState->LocalDeltaTimeSeconds);
 
 				Out.EmitQuad(FVector2D( ScreenX, ScreenY), FVector2D( ScreenWidth, ServerHeight), Color);
 				Out.EmitText(FVector2D( ScreenX, ScreenY), FColor::Black, FString::Printf(TEXT("%.2f"), (DebugState->LocalDeltaTimeSeconds * 1000.f)));
@@ -466,7 +466,7 @@ struct TNetworkSimulationModelDebugger : public INetworkSimulationModelDebugger
 						ClientSimTime += Cmd->GetFrameDeltaTime();
 						
 						float ClientSizeX = MinWidth;
-						float ClientSizeY =  CalcHeight(Cmd->GetFrameDeltaTime());
+						float ClientSizeY =  CalcHeight(Cmd->GetFrameDeltaTime().ToRealTimeSeconds());
 
 						FVector2D ScreenPos(ClientX, ClientY - ClientSizeY);
 						Out.EmitQuad(ScreenPos, FVector2D( ClientSizeX, ClientSizeY), FColor::Blue);
@@ -480,7 +480,7 @@ struct TNetworkSimulationModelDebugger : public INetworkSimulationModelDebugger
 					if (auto* Cmd = InputBuffer.FindElementByKeyframe(Keyframe))
 					{
 						float ClientSizeX = MinWidth;
-						float ClientSizeY = CalcHeight(Cmd->GetFrameDeltaTime());
+						float ClientSizeY = CalcHeight(Cmd->GetFrameDeltaTime().ToRealTimeSeconds());
 						FVector2D ScreenPos(ClientX, ClientY - ClientSizeY);
 						Out.EmitQuad(ScreenPos, FVector2D( ClientSizeX, ClientSizeY), FColor::Red);
 						Out.EmitText(ScreenPos, FColor::White, LexToString(Keyframe));
diff --git a/Engine/Plugins/Runtime/NetworkPrediction/Source/NetworkPrediction/Public/NetworkSimulationModelReplicators.h b/Engine/Plugins/Runtime/NetworkPrediction/Source/NetworkPrediction/Public/NetworkSimulationModelReplicators.h
--- a/Engine/Plugins/Runtime/NetworkPrediction/Source/NetworkPrediction/Public/NetworkSimulationModelReplicators.h
+++ b/Engine/Plugins/Runtime/NetworkPrediction/Source/NetworkPrediction/Public/NetworkSimulationModelReplicators.h
@@ -338,34 +338,16 @@ struct TReplicator_BasicReconciliar : public TBase
 			}
 			else
 			{
-				// We don't have corresponding local state. There are two cases:
-				if (NetworkSimulationModelCVars::EnableLocalPrediction) // Fixme: this is awkward. Expected if player prediction is disabled but coupling like this feels bad.
-				{
-					UE_LOG(LogNetworkSim, Warning, TEXT("::NetSerialize Fault: SyncBuffer does not contain data for frame %d. [%d-%d]"), SerializedHeadKeyframe, Buffers.Sync.GetTailKeyframe(), Buffers.Sync.GetHeadKeyframe());
-				}
-
-				// Case 1: the serialized state is older than what we've kept in our buffer. A bigger buffer would solve this! (at the price of more resimulated frames to recover when this happens)
 				if (SerializedHeadKeyframe < Buffers.Sync.GetTailKeyframe())
 				{
+					// Case 1: the serialized state is older than what we've kept in our buffer. A bigger buffer would solve this! (at the price of more resimulated frames to recover when this happens)
 					// This is a reconcile fault and we just need to chill. We'll stop sending user commands until the cmds in flight flush through the system and we catch back up.
 					bReconcileFaultDetected = true;
 				}
 				else
 				{
-					// Case 2: We've received a newer frame than what we've processed locally. This could happen if we are buffering our inputs locally (but still sending to the server).
-					// Since this doesn't require resimulating, we can just set the authoritative state here and continue on. 
-					
-					
-
-					// However, this keyframe better be a valid input. Otherwise how did the server generate it?
-					checkf(SerializedHeadKeyframe <= Buffers.Input.GetHeadKeyframe(), TEXT("Received newer motionstate that doesn't correspond to valid input cmd. SerializedHeadKeyframe: %d. {%s} {%s}"),
-						SerializedHeadKeyframe, *Buffers.Input.GetBasicDebugStr(), *Buffers.Sync.GetBasicDebugStr());
-
-					Buffers.Sync.ResetNextHeadKeyframe(SerializedHeadKeyframe);
-					TSyncState* ClientMotionState = Buffers.Sync.GetWriteNext();
-					*ClientMotionState = *SerializedState;
-
-					//!!TickInfo.LastProcessedInputKeyframe = SerializedHeadKeyframe;
+					// Case 2: We've received a newer frame than what we've processed locally. This could happen if we are buffering our inputs locally (but still sending to the server) or just not predicting
+					bPendingReconciliation =  true;
 				}
 			}
 		}
@@ -397,30 +379,38 @@ struct TReplicator_BasicReconciliar : public TBase
 			return;
 		}
 
-		TSyncState* ClientSyncState = Buffers.Sync.FindElementByKeyframe( ReconciliationKeyframe );
-
-		// Should not be here if the client doesn't have a valid state. See ::NetSerialize
-		checkf(ClientSyncState != nullptr, TEXT("SyncBuffer does not contain data for frame %d. %s"), ReconciliationKeyframe, *Buffers.Sync.GetBasicDebugStr());
-
 		// -------------------------------------------------------------------------------------------------------------------------
 		// Resimulate
 		// -------------------------------------------------------------------------------------------------------------------------
 
+		TSyncState* ClientSyncState = Buffers.Sync.FindElementByKeyframe( ReconciliationKeyframe );	
+
 		ServerState->VisualLog( FVisualLoggingParameters(EVisualLoggingContext::LastConfirmed, ReconciliationKeyframe, EVisualLoggingLifetime::Persistent), Driver, Driver);
-		ClientSyncState->VisualLog( FVisualLoggingParameters(EVisualLoggingContext::FirstMispredicted, ReconciliationKeyframe, EVisualLoggingLifetime::Persistent), Driver, Driver);
 
-		TSyncState* LatestClient = Buffers.Sync.GetElementFromHead(0);
+		if (ClientSyncState)
+		{
+			// Existing ClientSyncState, log it before overwriting it
+			ClientSyncState->VisualLog( FVisualLoggingParameters(EVisualLoggingContext::FirstMispredicted, ReconciliationKeyframe, EVisualLoggingLifetime::Persistent), Driver, Driver);
+		}
+		else
+		{
+			// No existing state, so create add it explicitly
+			Buffers.Sync.ResetNextHeadKeyframe( ReconciliationKeyframe );
+			ClientSyncState = Buffers.Sync.GetWriteNext();
+		}
 
-		// Copy authoritative state over client state (FIXME: we may want to store this off historically somewhere? Or will VLog be enough for debugging?)
+		// Set client's sync state to the server version
+		check(ClientSyncState);
 		*ClientSyncState = *ServerState;
 
 		// Set the canonical simulation time to what we received (we will advance it as we resimulate)
 		TickInfo.TotalProcessedSimulationTime = SerializedTime;
 		TickInfo.LastProcessedInputKeyframe = ReconciliationKeyframe;
+		TickInfo.MaxAllowedInputKeyframe = FMath::Max(TickInfo.MaxAllowedInputKeyframe, TickInfo.LastProcessedInputKeyframe); // Make sure this doesn't lag behind. This is the only place we should need to do this.
 		
-		// Resimulate up to our latest SyncBuffer frame. Note that we may have unprocessed user commands in the command buffer at this point.
-		const int32 LatestKeyframe = Buffers.Sync.GetHeadKeyframe();		
-		for (int32 Keyframe = ReconciliationKeyframe+1; Keyframe <= LatestKeyframe; ++Keyframe)
+		// Resimulate all user commands 
+		const int32 LastKeyframeToProcess = TickInfo.MaxAllowedInputKeyframe;
+		for (int32 Keyframe = ReconciliationKeyframe+1; Keyframe <= LastKeyframeToProcess; ++Keyframe)
 		{
 			// Keyframe is the frame we are resimulating right now.
 			TInputCmd* ResimulateCmd	= Buffers.Input.FindElementByKeyframe(Keyframe);
@@ -442,7 +432,7 @@ struct TReplicator_BasicReconciliar : public TBase
 			// ------------------------------------------------------
 
 			// Log out the Mispredicted state that we are about to overwrite.
-			NextMotionState->VisualLog( FVisualLoggingParameters(Keyframe == LatestKeyframe ? EVisualLoggingContext::LastMispredicted : EVisualLoggingContext::OtherMispredicted, Keyframe, EVisualLoggingLifetime::Persistent), Driver, Driver);
+			NextMotionState->VisualLog( FVisualLoggingParameters(Keyframe == LastKeyframeToProcess ? EVisualLoggingContext::LastMispredicted : EVisualLoggingContext::OtherMispredicted, Keyframe, EVisualLoggingLifetime::Persistent), Driver, Driver);
 
 			// Do the actual update
 			T::Update(Driver, ResimulateCmd->GetFrameDeltaTime(), *ResimulateCmd, *PrevMotionState, *NextMotionState, *AuxState);
@@ -452,7 +442,7 @@ struct TReplicator_BasicReconciliar : public TBase
 			TickInfo.LastProcessedInputKeyframe = Keyframe;
 
 			// Log out the newly predicted state that we got.
-			NextMotionState->VisualLog( FVisualLoggingParameters(Keyframe == LatestKeyframe ? EVisualLoggingContext::LastPredicted : EVisualLoggingContext::OtherPredicted, Keyframe, EVisualLoggingLifetime::Persistent), Driver, Driver);
+			NextMotionState->VisualLog( FVisualLoggingParameters(Keyframe == LastKeyframeToProcess ? EVisualLoggingContext::LastPredicted : EVisualLoggingContext::OtherPredicted, Keyframe, EVisualLoggingLifetime::Persistent), Driver, Driver);
 		}
 	}
 
diff --git a/Engine/Plugins/Runtime/NetworkPrediction/Source/NetworkPrediction/Public/NetworkSimulationModelTypes.h b/Engine/Plugins/Runtime/NetworkPrediction/Source/NetworkPrediction/Public/NetworkSimulationModelTypes.h
--- a/Engine/Plugins/Runtime/NetworkPrediction/Source/NetworkPrediction/Public/NetworkSimulationModelTypes.h
+++ b/Engine/Plugins/Runtime/NetworkPrediction/Source/NetworkPrediction/Public/NetworkSimulationModelTypes.h
@@ -176,6 +176,7 @@ struct TNetworkSimTickSettings
 	static constexpr TUnderlingSimTime GetMaxStepMS() { return static_cast<TUnderlingSimTime>(InMaxStepMS); }
 	static constexpr TUnderlingSimTime GetFixedStepMS() { return static_cast<TUnderlingSimTime>(InFixedStepMS); }
 	static constexpr TRealTime GetRealToSimFactor() { return static_cast<TRealTime>(InRealToSimFactor); }
+	static constexpr TRealTime GetSimToRealFactor() { return static_cast<TRealTime>(1.f / InRealToSimFactor); }
 };
 
 // Actual time value. This by default will store time in MSec (ultimately determined by TickSettings::RealToSimFactor)
@@ -186,10 +187,12 @@ struct TNetworkSimTime
 	using TRealTime = typename TickSettings::TRealTime;
 
 	TNetworkSimTime() { }
-	TNetworkSimTime(const TTime& InTime) : Time(InTime) { }
-	TNetworkSimTime(const TRealTime& InRealTime) : Time(InRealTime * TickSettings::RealToSimFactor) { }
 
-	TRealTime ToRealTimeSeconds() const { return (Time / TickSettings::RealToSimFactor); }
+	// Things get confusing with templated types and overloaded functions. To avoid that, use these funcs to construct from either msec or real time
+	static inline TNetworkSimTime FromMSec(const TTime& InTime) { return TNetworkSimTime(InTime); } 
+	static inline TNetworkSimTime FromRealTimeSeconds(const TRealTime& InRealTime) { return TNetworkSimTime( static_cast<TTime>(InRealTime * TickSettings::GetRealToSimFactor())); } 
+	
+	TRealTime ToRealTimeSeconds() const { return (Time * TickSettings::GetSimToRealFactor()); }
 	FString ToString() const { return LexToString(this->Time); }
 
 	// FIXME
@@ -199,7 +202,7 @@ struct TNetworkSimTime
 	T& operator+= (const T &rhs) { this->Time += rhs.Time; return(*this); }
 	T& operator-= (const T &rhs) { this->Time -= rhs.Time; return(*this); }
 	
-	T operator+ (const T &rhs) const { return T(this->Time - rhs.Time); }
+	T operator+ (const T &rhs) const { return T(this->Time + rhs.Time); }
 	T operator- (const T &rhs) const { return T(this->Time - rhs.Time); }
 
 	bool operator<  (const T &rhs) const { return(this->Time < rhs.Time); }
@@ -210,6 +213,9 @@ struct TNetworkSimTime
 	bool operator!= (const T &rhs) const { return(this->Time != rhs.Time); }
 
 	TTime Time = 0;
+
+private:
+	TNetworkSimTime(const TTime& InTime) { Time = InTime; }
 };
 
 // ----------------------------------------------------------------------------------------------------------------------------------------------
@@ -221,11 +227,12 @@ template<typename TickSettings, bool IsFixedTick=(TickSettings::FixedStepMS!=0)>
 struct TRealTimeAccumulator
 {
 	using TRealTime = typename TickSettings::TRealTime;
-	void Accumulate(TNetworkSimTime<TickSettings>& NetworkSimTime, TRealTime RealTimeSeconds)
+	void Accumulate(TNetworkSimTime<TickSettings>& NetworkSimTime, const TRealTime RealTimeSeconds)
 	{
 		// Non fixed step: just accumulate the time directly
 		// Note that MaxStepMS enforcement does NOT belong here. Dropping time due to MaxStepMS would just make the sim run slower. MaxStepMS is used at the input processing level.
-		NetworkSimTime += RealTimeSeconds;
+		// (wondering: does converting down to MSec make sense in variable rates? We aren't really gaining anything here. Specializing to avoid the multiply probably not worth the complexity though)
+		NetworkSimTime += TNetworkSimTime<TickSettings>::FromRealTimeSeconds(RealTimeSeconds);
 	}
 };
 
@@ -234,9 +241,9 @@ template<typename TickSettings>
 struct TRealTimeAccumulator<TickSettings, true>
 {
 	using TRealTime = typename TickSettings::TRealTime;
-	const TRealTime RealTimeFixedStep = TNetworkSimTime<TickSettings>(TickSettings::GetFixedStepMS() / TickSettings::GetRealToSimFactor()).ToRealTimeSeconds();
+	const TRealTime RealTimeFixedStep = static_cast<TRealTime>(TickSettings::GetFixedStepMS() * TickSettings::GetSimToRealFactor());
 
-	void Accumulate(TNetworkSimTime<TickSettings>& NetworkSimTime, TRealTime RealTimeSeconds)
+	void Accumulate(TNetworkSimTime<TickSettings>& NetworkSimTime, const TRealTime RealTimeSeconds)
 	{
 		AccumulatedTime += RealTimeSeconds;
 		if (AccumulatedTime > RealTimeFixedStep)
@@ -249,7 +256,7 @@ struct TRealTimeAccumulator<TickSettings, true>
 				AccumulatedTime = TRealTime(0.f);
 			}
 
-			NetworkSimTime += TNetworkSimTime<TickSettings>(NumFrames * TickSettings::GetFixedStepMS());
+			NetworkSimTime += TNetworkSimTime<TickSettings>::FromMSec(NumFrames * TickSettings::FixedStepMS);
 		}
 	}
 	TRealTime AccumulatedTime;
@@ -268,7 +275,6 @@ struct TSimulationTickState
 	int32 LastProcessedInputKeyframe;	// The last input keyframe that we processed
 	int32 MaxAllowedInputKeyframe;		// The max input keyframe that we are allowed to process (e.g, don't process input past this keyframe yet)
 
-
 	uint32 LastLocalInputGFrameNumber;	// Tracks the last time we accepted local input via GetNextInputForWrite. Used to guard against accidental input latency due to ordering of input/sim ticking.
 
 	TNetworkSimTime<TSettings> TotalAllowedSimulationTime;	// Total time we have been "given" to process. We cannot process more simulation time than this: doing so would be speed hacking.
@@ -327,7 +333,7 @@ private:
 template<typename BaseCmdType, typename TickSettings>
 struct TFrameCmd<BaseCmdType, TickSettings, true> : public BaseCmdType
 {
-	constexpr TNetworkSimTime<TickSettings> GetFrameDeltaTime() const { return TNetworkSimTime<TickSettings>(TickSettings::GetFixedStepMS()); }
+	constexpr TNetworkSimTime<TickSettings> GetFrameDeltaTime() const { return TNetworkSimTime<TickSettings>::FromMSec(TickSettings::GetFixedStepMS()); }
 	void SetFrameDeltaTime(const TNetworkSimTime<TickSettings>& InTime) { }
 };
 
diff --git a/Engine/Plugins/Runtime/NetworkPrediction/readme.txt b/Engine/Plugins/Runtime/NetworkPrediction/readme.txt
--- a/Engine/Plugins/Runtime/NetworkPrediction/readme.txt
+++ b/Engine/Plugins/Runtime/NetworkPrediction/readme.txt
@@ -18,7 +18,7 @@ High level goals:
  
 What this is NOT:
 -We are not moving to a global fixed ticking system. Though these systems have huge advantages, we don't think it is a good universal solution for the engine.
--That said, we recognize that it is the right solution for many things and we want to enable support for that. Nothing in this system /requires/ dynamic ticking, but it is something we feel the need to support.
+-That said, we recognize that it is the right solution for many things and we want to enable support for that. Our goal is that the user simulation code is independent of ticking being fixed or not.
 -To be clear: widely different, variable, fluctuating tick rates between clients and server is basically the root problems we are fighting here.
 -Related, we are also not all in on a "predict everything: seamlessly and automatically" solution. We still feel player prediction is best kept to a minimum (meaning: predict the minimum amount of stuff you can get away with).
 -But, we are trying to make it way easier to write predictive gameplay systems. Easier to write, to debug and reason about.
@@ -94,7 +94,13 @@ Smoothing: Taking the output of the simulation and applying an additional layer
 // ----------------------------------------------------------------------------------------------------------
 // Release notes
 // ----------------------------------------------------------------------------------------------------------
- 
+
+Update
+-Refactor on ticking state is taking shape and feeling much better. Some thing may have broke but will get smoothed out.
+-New goal is to support fixed ticking seamlessly. Meaning, when you instantiate or define your network sim you can opt into fix ticking and this won't affect how your simulation code works.
+-(Nice thing is that storing frame time on the user cmds is now an internal detail of the system and transparent to users. So you will not need seperate sets of input cmds for a fix tixed version of your sim)
+
+
 Update (8-9-19)
 -Big refactor underway within the NetworkSimModel. Some files are renamed or moved as well. Still a bit to do but this was a good commit point.
 -Mainly to facilitate clean implementation of interpolation/extrapolation/forward predict option for simulated proxy simulations: this was exposing weaknesses in the templated implementations.
diff --git a/Engine/Plugins/Runtime/NetworkPredictionExtras/Source/NetworkPredictionExtras/Private/MockNetworkSimulation.cpp b/Engine/Plugins/Runtime/NetworkPredictionExtras/Source/NetworkPredictionExtras/Private/MockNetworkSimulation.cpp
--- a/Engine/Plugins/Runtime/NetworkPredictionExtras/Source/NetworkPredictionExtras/Private/MockNetworkSimulation.cpp
+++ b/Engine/Plugins/Runtime/NetworkPredictionExtras/Source/NetworkPredictionExtras/Private/MockNetworkSimulation.cpp
@@ -164,10 +164,11 @@ void UMockNetworkSimulationComponent::TickComponent(float DeltaTime, enum ELevel
 			ForceMispredict = true;
 			MockNetworkSimCVars::RequestMispredict = 0;
 		}
-
+		
 		FMockNetworkSimulation::FTickParameters Parameters;
-		Parameters.Role = OwnerRole;
 		Parameters.LocalDeltaTimeSeconds = DeltaTime;
+		Parameters.Role = OwnerRole;
+		Parameters.bIsRemotelyControlled = (OwnerRole == ROLE_Authority && GetOwner()->GetNetConnection() != nullptr);
 
 		// Tick the core network sim, this will consume input and generate new sync state
 		NetworkSim->Tick((IMockNetworkSimulationDriver*)this, Parameters);
