Continued WIP on NetworkPrediction
-Refactor on TReplicators*
-Settling in on new ticking/update scheme
-Auto and simulated proxy should be working again

#rnx

diff --git a/Engine/Plugins/Runtime/NetworkPrediction/Source/NetworkPrediction/Private/Movement/FlyingMovementComponent.cpp b/Engine/Plugins/Runtime/NetworkPrediction/Source/NetworkPrediction/Private/Movement/FlyingMovementComponent.cpp
--- a/Engine/Plugins/Runtime/NetworkPrediction/Source/NetworkPrediction/Private/Movement/FlyingMovementComponent.cpp
+++ b/Engine/Plugins/Runtime/NetworkPrediction/Source/NetworkPrediction/Private/Movement/FlyingMovementComponent.cpp
@@ -113,10 +113,7 @@ void UFlyingMovementComponent::TickComponent(float DeltaTime, enum ELevelTick Ti
 			FlyingMovementCVars::RequestMispredict = 0;
 		}
 
-		FlyingMovement::FMovementSystem::FTickParameters Parameters;
-		Parameters.LocalDeltaTimeSeconds = DeltaTime;
-		Parameters.Role = OwnerRole;
-		Parameters.bIsRemotelyControlled = (OwnerRole == ROLE_Authority && GetOwner()->GetNetConnection() != nullptr);
+		FlyingMovement::FMovementSystem::FTickParameters Parameters(DeltaTime, GetOwner());;
 
 		// Tick the core network sim, this will consume input and generate new sync state
 		NetworkSim->Tick((IFlyingMovementDriver*)this, Parameters);
diff --git a/Engine/Plugins/Runtime/NetworkPrediction/Source/NetworkPrediction/Private/Movement/ParametricMovement.cpp b/Engine/Plugins/Runtime/NetworkPrediction/Source/NetworkPrediction/Private/Movement/ParametricMovement.cpp
--- a/Engine/Plugins/Runtime/NetworkPrediction/Source/NetworkPrediction/Private/Movement/ParametricMovement.cpp
+++ b/Engine/Plugins/Runtime/NetworkPrediction/Source/NetworkPrediction/Private/Movement/ParametricMovement.cpp
@@ -205,10 +205,7 @@ void UParametricMovementComponent::TickSimulation(float DeltaTimeSeconds)
 
 	if (NetworkSim && OwnerRole != ROLE_None)
 	{
-		ParametricMovement::FMovementSystem::FTickParameters Parameters;
-		Parameters.LocalDeltaTimeSeconds = DeltaTimeSeconds;
-		Parameters.Role = OwnerRole;
-		Parameters.bIsRemotelyControlled = (OwnerRole == ROLE_Authority && GetOwner()->GetNetConnection() != nullptr);
+		ParametricMovement::FMovementSystem::FTickParameters Parameters(DeltaTimeSeconds, GetOwner());
 
 		// Tick the core network sim, this will consume input and generate new sync state
 		NetworkSim->Tick((IParametricMovementDriver*)this, Parameters);
diff --git a/Engine/Plugins/Runtime/NetworkPrediction/Source/NetworkPrediction/Private/NetworkPredictionTypes.cpp b/Engine/Plugins/Runtime/NetworkPrediction/Source/NetworkPrediction/Private/NetworkPredictionTypes.cpp
--- a/Engine/Plugins/Runtime/NetworkPrediction/Source/NetworkPrediction/Private/NetworkPredictionTypes.cpp
+++ b/Engine/Plugins/Runtime/NetworkPrediction/Source/NetworkPrediction/Private/NetworkPredictionTypes.cpp
@@ -111,4 +111,15 @@ FScopedBandwidthLimitBypass::~FScopedBandwidthLimitBypass()
 	{
 		CachedNetConnection->QueuedBits = RestoreBits - CachedNetConnection->SendBuffer.GetNumBits();
 	}
+}
+
+FNetSimTickParametersBase::FNetSimTickParametersBase(AActor* OwningActor)
+{
+	// Default implementation that will probably be good for most actors
+	Role = OwningActor->GetLocalRole();
+	const bool bHasNetConnection = OwningActor->GetNetConnection();
+
+	// The only default case we really don;t want generating local input is on the server when there is a net owning player (auto proxy)
+	// This implies that we will (sim) extrapolate by default. To opt out of extrapolation, set to false or you can just not tick the net sim
+	bGenerateLocalInputCmds = !(Role == ROLE_Authority && bHasNetConnection);
 }
\ No newline at end of file
diff --git a/Engine/Plugins/Runtime/NetworkPrediction/Source/NetworkPrediction/Public/NetworkPredictionTypes.h b/Engine/Plugins/Runtime/NetworkPrediction/Source/NetworkPrediction/Public/NetworkPredictionTypes.h
--- a/Engine/Plugins/Runtime/NetworkPrediction/Source/NetworkPrediction/Public/NetworkPredictionTypes.h
+++ b/Engine/Plugins/Runtime/NetworkPrediction/Source/NetworkPrediction/Public/NetworkPredictionTypes.h
@@ -239,4 +239,19 @@ private:
 
 	int64 RestoreBits = 0;
 	class UNetConnection* CachedNetConnection = nullptr;
+};
+
+struct NETWORKPREDICTION_API FNetSimTickParametersBase
+{
+	FNetSimTickParametersBase();
+	FNetSimTickParametersBase(AActor* Actor);
+
+	// Owner's role. Necessary to know which proxy we should be forwarding functions in tick to
+	ENetRole Role = ROLE_None;
+
+	// Are we creating input cmds locally. Note this is distinct from Role/Authority:
+	//		-[On Server] Autonomous Proxy client = false
+	//		-[On Server] Non player controlled actor = true
+	//		-[On Client] Simulated proxies (everyone but client) = true, if you want extrapolation. Note clients can just not tick the netsim to disable extrapolation as well.
+	bool bGenerateLocalInputCmds = false;
 };
\ No newline at end of file
diff --git a/Engine/Plugins/Runtime/NetworkPrediction/Source/NetworkPrediction/Public/NetworkSimulationModel.h b/Engine/Plugins/Runtime/NetworkPrediction/Source/NetworkPrediction/Public/NetworkSimulationModel.h
--- a/Engine/Plugins/Runtime/NetworkPrediction/Source/NetworkPrediction/Public/NetworkSimulationModel.h
+++ b/Engine/Plugins/Runtime/NetworkPrediction/Source/NetworkPrediction/Public/NetworkSimulationModel.h
@@ -21,14 +21,18 @@
 // ----------------------------------------------------------------------------------------------------------------------------------------------------------------------
 
 template <
-	typename T,
-	typename TUserBufferTypes,
-	typename InTTickSettings=TNetworkSimTickSettings<>,
-	typename TRepProxyServerRPC =	TReplicator_DynamicSequence	<TInternalBufferTypes<TUserBufferTypes, InTTickSettings>,	InTTickSettings, ENetworkSimBufferTypeId::Input>,
-	typename TRepProxyAutonomous =	TReplicator_BasicReconciliar <TInternalBufferTypes<TUserBufferTypes, InTTickSettings>,	InTTickSettings>,
-	typename TRepProxySimulated =	TReplicator_SimulatedExtrapolatedReconciliar <TInternalBufferTypes<TUserBufferTypes, InTTickSettings>, InTTickSettings>,
-	typename TRepProxyReplay =		TReplicator_DynamicSequence	<TInternalBufferTypes<TUserBufferTypes, InTTickSettings>,	InTTickSettings, ENetworkSimBufferTypeId::Sync>,
-	typename TRepProxyDebug =		TReplicator_DynamicSequence	<TInternalBufferTypes<TUserBufferTypes, InTTickSettings>,	InTTickSettings, ENetworkSimBufferTypeId::Debug>
+	typename T,											// Final Driver class
+	typename TUserBufferTypes,							// The user types (input, sync, aux, debug). Note this gets wrapped in TInternalBufferTypes internally.
+	typename InTTickSettings=TNetworkSimTickSettings<>, // Defines global rules about time keeping and ticking
+
+	// Core proxies that dictate how data replicates and how the simulation evolves for the three main roles
+	typename TRepProxyServerRPC =	TReplicator_Server		<TInternalBufferTypes<TUserBufferTypes, InTTickSettings>,	InTTickSettings >,
+	typename TRepProxyAutonomous =	TReplicator_Autonomous	<TInternalBufferTypes<TUserBufferTypes, InTTickSettings>,	InTTickSettings>,
+	typename TRepProxySimulated =	TReplicator_Simulated	<TInternalBufferTypes<TUserBufferTypes, InTTickSettings>,	InTTickSettings>,
+
+	// Defines how replication happens on these special channels, but doesn't dictate how simulation evolves
+	typename TRepProxyReplay =		TReplicator_DynamicSequence	<TInternalBufferTypes<TUserBufferTypes, InTTickSettings>,	InTTickSettings, ENetworkSimBufferTypeId::Sync,  3>,
+	typename TRepProxyDebug =		TReplicator_DynamicSequence	<TInternalBufferTypes<TUserBufferTypes, InTTickSettings>,	InTTickSettings, ENetworkSimBufferTypeId::Debug, 3>
 >
 class TNetworkedSimulationModel : public IReplicationProxy
 {
@@ -53,13 +57,10 @@ public:
 		virtual void FinalizeFrame(const TSyncState& SyncState) = 0; // Called from the Network Sim at the end of the sim frame when there is new sync data.
 	};
 
-	struct FTickParameters
+	struct FTickParameters : public FNetSimTickParametersBase
 	{
-		// Owner's role. Necessary to know which proxy we should be forwarding functions in tick to
-		ENetRole Role = ROLE_None;
-
-		// (Only used when ROLE_Authorty: Whether we are expecting a remote client to control. e.g, there is an autoproxy client driving this)
-		bool bIsRemotelyControlled = false;
+		FTickParameters() { }
+		FTickParameters(typename TTickSettings::TRealTime InRealTime, class AActor* Actor=nullptr) : FNetSimTickParametersBase(Actor), LocalDeltaTimeSeconds(InRealTime) {}
 
 		// Local fraem delta time seconds. Just passed in from ::Tick
 		typename TTickSettings::TRealTime LocalDeltaTimeSeconds = 0.f;
@@ -104,59 +105,18 @@ public:
 		switch (Parameters.Role)
 		{
 			case ROLE_Authority:
-				RepProxy_ServerRPC.template PreSimTick<T, TDriver>(Driver, Buffers, TickInfo, Parameters.LocalDeltaTimeSeconds);
+				RepProxy_ServerRPC.template PreSimTick<T, TDriver, FTickParameters>(Driver, Buffers, TickInfo, Parameters);
 			break;
 
 			case ROLE_AutonomousProxy:
-				RepProxy_Autonomous.template PreSimTick<T, TDriver>(Driver, Buffers, TickInfo, Parameters.LocalDeltaTimeSeconds);
+				RepProxy_Autonomous.template PreSimTick<T, TDriver, FTickParameters>(Driver, Buffers, TickInfo, Parameters);
 			break;
 
 			case ROLE_SimulatedProxy:
-				RepProxy_Simulated.template PreSimTick<T, TDriver>(Driver, Buffers, TickInfo, Parameters.LocalDeltaTimeSeconds);
-			break;
-		}
-
-
-		if (Parameters.Role == ROLE_Authority)
-		{
-			// Always give simulation time on the authority
-			TickInfo.GiveSimulationTime(Parameters.LocalDeltaTimeSeconds);
-
-			TNetworkSimTime<TTickSettings> DeltaSimTime = TickInfo.GetRemaningAllowedSimulationTime();
-			if (DeltaSimTime.ToRealTimeSeconds() > 0)
-			{
-				if (Parameters.bIsRemotelyControlled)
-				{
-					TickInfo.MaxAllowedInputKeyframe = Buffers.Input.GetHeadKeyframe();
-				}
-				else
-				{
-					if (TInputCmd* InputCmd = Buffers.Input.GetWriteNext())
-					{
-						*InputCmd = TInputCmd();
-						InputCmd->SetFrameDeltaTime(DeltaSimTime);
-						Driver->ProduceInput(DeltaSimTime, *InputCmd);
-						TickInfo.MaxAllowedInputKeyframe++;
-					}
-				}
-			}
-		}
-
-		switch (Parameters.Role) {
-
-		case ROLE_Authority:
-		{
-			if ( TickInfo.LastProcessedInputKeyframe+1 < Buffers.Input.GetTailKeyframe() )
-			{
-				// We've missed commands
-				UE_LOG(LogNetworkSim, Warning, TEXT("::Tick missing inputcmds. LastProcessedInputKeyframe: %d. %s"), TickInfo.LastProcessedInputKeyframe, *Buffers.Input.GetBasicDebugStr());
-				TickInfo.LastProcessedInputKeyframe = Buffers.Input.GetTailKeyframe()+1;
-			}
+				RepProxy_Simulated.template PreSimTick<T, TDriver, FTickParameters>(Driver, Buffers, TickInfo, Parameters);
 			break;
 		}
 
-		} // end switch
-
 		// -------------------------------------------------------------------------------------------------------------------------------------------------
 		//	PreInput processing
 		//	We are just about to start mechanically processing input. Lets just make sure we have a valid sync state to start with
@@ -236,7 +196,6 @@ public:
 			}
 		}
 
-
 		// FIXME: this needs to be sorted out. We really want to check if there is new sync state and then call this here.
 		// Call into the driver to sync to the latest state if we processed any input
 		//if (NumProcessed > 0)
@@ -352,7 +311,6 @@ public:
 
 		// We want to start with an empty command in the input buffer. See notes in input buffer processing function.
 		*Buffers.Input.GetWriteNext() = TInputCmd();
-		TickInfo.LastLocalInputGFrameNumber = 0;
 	}	
 	
 	TSimulationTickState<TTickSettings> TickInfo;	// Manages simulation time and what inputs we are processed
diff --git a/Engine/Plugins/Runtime/NetworkPrediction/Source/NetworkPrediction/Public/NetworkSimulationModelReplicators.h b/Engine/Plugins/Runtime/NetworkPrediction/Source/NetworkPrediction/Public/NetworkSimulationModelReplicators.h
--- a/Engine/Plugins/Runtime/NetworkPrediction/Source/NetworkPrediction/Public/NetworkSimulationModelReplicators.h
+++ b/Engine/Plugins/Runtime/NetworkPrediction/Source/NetworkPrediction/Public/NetworkSimulationModelReplicators.h
@@ -53,17 +53,194 @@ struct FNetworkSimulationSerialization
 //	Templated Replicators/Reconciliars
 // ----------------------------------------------------------------------------------------------------------------------------------------------------------------------
 
+template<typename TBufferTypes, typename TTickSettings>
+struct TReplicatorEmpty
+{ 
+	// Used for client shadowstate compares. Should just call GetDirtyCount() on the buffer you are replicating
+	int32 GetProxyDirtyCount(TNetworkSimBufferContainer<TBufferTypes>& Buffers) const { return 0; }
+
+	// NetSerialize: just serialize the network data. Don't run simulation steps. Every replicator will be NetSerialized before moving on to Reconcile phase.
+	void NetSerialize(const FNetSerializeParams& P, TNetworkSimBufferContainer<TBufferTypes>& Buffers, const TSimulationTickState<TTickSettings>& TickInfo) { }
+
+	// Reconcile: called after everyone has NetSerialized. "Get right with the server": this function is about reconciling what the server told you vs what you have extrapolated or forward predicted locally
+	template<typename T, typename TDriver>
+	void Reconcile(TDriver* Driver, TNetworkSimBufferContainer<TBufferTypes>& Buffers, TSimulationTickState<TTickSettings>& TickInfo) { }
+
+	template<typename T, typename TDriver, typename TTickParameters>
+	void PreSimTick(TDriver* Driver, TNetworkSimBufferContainer<TBufferTypes>& Buffers, TSimulationTickState<TTickSettings>& TickInfo, const TTickParameters& TickParameters) { }
+};
+
+// This is the "templated base class" for replicators but is not required (i.e., this not an official interface used by TNetworkedSimulation model. Just a base implementation you can start with)
 template<typename TBufferTypes, typename TTickSettings>
 struct TReplicatorBase
 {
+	using TInputCmd = typename TBufferTypes::TInputCmd;
+
+	// Used for client shadowstate compares. Should just call GetDirtyCount() on the buffer you are replicating
+	int32 GetProxyDirtyCount(TNetworkSimBufferContainer<TBufferTypes>& Buffers) const { return 0; }
+
+	// NetSerialize: just serialize the network data. Don't run simulation steps. Every replicator will be NetSerialized before moving on to Reconcile phase.
+	void NetSerialize(const FNetSerializeParams& P, TNetworkSimBufferContainer<TBufferTypes>& Buffers, const TSimulationTickState<TTickSettings>& TickInfo) { }
+
+	// Reconcile: called after everyone has NetSerialized. "Get right with the server": this function is about reconciling what the server told you vs what you have extrapolated or forward predicted locally
 	template<typename T, typename TDriver>
 	void Reconcile(TDriver* Driver, TNetworkSimBufferContainer<TBufferTypes>& Buffers, TSimulationTickState<TTickSettings>& TickInfo) { }
 
-	template<typename T, typename TDriver>
-	void PreSimTick(TDriver* Driver, TNetworkSimBufferContainer<TBufferTypes>& Buffers, TSimulationTickState<TTickSettings>& TickInfo, typename TTickSettings::TRealTime RealDeltaTimeSeconds) { }
+	// Called prior to input processing. This function must updated TickInfo to allow simulation time (from TickParameters) and to possibly get new input.
+	template<typename T, typename TDriver, typename TTickParameters>
+	void PreSimTick(TDriver* Driver, TNetworkSimBufferContainer<TBufferTypes>& Buffers, TSimulationTickState<TTickSettings>& TickInfo, const TTickParameters& TickParameters)
+	{
+		// Accumulate local delta time into TickInfo
+		TickInfo.GiveSimulationTime(TickParameters.LocalDeltaTimeSeconds);
+
+		// See if we have sim time to spend (variable tick will always go through. fixed step will miss frames while accumulating)
+		TNetworkSimTime<TTickSettings> DeltaSimTime = TickInfo.GetRemaningAllowedSimulationTime();
+		if (DeltaSimTime.ToRealTimeSeconds() > 0)
+		{
+			if (TickParameters.bGenerateLocalInputCmds)
+			{
+				if (TInputCmd* InputCmd = Buffers.Input.GetWriteNext())
+				{
+					*InputCmd = TInputCmd();
+					InputCmd->SetFrameDeltaTime(DeltaSimTime);
+					Driver->ProduceInput(DeltaSimTime, *InputCmd);
+					TickInfo.MaxAllowedInputKeyframe++;
+				}
+			}
+			else
+			{
+				// Allowed to process all remove commands (we will still check for frame time in the core input processing loop to prevent speed hacks)
+				TickInfo.MaxAllowedInputKeyframe = Buffers.Input.GetHeadKeyframe();
+
+				// Check for gaps in commands
+				if ( TickInfo.LastProcessedInputKeyframe+1 < Buffers.Input.GetTailKeyframe() )
+				{
+					UE_LOG(LogNetworkSim, Warning, TEXT("::Tick missing inputcmds. LastProcessedInputKeyframe: %d. %s"), TickInfo.LastProcessedInputKeyframe, *Buffers.Input.GetBasicDebugStr());
+					TickInfo.LastProcessedInputKeyframe = Buffers.Input.GetTailKeyframe()+1;
+				}
+			}
+		}
+	}
+};
+
+// -------------------------------------------------------------------------------------------------------
+//	"Reusable" pieces: mainly for replicating data on specified buffers
+// -------------------------------------------------------------------------------------------------------
+
+template<typename TBufferTypes, typename TTickSettings, typename TBase=TReplicatorEmpty<TBufferTypes, TTickSettings>, bool Enabled=true>
+struct TReplicator_SimTime : public TBase
+{
+	void NetSerialize(const FNetSerializeParams& P, TNetworkSimBufferContainer<TBufferTypes>& Buffers, const TSimulationTickState<TTickSettings>& TickInfo)
+	{
+		TBase::NetSerialize(P, Buffers, TickInfo);
+		SerializedTime = TickInfo.TotalProcessedSimulationTime;
+		SerializedTime.NetSerialize(P.Ar);
+	}
+
+	TNetworkSimTime<TTickSettings> SerializedTime;
+};
+
+template<typename TBufferTypes, typename TTickSettings, typename TBase>
+struct TReplicator_SimTime<TBufferTypes, TTickSettings, TBase, false> : public TBase { };
+
+template<typename TBufferTypes, typename TTickSettings, ENetworkSimBufferTypeId BufferId, typename TBase=TReplicatorEmpty<TBufferTypes, TTickSettings>, bool Enabled=true>
+struct TReplicator_HeadKeyframe : public TBase
+{
+	void NetSerialize(const FNetSerializeParams& P, TNetworkSimBufferContainer<TBufferTypes>& Buffers, const TSimulationTickState<TTickSettings>& TickInfo)
+	{
+		TBase::NetSerialize(P, Buffers, TickInfo);
+
+		auto& Buffer = Buffers.template Get<BufferId>();
+		SerializedHeadKeyframe = FNetworkSimulationSerialization::SerializeKeyframe(P.Ar, Buffer.GetHeadKeyframe());
+	}
+
+	int32 SerializedHeadKeyframe;
+};
+
+template<typename TBufferTypes, typename TTickSettings, ENetworkSimBufferTypeId BufferId, typename TBase>
+struct TReplicator_HeadKeyframe<TBufferTypes, TTickSettings, BufferId, TBase, false> : public TBase { };
+
+
+// Helper to make using the above serializers less gross. Selects the appropriate base class depending on SerializeKeyframeNumber and SerializeSimTime
+template<typename TBufferTypes, typename TTickSettings, ENetworkSimBufferTypeId BufferId, bool SerializeKeyframeNumber=false, bool SerializeSimTime=true, typename TRoot=TReplicatorBase<TBufferTypes, TTickSettings>,
+	typename TBase=
+		TReplicator_SimTime<TBufferTypes, TTickSettings,						// Replicates sim time....
+			TReplicator_HeadKeyframe<TBufferTypes, TTickSettings, BufferId,		// Replicates buffer keyframe number...
+				TRoot,						// "final" parent class
+			SerializeKeyframeNumber>,		// ... only if SerializeKeyframeNumber=true
+		SerializeSimTime>>					// ... only if SerializeSimTime=true
+struct TSelectBase
+{
+	using type = TBase;
+};
+
+// -------------------------------------------------------------------------------------------------------
+//
+// -------------------------------------------------------------------------------------------------------
+
+template<typename TBufferTypes, typename TTickSettings, ENetworkSimBufferTypeId BufferId, int32 NumElements=1, bool SerializeKeyframeNumber=false, bool SerializeSimTime=true, typename TRoot=TReplicatorBase<TBufferTypes, TTickSettings>,
+	typename TBase=TSelectBase<TBufferTypes, TTickSettings, BufferId, SerializeKeyframeNumber, SerializeSimTime, TRoot>>
+struct TReplicator_StaticSequence : public TBase
+{
+	int32 GetProxyDirtyCount(TNetworkSimBufferContainer<TBufferTypes>& Buffers) const 
+	{
+		return Buffers.template Get<BufferId>().GetDirtyCount() ^ (TBase::GetProxyDirtyCount(Buffers) << 2); 
+	}
+
+	void NetSerialize(const FNetSerializeParams& P, TNetworkSimBufferContainer<TBufferTypes>& Buffers, const TSimulationTickState<TTickSettings>& TickInfo)
+	{
+		TBase::NetSerialize(P, Buffers, TickInfo);
+
+		auto& Buffer = Buffers.template Get<BufferId>();
+		unimplemented();
+	}
+};
+
+/** Replicates the latest sequence of N elements. N is dynamic, not compiled in. */
+template<typename TBufferTypes, typename TTickSettings, ENetworkSimBufferTypeId BufferId, int32 MaxNumElements, bool SerializeKeyframeNumber=false, bool SerializeSimTime=true, typename TRoot=TReplicatorBase<TBufferTypes, TTickSettings>,
+	typename TBase=typename TSelectBase<TBufferTypes, TTickSettings, BufferId, SerializeKeyframeNumber, SerializeSimTime, TRoot>::type>
+struct TReplicator_DynamicSequence : public TBase
+{
+	int32 GetProxyDirtyCount(TNetworkSimBufferContainer<TBufferTypes>& Buffers) const 
+	{
+		return Buffers.template Get<BufferId>().GetDirtyCount() ^ (TBase::GetProxyDirtyCount(Buffers) << 2); 
+	}
+	
+	void NetSerialize(const FNetSerializeParams& P, TNetworkSimBufferContainer<TBufferTypes>& Buffers, const TSimulationTickState<TTickSettings>& TickInfo)
+	{
+		TBase::NetSerialize(P, Buffers, TickInfo);
+
+		auto& Buffer = Buffers.template Get<BufferId>();
+		FArchive& Ar = P.Ar;
+		uint8 SerializedNumElements = FMath::Min<uint8>(MaxNumElements, Buffer.GetNumValidElements());
+		Ar << SerializedNumElements;
+
+		const int32 HeadKeyframe = FNetworkSimulationSerialization::SerializeKeyframe(Ar, Buffer.GetHeadKeyframe());
+		const int32 StartingKeyframe = FMath::Max(0, HeadKeyframe - SerializedNumElements + 1);
+
+		if (Ar.IsLoading())
+		{
+			const int32 PrevHead = Buffer.GetHeadKeyframe();
+			if (PrevHead < StartingKeyframe && PrevHead >= 0)
+			{
+				// There is a gap in the stream. In some cases, we want this to be a "fault" and bubble up. We may want to synthesize state or maybe we just skip ahead.
+				UE_LOG(LogNetworkSim, Warning, TEXT("Fault: gap in received buffer. PrevHead: %d. Received: %d-%d. Reseting previous buffer contents"), PrevHead, StartingKeyframe, HeadKeyframe);
+			}
+
+			Buffer.ResetNextHeadKeyframe(StartingKeyframe);
+		}
+
+		for (int32 Keyframe = StartingKeyframe; Keyframe <= HeadKeyframe; ++Keyframe)
+		{
+			// This, as is, is bad. The intention is that these functions serialize multiple items in some delta compressed fashion.
+			// As is, we are just serializing the elements individually.
+			auto* Cmd = Ar.IsLoading() ? Buffer.GetWriteNext() : Buffer.FindElementByKeyframe(Keyframe);
+			Cmd->NetSerialize(P);
+		}
+	}
 };
 
-/** Replicates only the latest element. Does not synchronize keyframe */
+/** Replicates only the latest element. Does not synchronize keyframe or simulation time */
 template<typename TBufferTypes, typename TTickSettings, ENetworkSimBufferTypeId BufferId, typename TBase=TReplicatorBase<TBufferTypes, TTickSettings>>
 struct TReplicator_LatestOnly : public TBase
 {
@@ -71,7 +248,7 @@ struct TReplicator_LatestOnly : public TBase
 
 	int32 GetProxyDirtyCount(TNetworkSimBufferContainer<TBufferTypes>& Buffers) const
 	{ 
-		return Buffers.template Get<BufferId>().GetDirtyCount(); 
+		return Buffers.template Get<BufferId>().GetDirtyCount() ^ (TBase::GetProxyDirtyCount(Buffers) << 2); ; 
 	}
 
 	void NetSerialize(const FNetSerializeParams& P, TNetworkSimBufferContainer<TBufferTypes>& Buffers, const TSimulationTickState<TTickSettings>& TickInfo)
@@ -95,15 +272,66 @@ struct TReplicator_LatestOnly : public TBase
 	}
 };
 
-template<typename TBufferTypes, typename TTickSettings, typename TBase=TReplicatorBase<TBufferTypes, TTickSettings>, bool IsFixed=TTickSettings::FixedStepMS==0>
-struct TReplicator_SimulatedExtrapolatedReconciliar : public TBase
+// -------------------------------------------------------------------------------------------------------
+//	Role based Replicators: these replicators are meant to server specific roles
+// -------------------------------------------------------------------------------------------------------
+
+// Default Replicator for the Server
+//	-Replicates the InputBuffer client->server
+template<typename TBufferTypes, typename TTickSettings, ENetworkSimBufferTypeId BufferId=ENetworkSimBufferTypeId::Input, typename TBase=TReplicator_DynamicSequence<TBufferTypes, TTickSettings, BufferId, 3>>
+struct TReplicator_Server : public TBase
+{
+
+};
+
+/*
+template<typename TBufferTypes, typename TTickSettings, typename TBase=TReplicator_DynamicSequence<TBufferTypes, TTickSettings, ENetworkSimBufferTypeId::Sync, 1>>
+struct TReplicator_Simulated : public TBase
 {
 	using TInputCmd = typename TBufferTypes::TInputCmd;
 	using TSyncState = typename TBufferTypes::TSyncState;
+
+	void NetSerialize(const FNetSerializeParams& P, TNetworkSimBufferContainer<TBufferTypes>& Buffers, const TSimulationTickState<TTickSettings>& TickInfo)
+	{
+		TBase::NetSerialize(P, Buffers, TickInfo);
+		bPendingReconcile = true;
+	}
+
+	template<typename T, typename TDriver>
+	void Reconcile(TDriver* Driver, TNetworkSimBufferContainer<TBufferTypes>& Buffers, TSimulationTickState<TTickSettings>& TickInfo)
+	{
+		// Fixme
+		if (bPendingReconcile==false)
+		{
+			return;
+		}
+		bPendingReconcile = false;
+
+		// Generate a new, fake, command since we just added a new sync state to head
+		Buffers.Input.GetWriteNext();
+			
+		// Set our LastProcessedInputKeyframe to fake that we handled it
+		TickInfo.LastProcessedInputKeyframe = Buffers.Input.GetHeadKeyframe();
+		TickInfo.MaxAllowedInputKeyframe = FMath::Max(TickInfo.MaxAllowedInputKeyframe, TickInfo.LastProcessedInputKeyframe);
+	}
+
+private:
+
+	bool bPendingReconcile = false;
+	
+};
+*/
+
+template<typename TBufferTypes, typename TTickSettings, typename TBase=TReplicatorBase<TBufferTypes, TTickSettings>>
+struct TReplicator_Simulated : public TBase
+{
+	using TInputCmd = typename TBufferTypes::TInputCmd;
+	using TSyncState = typename TBufferTypes::TSyncState;
+	using TAuxState = typename TBufferTypes::TAuxState;
 	
 	int32 GetProxyDirtyCount(TNetworkSimBufferContainer<TBufferTypes>& Buffers) const
 	{
-		return Buffers.Sync.GetDirtyCount(); 
+		return Buffers.Sync.GetDirtyCount() ^ (TBase::GetProxyDirtyCount(Buffers) << 2); 
 	}
 
 	void NetSerialize(const FNetSerializeParams& P, TNetworkSimBufferContainer<TBufferTypes>& Buffers, const TSimulationTickState<TTickSettings>& TickInfo)
@@ -133,32 +361,40 @@ struct TReplicator_SimulatedExtrapolatedReconciliar : public TBase
 			return;
 
 		bPendingReconciliation = false;
+		TInputCmd* LastCmd = Buffers.Input.GetElementFromHead(0);
+		
+		// Generate a new, fake, command since we just added a new sync state to head
+		Buffers.Input.GetWriteNext();
 
-		// This is still being worked out. We are almost there but not quite.
-		// Processing input commands is pretty rigid right now. Since we just added a new SyncState keyframe in ::NetSerialize,
-		// we need to insert a fake input command so the extrapolation code in the NetSimModel can continue to work.
-		// This can be improved.
-
+		TNetworkSimTime<TTickSettings> DeltaSimTime = TickInfo.TotalProcessedSimulationTime - SerializedTime;
+		if (DeltaSimTime > TNetworkSimTime<TTickSettings>::FromMSec(0))
 		{
-			// Generate a new, fake, command since we just added a new sync state to head (via GetWriteNext)
-			TickInfo.GetNextInputForWrite(Buffers);
-			// Gotta fake this too
-			TickInfo.LastLocalInputGFrameNumber--; 
-			// Set our LastProcessedInputKeyframe to fake that we handled it
-			TickInfo.LastProcessedInputKeyframe = Buffers.Sync.GetHeadKeyframe();
-		}
-	}
+			// We are extrapolating ahead of the server. We will make up the difference here
+			// (this can easily happen: variance in latency may delay a packet and we may get ahead)
+			TInputCmd* NewCmd = Buffers.Input.GetWriteNext();
+			*NewCmd = LastCmd ? *LastCmd : TInputCmd();
+			NewCmd->SetFrameDeltaTime(DeltaSimTime);
 
-	template<typename T, typename TDriver>
-	void PreSimTick(TDriver* Driver, TNetworkSimBufferContainer<TBufferTypes>& Buffers, TSimulationTickState<TTickSettings>& TickInfo, typename TTickSettings::TRealTime RealDeltaTimeSeconds)
-	{
-		// Extrapolation: still needs to be a setting/state somewhere that we can change from actor code
-		if (TInputCmd* SynthesizedCmd = TickInfo.GetNextInputForWrite(Buffers))
+			TSyncState* PrevSyncState = Buffers.Sync.GetElementFromHead(0);
+			TSyncState* NextSyncState = Buffers.Sync.GetWriteNext();
+
+			TAuxState Junk;
+
+			// TODO: log this guy
+			//NextMotionState->VisualLog( FVisualLoggingParameters(Keyframe == LastKeyframeToProcess ? EVisualLoggingContext::LastMispredicted : EVisualLoggingContext::OtherMispredicted, Keyframe, EVisualLoggingLifetime::Persistent), Driver, Driver);
+
+			// Do the actual update
+			T::Update(Driver, NewCmd->GetFrameDeltaTime(), *NewCmd, *PrevSyncState, *NextSyncState, Junk);
+		}
+		else
 		{
-			// We will create a single cmd to simulate the remaining time
-			auto TimeToSimulate = TickInfo.GetRemaningAllowedSimulationTime();
-			//SynthesizedCmd->FrameDeltaTime;
+			TickInfo.TotalProcessedSimulationTime = SerializedTime;
+			TickInfo.TotalAllowedSimulationTime = SerializedTime;
 		}
+
+		// Set our LastProcessedInputKeyframe to fake that we handled it
+		TickInfo.LastProcessedInputKeyframe = Buffers.Input.GetHeadKeyframe();
+		TickInfo.MaxAllowedInputKeyframe = FMath::Max(TickInfo.MaxAllowedInputKeyframe, TickInfo.LastProcessedInputKeyframe);
 	}
 
 private:
@@ -167,14 +403,15 @@ private:
 	bool bPendingReconciliation = false;
 };
 
+/*
 template<typename TBufferTypes, typename TTickSettings, typename TBase>
-struct TReplicator_SimulatedExtrapolatedReconciliar<TBufferTypes, TTickSettings, TBase, true> : TReplicator_SimulatedExtrapolatedReconciliar<TBufferTypes, TTickSettings, TBase, false>
+struct TReplicator_Simulated<TBufferTypes, TTickSettings, TBase, true> : TReplicator_Simulated<TBufferTypes, TTickSettings, TBase, false>
 {
 	using TInputCmd = typename TBufferTypes::TInputCmd;
 	using TSyncState = typename TBufferTypes::TSyncState;
 
-	template<typename T, typename TDriver>
-	void PreSimTick(TDriver* Driver, TNetworkSimBufferContainer<TBufferTypes>& Buffers, TSimulationTickState<TTickSettings>& TickInfo, typename TTickSettings::TRealTime RealDeltaTimeSeconds)
+	template<typename T, typename TDriver, typename TTickParameters>
+	void PreSimTick(TDriver* Driver, TNetworkSimBufferContainer<TBufferTypes>& Buffers, TSimulationTickState<TTickSettings>& TickInfo, const TTickParameters& TickParameters)
 	{
 		// Extrapolation: still needs to be a setting/state somewhere that we can change from actor code
 		if (TInputCmd* SynthesizedCmd = TickInfo.GetNextInputForWrite(Buffers))
@@ -185,55 +422,12 @@ struct TReplicator_SimulatedExtrapolatedReconciliar<TBufferTypes, TTickSettings,
 		}
 	}
 };
+*/
 
-/** Replicates the latest sequence of N elements. N is dynamic, not compiled in. */
-template<typename TBufferTypes, typename TTickSettings, ENetworkSimBufferTypeId BufferId, typename TBase=TReplicatorBase<TBufferTypes, TTickSettings>>
-struct TReplicator_DynamicSequence : public TBase
-{
-	int32 GetProxyDirtyCount(TNetworkSimBufferContainer<TBufferTypes>& Buffers) const 
-	{
-		return Buffers.template Get<BufferId>().GetDirtyCount(); 
-	}
-	
-	void NetSerialize(const FNetSerializeParams& P, TNetworkSimBufferContainer<TBufferTypes>& Buffers, const TSimulationTickState<TTickSettings>& TickInfo)
-	{
-		auto& Buffer = Buffers.template Get<BufferId>();
-		FArchive& Ar = P.Ar;
-		uint8 SerializedNumElements = FMath::Min<uint8>(MaxNumElements, Buffer.GetNumValidElements());
-		Ar << SerializedNumElements;
-
-		const int32 HeadKeyframe = FNetworkSimulationSerialization::SerializeKeyframe(Ar, Buffer.GetHeadKeyframe());
-		const int32 StartingKeyframe = FMath::Max(0, HeadKeyframe - SerializedNumElements + 1);
-
-		if (Ar.IsLoading())
-		{
-			const int32 PrevHead = Buffer.GetHeadKeyframe();
-			if (PrevHead < StartingKeyframe && PrevHead >= 0)
-			{
-				// There is a gap in the stream. In some cases, we want this to be a "fault" and bubble up. We may want to synthesize state or maybe we just skip ahead.
-				UE_LOG(LogNetworkSim, Warning, TEXT("Fault: gap in received buffer. PrevHead: %d. Received: %d-%d. Reseting previous buffer contents"), PrevHead, StartingKeyframe, HeadKeyframe);
-			}
-
-			Buffer.ResetNextHeadKeyframe(StartingKeyframe);
-		}
-
-		for (int32 Keyframe = StartingKeyframe; Keyframe <= HeadKeyframe; ++Keyframe)
-		{
-			// This, as is, is bad. The intention is that these functions serialize multiple items in some delta compressed fashion.
-			// As is, we are just serializing the elements individually.
-			auto* Cmd = Ar.IsLoading() ? Buffer.GetWriteNext() : Buffer.FindElementByKeyframe(Keyframe);
-			Cmd->NetSerialize(P);
-		}
-	}
-
-private:
-	
-	uint8 MaxNumElements = 3;
-};
 
 /** Replicates TSyncState and does basic reconciliation. */
 template<typename TBufferTypes, typename TTickSettings, typename TBase=TReplicatorBase<TBufferTypes, TTickSettings>>
-struct TReplicator_BasicReconciliar : public TBase
+struct TReplicator_Autonomous : public TBase
 {
 	using TInputCmd = typename TBufferTypes::TInputCmd;
 	using TSyncState = typename TBufferTypes::TSyncState;
@@ -245,14 +439,14 @@ struct TReplicator_BasicReconciliar : public TBase
 
 	int32 GetProxyDirtyCount(TNetworkSimBufferContainer<TBufferTypes>& Buffers) const
 	{
-		return Buffers.Sync.GetDirtyCount();
+		return Buffers.Sync.GetDirtyCount() ^ (TBase::GetProxyDirtyCount(Buffers) << 2); ;
 	}
 
 	// --------------------------------------------------------------------
 	//	PreSimTick
 	// --------------------------------------------------------------------
-	template<typename T, typename TDriver>
-	void PreSimTick(TDriver* Driver, TNetworkSimBufferContainer<TBufferTypes>& Buffers, TSimulationTickState<TTickSettings>& TickInfo, typename TTickSettings::TRealTime RealDeltaTimeSeconds)
+	template<typename T, typename TDriver, typename TTickParameters>
+	void PreSimTick(TDriver* Driver, TNetworkSimBufferContainer<TBufferTypes>& Buffers, TSimulationTickState<TTickSettings>& TickInfo, const TTickParameters& TickParameters)
 	{
 		// If we have a reconcile fault, we cannot continue on with the simulation until it clears itself out. This effectively drops the input time and does not sample new inputs
 		if (bReconcileFaultDetected)
@@ -261,7 +455,7 @@ struct TReplicator_BasicReconciliar : public TBase
 		}
 
 		// Accumulate real time into the simulation. This will results in more Allowed Simulation Time (or will accumulate internally until we are ready for the next tick)
-		TickInfo.GiveSimulationTime(RealDeltaTimeSeconds);
+		TickInfo.GiveSimulationTime(TickParameters.LocalDeltaTimeSeconds);
 
 		TNetworkSimTime<TTickSettings> DeltaSimTime = TickInfo.GetRemaningAllowedSimulationTime();
 		if (DeltaSimTime.ToRealTimeSeconds() > 0)
diff --git a/Engine/Plugins/Runtime/NetworkPrediction/Source/NetworkPrediction/Public/NetworkSimulationModelTypes.h b/Engine/Plugins/Runtime/NetworkPrediction/Source/NetworkPrediction/Public/NetworkSimulationModelTypes.h
--- a/Engine/Plugins/Runtime/NetworkPrediction/Source/NetworkPrediction/Public/NetworkSimulationModelTypes.h
+++ b/Engine/Plugins/Runtime/NetworkPrediction/Source/NetworkPrediction/Public/NetworkSimulationModelTypes.h
@@ -157,7 +157,7 @@ struct TNetworkSimBufferContainer
 // ----------------------------------------------------------------------------------------------------------------------------------------------
 
 // Holds all settings for a network sim related to ticking
-template<uint32 InFixedStepMS=0, uint32 InMaxStepMS=0, typename TUnderlyingSimTimeType=uint32, typename TUnderlyingRealTimeType=float, int32 InRealToSimFactor=1000>
+template<int32 InFixedStepMS=0, int32 InMaxStepMS=0, typename TUnderlyingSimTimeType=int32, typename TUnderlyingRealTimeType=float, int32 InRealToSimFactor=1000>
 struct TNetworkSimTickSettings
 {
 	using TUnderlingSimTime = TUnderlyingSimTimeType;	// Underlying type of time. This type shouldn't be used directly (use TNetworkSimTime)
@@ -275,27 +275,8 @@ struct TSimulationTickState
 	int32 LastProcessedInputKeyframe;	// The last input keyframe that we processed
 	int32 MaxAllowedInputKeyframe;		// The max input keyframe that we are allowed to process (e.g, don't process input past this keyframe yet)
 
-	uint32 LastLocalInputGFrameNumber;	// Tracks the last time we accepted local input via GetNextInputForWrite. Used to guard against accidental input latency due to ordering of input/sim ticking.
-
 	TNetworkSimTime<TSettings> TotalAllowedSimulationTime;	// Total time we have been "given" to process. We cannot process more simulation time than this: doing so would be speed hacking.
 	TNetworkSimTime<TSettings> TotalProcessedSimulationTime;	// How much time we've actually processed. The only way to increment this is to process user commands or receive authoritative state from the network.
-	
-	template<typename TBufferTypes>
-	typename TBufferTypes::TInputCmd* GetNextInputForWrite(TNetworkSimBufferContainer<TBufferTypes>& Buffers)
-	{
-		// Not really necessary anymore since input command is requested by the sim now instead of pushed into it
-		ensure(LastLocalInputGFrameNumber != GFrameNumber);
-		LastLocalInputGFrameNumber = GFrameNumber;
-
-		typename TBufferTypes::TInputCmd* Next = nullptr;
-		if (Buffers.Input.GetHeadKeyframe() == LastProcessedInputKeyframe)
-		{
-			// Only return a cmd if we have processed the last one. This is a bit heavy handed but is a good practice to start with. We want buffering of input to be very explicit, never something that accidentally happens.
-			Next = Buffers.Input.GetWriteNext();
-			*Next = typename TBufferTypes::TInputCmd();
-		}
-		return Next;
-	}
 
 	// "Grants" allowed simulation time to this tick state. That is, we are now allowed to advance the simulation by this amount the next time the sim ticks.
 	// Note the input is RealTime in SECONDS. This is what the rest of the engine uses when dealing with float delta time.
@@ -304,6 +285,7 @@ struct TSimulationTickState
 		RealtimeAccumulator.Accumulate(TotalAllowedSimulationTime, RealTimeSeconds);
 	}
 
+	// How much granted simulation time is left to process
 	TNetworkSimTime<TSettings> GetRemaningAllowedSimulationTime() const
 	{
 		return TotalAllowedSimulationTime - TotalProcessedSimulationTime;
@@ -324,6 +306,11 @@ struct TFrameCmd : public BaseCmdType
 {
 	TNetworkSimTime<TickSettings> GetFrameDeltaTime() const { return FrameDeltaTime; }
 	void SetFrameDeltaTime(const TNetworkSimTime<TickSettings>& InTime) { FrameDeltaTime = InTime; }
+	void NetSerialize(const FNetSerializeParams& P)
+	{
+		FrameDeltaTime.NetSerialize(P.Ar);
+		BaseCmdType::NetSerialize(P); 
+	}
 
 private:
 	TNetworkSimTime<TickSettings> FrameDeltaTime;
@@ -335,6 +322,7 @@ struct TFrameCmd<BaseCmdType, TickSettings, true> : public BaseCmdType
 {
 	constexpr TNetworkSimTime<TickSettings> GetFrameDeltaTime() const { return TNetworkSimTime<TickSettings>::FromMSec(TickSettings::GetFixedStepMS()); }
 	void SetFrameDeltaTime(const TNetworkSimTime<TickSettings>& InTime) { }
+	void NetSerialize(const FNetSerializeParams& P) { BaseCmdType::NetSerialize(P); }
 };
 
 // Helper to turn user supplied buffer types into the "real" buffer types: the InputCmd struct is wrapped in TFrameCmd
diff --git a/Engine/Plugins/Runtime/NetworkPrediction/readme.txt b/Engine/Plugins/Runtime/NetworkPrediction/readme.txt
--- a/Engine/Plugins/Runtime/NetworkPrediction/readme.txt
+++ b/Engine/Plugins/Runtime/NetworkPrediction/readme.txt
@@ -50,7 +50,12 @@ UE4 Networking info (How this plugs into UE4 networking)
 -"Replication Proxy" (FReplicationProxy): struct that points to the network simulation and has policies about how it replicates and to who. Custom NetSerialize.
 -E.g, there is a replication proxy for owner, non owners, replays, and debug. All point to the same underlying data.
 -Still using actor replication to determine when to replicate.
- 
+
+Simulated Proxy behavior / "modes"
+-Interpolate: Sync buffer is replicated but there is no ticking of the net sim. Literally do not have to Tick the simulation.
+	-Actual interpolation will happen in some class outside of TNetworkSimulationModel (doesn't exist yet). It will look at sync buffer and interpolate over it.
+-(Sim) Extrapolation: by default if the netsim is ticked, we synthesize command and extrapolate the simulation. With basic reconciliation to absorb fluctuations in latency.
+-Forward Predict (not implemented yet): must be explicitly enabled by outside call, ties simulated proxy sim to an autonomous proxy sim. Sims will be reconciled together in step.
  
  
 TODO: Major missing elements
@@ -95,6 +100,11 @@ Smoothing: Taking the output of the simulation and applying an additional layer
 // Release notes
 // ----------------------------------------------------------------------------------------------------------
 
+
+Update (8-13-19)
+-Refactor really feeling good, things are falling into place. Still a few things to do and bugs to chase down.
+-Added simulated proxy notes above to detail current plan
+
 Update
 -Refactor on ticking state is taking shape and feeling much better. Some thing may have broke but will get smoothed out.
 -New goal is to support fixed ticking seamlessly. Meaning, when you instantiate or define your network sim you can opt into fix ticking and this won't affect how your simulation code works.
diff --git a/Engine/Plugins/Runtime/NetworkPredictionExtras/Source/NetworkPredictionExtras/Private/MockNetworkSimulation.cpp b/Engine/Plugins/Runtime/NetworkPredictionExtras/Source/NetworkPredictionExtras/Private/MockNetworkSimulation.cpp
--- a/Engine/Plugins/Runtime/NetworkPredictionExtras/Source/NetworkPredictionExtras/Private/MockNetworkSimulation.cpp
+++ b/Engine/Plugins/Runtime/NetworkPredictionExtras/Source/NetworkPredictionExtras/Private/MockNetworkSimulation.cpp
@@ -165,10 +165,7 @@ void UMockNetworkSimulationComponent::TickComponent(float DeltaTime, enum ELevel
 			MockNetworkSimCVars::RequestMispredict = 0;
 		}
 		
-		FMockNetworkSimulation::FTickParameters Parameters;
-		Parameters.LocalDeltaTimeSeconds = DeltaTime;
-		Parameters.Role = OwnerRole;
-		Parameters.bIsRemotelyControlled = (OwnerRole == ROLE_Authority && GetOwner()->GetNetConnection() != nullptr);
+		FMockNetworkSimulation::FTickParameters Parameters(DeltaTime, GetOwner());
 
 		// Tick the core network sim, this will consume input and generate new sync state
 		NetworkSim->Tick((IMockNetworkSimulationDriver*)this, Parameters);
