NetworkPrediction continued refactor and cleanup  WIP
#rnx

diff --git a/Engine/Plugins/Runtime/NetworkPrediction/Source/NetworkPrediction/Private/Movement/BaseMovementComponent.cpp b/Engine/Plugins/Runtime/NetworkPrediction/Source/NetworkPrediction/Private/Movement/BaseMovementComponent.cpp
--- a/Engine/Plugins/Runtime/NetworkPrediction/Source/NetworkPrediction/Private/Movement/BaseMovementComponent.cpp
+++ b/Engine/Plugins/Runtime/NetworkPrediction/Source/NetworkPrediction/Private/Movement/BaseMovementComponent.cpp
@@ -41,10 +41,10 @@ static FAutoConsoleVariableRef CVarUseVLogger(TEXT("bm.Debug.UseUnrealVLogger"),
 
 static int32 UseDrawDebug = 1;
 static FAutoConsoleVariableRef CVarUseDrawDebug(TEXT("bm.Debug.UseDrawDebug"),
-	UseVLogger,	TEXT("Use built in DrawDebug* functions for visual logging\n"), ECVF_Default);
+	UseDrawDebug,	TEXT("Use built in DrawDebug* functions for visual logging\n"), ECVF_Default);
 
 static float DrawDebugDefaultLifeTime = 30.f;
-static FAutoConsoleVariableRef CVarDrawDebugDefaultLifeTime(TEXT("bm.Debug.UseDrawDebug.DefaultLifeTime"),
+static FAutoConsoleVariableRef CVarDrawDebugDefaultLifeTime(TEXT("bm.Debug.DrawDebugLifetime.DefaultLifeTime"),
 	DrawDebugDefaultLifeTime, TEXT("Use built in DrawDebug* functions for visual logging"), ECVF_Default);
 }
 
diff --git a/Engine/Plugins/Runtime/NetworkPrediction/Source/NetworkPrediction/Private/Movement/FlyingMovementComponent.cpp b/Engine/Plugins/Runtime/NetworkPrediction/Source/NetworkPrediction/Private/Movement/FlyingMovementComponent.cpp
--- a/Engine/Plugins/Runtime/NetworkPrediction/Source/NetworkPrediction/Private/Movement/FlyingMovementComponent.cpp
+++ b/Engine/Plugins/Runtime/NetworkPrediction/Source/NetworkPrediction/Private/Movement/FlyingMovementComponent.cpp
@@ -75,7 +75,7 @@ IReplicationProxy* UFlyingMovementComponent::InstantiateNetworkSimulation()
 void UFlyingMovementComponent::InitializeForNetworkRole(ENetRole Role)
 {	
 	check(NetworkSim);
-	NetworkSim->InitializeForNetworkRole(Role, IsLocallyControlled(), GetSimulationInitParameters(Role));
+	NetworkSim->InitializeForNetworkRole(Role, GetSimulationInitParameters(Role));
 }
 
 void UFlyingMovementComponent::TickComponent(float DeltaTime, enum ELevelTick TickType, FActorComponentTickFunction* ThisTickFunction)
@@ -169,4 +169,9 @@ void UFlyingMovementComponent::ProduceInput(const FlyingMovement::TSimTime& SimF
 void UFlyingMovementComponent::FinalizeFrame(const FlyingMovement::FMoveState& SyncState)
 {
 	PreSimSync(SyncState);
+}
+
+FString UFlyingMovementComponent::GetDebugName() const
+{
+	return FString::Printf(TEXT("FlyingMovement. %s. %s"), *UEnum::GetValueAsString(TEXT("Engine.ENetRole"), GetOwnerRole()), *GetName());
 }
\ No newline at end of file
diff --git a/Engine/Plugins/Runtime/NetworkPrediction/Source/NetworkPrediction/Private/Movement/ParametricMovement.cpp b/Engine/Plugins/Runtime/NetworkPrediction/Source/NetworkPrediction/Private/Movement/ParametricMovement.cpp
--- a/Engine/Plugins/Runtime/NetworkPrediction/Source/NetworkPrediction/Private/Movement/ParametricMovement.cpp
+++ b/Engine/Plugins/Runtime/NetworkPrediction/Source/NetworkPrediction/Private/Movement/ParametricMovement.cpp
@@ -80,7 +80,7 @@ namespace ParametricMovement
 		DrawParams.Lifetime = Parameters.Lifetime;
 		DrawParams.DrawColor = Parameters.GetDebugColor();
 		DrawParams.InWorldText = LexToString(Parameters.Keyframe);
-		DrawParams.LogText = FString::Printf(TEXT("[%d] %s. Location: %s. Rotation: %s"), Parameters.Keyframe, *LexToString(Parameters.Context), *Transform.GetLocation().ToString(), *Transform.GetRotation().Rotator().ToString());
+		DrawParams.LogText = FString::Printf(TEXT("[%d] %s. Position: %.4f. Location: %s. Rotation: %s"), Parameters.Keyframe, *LexToString(Parameters.Context), Position, *Transform.GetLocation().ToString(), *Transform.GetRotation().Rotator().ToString());
 
 		BaseMovementDriver.DrawDebug(DrawParams);
 	}
@@ -109,7 +109,7 @@ IReplicationProxy* UParametricMovementComponent::InstantiateNetworkSimulation()
 void UParametricMovementComponent::InitializeForNetworkRole(ENetRole Role)
 {	
 	check(NetworkSim);
-	NetworkSim->InitializeForNetworkRole(Role, true /* fixme IsLocallyControlled() */, GetSimulationInitParameters(Role));
+	NetworkSim->InitializeForNetworkRole(Role, GetSimulationInitParameters(Role));
 }
 
 FNetworkSimulationModelInitParameters UParametricMovementComponent::GetSimulationInitParameters(ENetRole Role)
@@ -154,6 +154,11 @@ void UParametricMovementComponent::FinalizeFrame(const ParametricMovement::FMove
 	UpdatedComponent->SetWorldTransform(NewTransform, false, nullptr, ETeleportType::TeleportPhysics);
 }
 
+FString UParametricMovementComponent::GetDebugName() const
+{
+	return FString::Printf(TEXT("ParametricMovement. %s. %s"), *UEnum::GetValueAsString(TEXT("Engine.ENetRole"), GetOwnerRole()), *GetName());
+}
+
 void UParametricMovementComponent::AdvanceParametricTime(const float InPosition, const float InPlayRate, float &OutPosition, float& OutPlayRate, const float DeltaTimeSeconds) const
 {
 	// Real simple oscillation for now
diff --git a/Engine/Plugins/Runtime/NetworkPrediction/Source/NetworkPrediction/Private/NetworkPredictionComponent.cpp b/Engine/Plugins/Runtime/NetworkPrediction/Source/NetworkPrediction/Private/NetworkPredictionComponent.cpp
--- a/Engine/Plugins/Runtime/NetworkPrediction/Source/NetworkPrediction/Private/NetworkPredictionComponent.cpp
+++ b/Engine/Plugins/Runtime/NetworkPrediction/Source/NetworkPrediction/Private/NetworkPredictionComponent.cpp
@@ -62,7 +62,7 @@ FNetworkSimulationModelInitParameters UNetworkPredictionComponent::GetSimulation
 {
 	// These are reasonable defaults but may not be right for everyone
 	FNetworkSimulationModelInitParameters InitParams;
-	InitParams.InputBufferSize = Role != ROLE_SimulatedProxy ? 32 : 0;
+	InitParams.InputBufferSize = Role != ROLE_SimulatedProxy ? 32 : 2;
 	InitParams.SyncedBufferSize = Role != ROLE_AutonomousProxy ? 2 : 32;
 	InitParams.AuxBufferSize = Role != ROLE_AutonomousProxy ? 2 : 32;
 	InitParams.DebugBufferSize = 32;
diff --git a/Engine/Plugins/Runtime/NetworkPrediction/Source/NetworkPrediction/Public/Movement/FlyingMovement.h b/Engine/Plugins/Runtime/NetworkPrediction/Source/NetworkPrediction/Public/Movement/FlyingMovement.h
--- a/Engine/Plugins/Runtime/NetworkPrediction/Source/NetworkPrediction/Public/Movement/FlyingMovement.h
+++ b/Engine/Plugins/Runtime/NetworkPrediction/Source/NetworkPrediction/Public/Movement/FlyingMovement.h
@@ -160,6 +160,7 @@ public:
 
 	IBaseMovementDriver& GetBaseMovementDriver() override final { return *static_cast<IBaseMovementDriver*>(this); }
 
+	FString GetDebugName() const override;
 	void InitSyncState(FlyingMovement::FMoveState& OutSyncState) const override;
 	void PreSimSync(const FlyingMovement::FMoveState& SyncState) override;
 	void ProduceInput(const FlyingMovement::TSimTime& SimFrameTime, FlyingMovement::FInputCmd& Cmd) override;
diff --git a/Engine/Plugins/Runtime/NetworkPrediction/Source/NetworkPrediction/Public/Movement/ParametricMovement.h b/Engine/Plugins/Runtime/NetworkPrediction/Source/NetworkPrediction/Public/Movement/ParametricMovement.h
--- a/Engine/Plugins/Runtime/NetworkPrediction/Source/NetworkPrediction/Public/Movement/ParametricMovement.h
+++ b/Engine/Plugins/Runtime/NetworkPrediction/Source/NetworkPrediction/Public/Movement/ParametricMovement.h
@@ -146,6 +146,7 @@ class NETWORKPREDICTION_API UParametricMovementComponent : public UBaseMovementC
 	virtual void TickSimulation(float DeltaTimeSeconds) override;
 
 	// Base TNetworkModelSimulation driver
+	FString GetDebugName() const override;
 	void InitSyncState(ParametricMovement::FMoveState& OutSyncState) const override;
 	void FinalizeFrame(const ParametricMovement::FMoveState& SyncState) override;
 	void ProduceInput(const ParametricMovement::TSimTime& SimFrameTime, ParametricMovement::FInputCmd& Cmd);
diff --git a/Engine/Plugins/Runtime/NetworkPrediction/Source/NetworkPrediction/Public/NetworkPredictionTypes.h b/Engine/Plugins/Runtime/NetworkPrediction/Source/NetworkPrediction/Public/NetworkPredictionTypes.h
--- a/Engine/Plugins/Runtime/NetworkPrediction/Source/NetworkPrediction/Public/NetworkPredictionTypes.h
+++ b/Engine/Plugins/Runtime/NetworkPrediction/Source/NetworkPrediction/Public/NetworkPredictionTypes.h
@@ -3,7 +3,6 @@
 #pragma once
 
 #include "Engine/EngineTypes.h"
-#include "Net/UnrealNetwork.h" // For MakeRelative
 #include "NetworkPredictionTypes.generated.h"
 
 // -------------------------------------------------------------------------------------------------------------------------------
diff --git a/Engine/Plugins/Runtime/NetworkPrediction/Source/NetworkPrediction/Public/NetworkSimulationModel.h b/Engine/Plugins/Runtime/NetworkPrediction/Source/NetworkPrediction/Public/NetworkSimulationModel.h
--- a/Engine/Plugins/Runtime/NetworkPrediction/Source/NetworkPrediction/Public/NetworkSimulationModel.h
+++ b/Engine/Plugins/Runtime/NetworkPrediction/Source/NetworkPrediction/Public/NetworkSimulationModel.h
@@ -31,8 +31,8 @@ template <
 	typename TRepProxySimulated =	TReplicator_Simulated	<TInternalBufferTypes<TUserBufferTypes, InTTickSettings>,	InTTickSettings>,
 
 	// Defines how replication happens on these special channels, but doesn't dictate how simulation evolves
-	typename TRepProxyReplay =		TReplicator_DynamicSequence	<TInternalBufferTypes<TUserBufferTypes, InTTickSettings>,	InTTickSettings, ENetworkSimBufferTypeId::Sync,  3>,
-	typename TRepProxyDebug =		TReplicator_DynamicSequence	<TInternalBufferTypes<TUserBufferTypes, InTTickSettings>,	InTTickSettings, ENetworkSimBufferTypeId::Debug, 3>
+	typename TRepProxyReplay =		TReplicator_Sequence	<TInternalBufferTypes<TUserBufferTypes, InTTickSettings>,	InTTickSettings, ENetworkSimBufferTypeId::Sync,  3>,
+	typename TRepProxyDebug =		TReplicator_Sequence	<TInternalBufferTypes<TUserBufferTypes, InTTickSettings>,	InTTickSettings, ENetworkSimBufferTypeId::Debug, 3>
 >
 class TNetworkedSimulationModel : public IReplicationProxy
 {
@@ -52,6 +52,8 @@ public:
 	class IDriver
 	{
 	public:
+		virtual FString GetDebugName() const = 0; // Used for debugging. Recommended to emit the simulation name and the actor name/role.
+
 		virtual void InitSyncState(TSyncState& OutSyncState) const = 0;	// Called to create initial value of the sync state.
 		virtual void ProduceInput(const TSimTime&, typename TUserBufferTypes::TInputCmd&) = 0; // Called when the sim is ready to process new local input
 		virtual void FinalizeFrame(const TSyncState& SyncState) = 0; // Called from the Network Sim at the end of the sim frame when there is new sync data.
@@ -130,14 +132,14 @@ public:
 		// (E.g, InputCmd @ keyframe=X is used to generate MotionState @ keyframe=X)
 		// This means that SyncedState @ keyframe=0 is always created here via InitSyncState.
 		// This also means that we never actually process InputCmd @ keyframe=0. Which is why LastProcessedInputKeyframe is initialized to 0 ("already processed")
-		// and the buffer has an empty element inserted in InitLocalInputBuffer.
+		// and the buffer has an empty element inserted in InitializeForNetworkRole.
 
 		if (Buffers.Sync.GetHeadKeyframe() != TickInfo.LastProcessedInputKeyframe)
 		{
 			if (TickInfo.LastProcessedInputKeyframe != 0)
 			{
 				// This shouldn't happen, but is not fatal. We are reseting the sync state buffer.
-				UE_LOG(LogNetworkSim, Warning, TEXT("Break in SyncState continuity. LastProcessedInputKeyframe: %d. SyncBuffer.GetHeadKeyframe(): %d. Role=%d"), TickInfo.LastProcessedInputKeyframe, Buffers.Sync.GetHeadKeyframe(), (int32)Parameters.Role);
+				UE_LOG(LogNetworkSim, Warning, TEXT("%s. Break in SyncState continuity. LastProcessedInputKeyframe: %d. SyncBuffer.GetHeadKeyframe(): %d."), *Driver->GetDebugName() , TickInfo.LastProcessedInputKeyframe, Buffers.Sync.GetHeadKeyframe());
 			}
 
 			// We need an initial/current state. Get this from the sim driver
@@ -227,7 +229,7 @@ public:
 		}
 	}	
 	
-	void InitializeForNetworkRole(const ENetRole Role, const bool IsLocallyControlled, const FNetworkSimulationModelInitParameters& Parameters)
+	void InitializeForNetworkRole(const ENetRole Role, const FNetworkSimulationModelInitParameters& Parameters)
 	{
 		Buffers.Input.SetBufferSize(Parameters.InputBufferSize);
 		Buffers.Sync.SetBufferSize(Parameters.SyncedBufferSize);
@@ -243,13 +245,10 @@ public:
 			MyHistoricBuffers->Input.SetBufferSize(Parameters.HistoricBufferSize);
 			MyHistoricBuffers->Sync.SetBufferSize(Parameters.HistoricBufferSize);
 			MyHistoricBuffers->Aux.SetBufferSize(Parameters.HistoricBufferSize);
-		}
+		}		
 
-		if (IsLocallyControlled)
-		{
-			check(Parameters.InputBufferSize > 0); // If you tell me this is locally controlled, you need to have an input buffer.
-			InitLocalInputBuffer();
-		}
+		// We want to start with an empty command in the input buffer. The sync buffer will be populated @ frame 0 with the "current" state when we actually sim. This keeps them in sync
+		*Buffers.Input.GetWriteNext() = TInputCmd();
 	}
 
 	void NetSerializeProxy(EReplicationProxyTarget Target, const FNetSerializeParams& Params)
@@ -303,15 +302,6 @@ public:
 			return 0;
 		};
 	}
-
-	void InitLocalInputBuffer()
-	{
-		// Buffer should also be empty before calling this
-		check(Buffers.Input.GetHeadKeyframe() == INDEX_NONE);
-
-		// We want to start with an empty command in the input buffer. See notes in input buffer processing function.
-		*Buffers.Input.GetWriteNext() = TInputCmd();
-	}	
 	
 	TSimulationTickState<TTickSettings> TickInfo;	// Manages simulation time and what inputs we are processed
 
diff --git a/Engine/Plugins/Runtime/NetworkPrediction/Source/NetworkPrediction/Public/NetworkSimulationModelBuffer.h b/Engine/Plugins/Runtime/NetworkPrediction/Source/NetworkPrediction/Public/NetworkSimulationModelBuffer.h
--- a/Engine/Plugins/Runtime/NetworkPrediction/Source/NetworkPrediction/Public/NetworkSimulationModelBuffer.h
+++ b/Engine/Plugins/Runtime/NetworkPrediction/Source/NetworkPrediction/Public/NetworkSimulationModelBuffer.h
@@ -86,7 +86,7 @@ struct TReplicationBuffer
 		const int32 Offset = (Position % Data.Num());
 		return &Data[Offset];
 	}
-
+	
 	const T* FindElementByKeyframe(int32 Key) const
 	{
 		return FindElementByKeyframeImpl(Key);
@@ -168,7 +168,8 @@ struct TReplicationBuffer
 	int32 GetDirtyCount() const { return DirtyCount; }
 
 	// Create an iterator from tail->head. Note that no matter what template class is, this WILL iterate correctly across the element. The templated type
-	// is just for casting the return element type. E.g, it is fine to use <uint8> in generic code, this will still step by 'StructSize'.
+	// is just for casting the return element type. E.g, it is fine to use <uint8> in generic code, this will always step by 'StructSize'.
+	// FIXME: this is slow using FindElementByKeyframe. In general this entire class needs to be optimized
 	struct TGenericIterator
 	{
 		TGenericIterator(TReplicationBuffer<T>& InBuffer) : Buffer(InBuffer)
diff --git a/Engine/Plugins/Runtime/NetworkPrediction/Source/NetworkPrediction/Public/NetworkSimulationModelCVars.h b/Engine/Plugins/Runtime/NetworkPrediction/Source/NetworkPrediction/Public/NetworkSimulationModelCVars.h
new file mode 100644
--- /dev/null
+++ b/Engine/Plugins/Runtime/NetworkPrediction/Source/NetworkPrediction/Public/NetworkSimulationModelCVars.h
@@ -0,0 +1,63 @@
+// Copyright 1998-2019 Epic Games, Inc. All Rights Reserved.
+
+#pragma once
+
+// ------------------------------------------------------------------------------------------------------------
+//	"Shipping const" cvars: cvars that should compile out to const functions in shipping/test builds
+//		This got a little tricky due to templated usage across different modules. Previous patterns for
+//		this thing didn't work. This implementation requires manual finding of consolve variables so will be 
+//		a bit slower but shouldn't matter since it is compiled out of shipping/test.
+// ------------------------------------------------------------------------------------------------------------
+
+// Whether to treat these cvars as consts
+#define NETSIM_CONST_CVARS (UE_BUILD_SHIPPING || UE_BUILD_TEST)
+
+// This is required because these cvars live in header files that will be included across different compilation units
+// Just using standard FAutoConsoleVariableRef will cause multiple registrations of the same variable
+struct FConditionalAutoConsoleRegister
+{
+	FConditionalAutoConsoleRegister(const TCHAR* VarName, int32 Value, const TCHAR* Help)
+	{
+		if (!IConsoleManager::Get().FindConsoleVariable(VarName, false))
+		{
+			IConsoleManager::Get().RegisterConsoleVariable(VarName, Value, Help, ECVF_Cheat);
+		}
+	}
+
+	FConditionalAutoConsoleRegister(const TCHAR* VarName, float Value, const TCHAR* Help)
+	{
+		if (!IConsoleManager::Get().FindConsoleVariable(VarName, false))
+		{
+			IConsoleManager::Get().RegisterConsoleVariable(VarName, Value, Help, ECVF_Cheat);
+		}
+	}
+};
+
+
+#if NETSIM_CONST_CVARS
+#define NETSIM_DEVCVAR_SHIPCONST_INT(Var,Value,VarName,Help) \
+	inline int32 Var() { return Value; }
+#else
+#define NETSIM_DEVCVAR_SHIPCONST_INT(Var,Value,VarName,Help) \
+	static FConditionalAutoConsoleRegister Var##Auto(TEXT(VarName),(int32)Value,TEXT(Help)); \
+	inline int32 Var() \
+	{ \
+		static const auto* Existing = IConsoleManager::Get().FindConsoleVariable(TEXT(VarName), false); \
+		check(Existing) \
+		return Existing->GetInt(); \
+	}
+#endif
+
+#if NETSIM_CONST_CVARS
+#define NETSIM_DEVCVAR_SHIPCONST_FLOAT(Var,Value,VarName,Help) \
+	inline float Var() { return Value; }
+#else
+#define NETSIM_DEVCVAR_SHIPCONST_FLOAT(Var,Value,VarName,Help) \
+	static FConditionalAutoConsoleRegister Var##Auto(TEXT(VarName),(float)Value,TEXT(Help)); \
+	inline float Var() \
+	{ \
+		static const auto* Existing = IConsoleManager::Get().FindConsoleVariable(TEXT(VarName), false); \
+		check(Existing) \
+		return Existing->GetFloat(); \
+	}
+#endif
\ No newline at end of file
diff --git a/Engine/Plugins/Runtime/NetworkPrediction/Source/NetworkPrediction/Public/NetworkSimulationModelDebugger.h b/Engine/Plugins/Runtime/NetworkPrediction/Source/NetworkPrediction/Public/NetworkSimulationModelDebugger.h
--- a/Engine/Plugins/Runtime/NetworkPrediction/Source/NetworkPrediction/Public/NetworkSimulationModelDebugger.h
+++ b/Engine/Plugins/Runtime/NetworkPrediction/Source/NetworkPrediction/Public/NetworkSimulationModelDebugger.h
@@ -15,18 +15,14 @@
 #include "Containers/Array.h"
 #include "Templates/UniquePtr.h"
 #include "NetworkSimulationModel.h"
+#include "NetworkSimulationModelCvars.h"
 
 DEFINE_LOG_CATEGORY_STATIC(LogNetworkSimDebug, Log, All);
 
 namespace NetworkSimulationModelDebugCVars
 {
-static int32 DrawKeyframes = 1;
-static FAutoConsoleVariableRef CVarDrawKeyframes(TEXT("nsm.debug.DrawKeyFrames"),
-	DrawKeyframes, TEXT("Draws keyframe data (text) in debug graphs"), ECVF_Default);
-
-static int32 GatherServerSidePIE = 1;
-static FAutoConsoleVariableRef CVarGatherServerSidePIE(TEXT("nsm.debug.GatherServerSide"),
-	GatherServerSidePIE, TEXT("Whenever we gather debug info from a client side actor, also gather server side equivelent. Only works in PIE."), ECVF_Default);
+	NETSIM_DEVCVAR_SHIPCONST_INT(DrawKeyframes, 1, "nsm.debug.DrawKeyFrames", "Draws keyframe data (text) in debug graphs");
+	NETSIM_DEVCVAR_SHIPCONST_INT(GatherServerSidePIE, 1, "nsm.debug.GatherServerSide", "Whenever we gather debug info from a client side actor, also gather server side equivelent. Only works in PIE.");
 }
 
 struct FNetworkSimulationModelDebuggerManager;
@@ -148,7 +144,7 @@ struct NETWORKPREDICTION_API FNetworkSimulationModelDebuggerManager: public FTic
 			C->DrawItem(*Item.Get());
 		}
 
-		if (NetworkSimulationModelDebugCVars::DrawKeyframes > 0)
+		if (NetworkSimulationModelDebugCVars::DrawKeyframes() > 0)
 		{
 			for (auto& Item : CanvasItems[1])
 			{
@@ -285,7 +281,7 @@ private:
 			if (It.Value()->IsActive())
 			{
 				It.Value()->GatherCurrent(*this, C);
-				if (NetworkSimulationModelDebugCVars::GatherServerSidePIE > 0)
+				if (NetworkSimulationModelDebugCVars::GatherServerSidePIE() > 0)
 				{
 					if (AActor* ServerSideActor = Cast<AActor>(FindReplicatedObjectOnPIEServer(Owner)))
 					{
diff --git a/Engine/Plugins/Runtime/NetworkPrediction/Source/NetworkPrediction/Public/NetworkSimulationModelReplicators.h b/Engine/Plugins/Runtime/NetworkPrediction/Source/NetworkPrediction/Public/NetworkSimulationModelReplicators.h
--- a/Engine/Plugins/Runtime/NetworkPrediction/Source/NetworkPrediction/Public/NetworkSimulationModelReplicators.h
+++ b/Engine/Plugins/Runtime/NetworkPrediction/Source/NetworkPrediction/Public/NetworkSimulationModelReplicators.h
@@ -2,19 +2,19 @@
 
 #pragma once
 
+#include "Net/UnrealNetwork.h" // For MakeRelative
+#include "NetworkSimulationModelCvars.h"
+
 // ----------------------------------------------------------------------------------------------------------------------------------------------------------------------
 //	CVars and compile time constants
 // ----------------------------------------------------------------------------------------------------------------------------------------------------------------------
 
 namespace NetworkSimulationModelCVars
 {
-static int32 EnableLocalPrediction = 1;
-static FAutoConsoleVariableRef CVarEnableLocalPrediction(TEXT("ns.EnableLocalPrediction"),
-	EnableLocalPrediction, TEXT("Toggle local prediction. Meant for debugging."), ECVF_Default);
-
-static int32 ForceReconcile = 0;
-static FAutoConsoleVariableRef CVarForceReconcile(TEXT("ns.ForceReconcile"),
-	ForceReconcile, TEXT("Forces reconcile even if state does not differ. E.g, force resimulation after every netupdate."), ECVF_Default);
+	NETSIM_DEVCVAR_SHIPCONST_INT(EnableLocalPrediction, 1, "ns.EnableLocalPrediction", "Toggle local prediction.");
+	NETSIM_DEVCVAR_SHIPCONST_INT(EnableSimulatedReconcile, 1, "ns.EnableSimulatedReconcile", "Toggle simulated proxy reconciliation.");
+	NETSIM_DEVCVAR_SHIPCONST_INT(EnableSimulatedExtrapolation, 1, "ns.EnableSimulatedExtrapolation", "Toggle simulated proxy extrapolation.");
+	NETSIM_DEVCVAR_SHIPCONST_INT(ForceReconcile, 0, "ns.ForceReconcile", "Forces reconcile even if state does not differ. E.g, force resimulation after every netupdate.");
 }
 
 #ifndef NETSIM_NETCONSTANT_NUM_BITS_KEYFRAME
@@ -50,7 +50,10 @@ struct FNetworkSimulationSerialization
 };
 
 // ----------------------------------------------------------------------------------------------------------------------------------------------------------------------
-//	Templated Replicators/Reconciliars
+//	Templated Replicators
+//
+//	
+//
 // ----------------------------------------------------------------------------------------------------------------------------------------------------------------------
 
 template<typename TBufferTypes, typename TTickSettings>
@@ -109,7 +112,7 @@ struct TReplicatorBase
 			}
 			else
 			{
-				// Allowed to process all remove commands (we will still check for frame time in the core input processing loop to prevent speed hacks)
+				// Allowed to process all remaining commands (we will still check for frame time in the core input processing loop to prevent speed hacks)
 				TickInfo.MaxAllowedInputKeyframe = Buffers.Input.GetHeadKeyframe();
 
 				// Check for gaps in commands
@@ -127,7 +130,8 @@ struct TReplicatorBase
 //	"Reusable" pieces: mainly for replicating data on specified buffers
 // -------------------------------------------------------------------------------------------------------
 
-template<typename TBufferTypes, typename TTickSettings, typename TBase=TReplicatorEmpty<TBufferTypes, TTickSettings>, bool Enabled=true>
+// Just serializes simulation time. Can be disabled by Enabled templated parameter. Used as base class in other classes below.
+template<typename TBufferTypes, typename TTickSettings, typename TBase=TReplicatorBase<TBufferTypes, TTickSettings>, bool Enabled=true>
 struct TReplicator_SimTime : public TBase
 {
 	void NetSerialize(const FNetSerializeParams& P, TNetworkSimBufferContainer<TBufferTypes>& Buffers, const TSimulationTickState<TTickSettings>& TickInfo)
@@ -140,72 +144,25 @@ struct TReplicator_SimTime : public TBase
 	TNetworkSimTime<TTickSettings> SerializedTime;
 };
 
+// Enabled=false specialization: do nothing
 template<typename TBufferTypes, typename TTickSettings, typename TBase>
 struct TReplicator_SimTime<TBufferTypes, TTickSettings, TBase, false> : public TBase { };
 
-template<typename TBufferTypes, typename TTickSettings, ENetworkSimBufferTypeId BufferId, typename TBase=TReplicatorEmpty<TBufferTypes, TTickSettings>, bool Enabled=true>
-struct TReplicator_HeadKeyframe : public TBase
-{
-	void NetSerialize(const FNetSerializeParams& P, TNetworkSimBufferContainer<TBufferTypes>& Buffers, const TSimulationTickState<TTickSettings>& TickInfo)
-	{
-		TBase::NetSerialize(P, Buffers, TickInfo);
-
-		auto& Buffer = Buffers.template Get<BufferId>();
-		SerializedHeadKeyframe = FNetworkSimulationSerialization::SerializeKeyframe(P.Ar, Buffer.GetHeadKeyframe());
-	}
-
-	int32 SerializedHeadKeyframe;
-};
-
-template<typename TBufferTypes, typename TTickSettings, ENetworkSimBufferTypeId BufferId, typename TBase>
-struct TReplicator_HeadKeyframe<TBufferTypes, TTickSettings, BufferId, TBase, false> : public TBase { };
-
-
-// Helper to make using the above serializers less gross. Selects the appropriate base class depending on SerializeKeyframeNumber and SerializeSimTime
-template<typename TBufferTypes, typename TTickSettings, ENetworkSimBufferTypeId BufferId, bool SerializeKeyframeNumber=false, bool SerializeSimTime=true, typename TRoot=TReplicatorBase<TBufferTypes, TTickSettings>,
-	typename TBase=
-		TReplicator_SimTime<TBufferTypes, TTickSettings,						// Replicates sim time....
-			TReplicator_HeadKeyframe<TBufferTypes, TTickSettings, BufferId,		// Replicates buffer keyframe number...
-				TRoot,						// "final" parent class
-			SerializeKeyframeNumber>,		// ... only if SerializeKeyframeNumber=true
-		SerializeSimTime>>					// ... only if SerializeSimTime=true
-struct TSelectBase
-{
-	using type = TBase;
-};
-
 // -------------------------------------------------------------------------------------------------------
-//
+//	
 // -------------------------------------------------------------------------------------------------------
 
-template<typename TBufferTypes, typename TTickSettings, ENetworkSimBufferTypeId BufferId, int32 NumElements=1, bool SerializeKeyframeNumber=false, bool SerializeSimTime=true, typename TRoot=TReplicatorBase<TBufferTypes, TTickSettings>,
-	typename TBase=TSelectBase<TBufferTypes, TTickSettings, BufferId, SerializeKeyframeNumber, SerializeSimTime, TRoot>>
-struct TReplicator_StaticSequence : public TBase
+// Replicates a sequence of elements. i.e., "the last MaxNumElements".
+// On the receiving side, we merge the sequence into whatever we have locally
+// Keyframes are synchronized. SimTime is also serialized by default (change by changing TBase)
+template<typename TBufferTypes, typename TTickSettings, ENetworkSimBufferTypeId BufferId, int32 MaxNumElements=3, typename TBase=TReplicator_SimTime<TBufferTypes, TTickSettings>>
+struct TReplicator_Sequence : public TBase
 {
 	int32 GetProxyDirtyCount(TNetworkSimBufferContainer<TBufferTypes>& Buffers) const 
 	{
 		return Buffers.template Get<BufferId>().GetDirtyCount() ^ (TBase::GetProxyDirtyCount(Buffers) << 2); 
 	}
 
-	void NetSerialize(const FNetSerializeParams& P, TNetworkSimBufferContainer<TBufferTypes>& Buffers, const TSimulationTickState<TTickSettings>& TickInfo)
-	{
-		TBase::NetSerialize(P, Buffers, TickInfo);
-
-		auto& Buffer = Buffers.template Get<BufferId>();
-		unimplemented();
-	}
-};
-
-/** Replicates the latest sequence of N elements. N is dynamic, not compiled in. */
-template<typename TBufferTypes, typename TTickSettings, ENetworkSimBufferTypeId BufferId, int32 MaxNumElements, bool SerializeKeyframeNumber=false, bool SerializeSimTime=true, typename TRoot=TReplicatorBase<TBufferTypes, TTickSettings>,
-	typename TBase=typename TSelectBase<TBufferTypes, TTickSettings, BufferId, SerializeKeyframeNumber, SerializeSimTime, TRoot>::type>
-struct TReplicator_DynamicSequence : public TBase
-{
-	int32 GetProxyDirtyCount(TNetworkSimBufferContainer<TBufferTypes>& Buffers) const 
-	{
-		return Buffers.template Get<BufferId>().GetDirtyCount() ^ (TBase::GetProxyDirtyCount(Buffers) << 2); 
-	}
-	
 	void NetSerialize(const FNetSerializeParams& P, TNetworkSimBufferContainer<TBufferTypes>& Buffers, const TSimulationTickState<TTickSettings>& TickInfo)
 	{
 		TBase::NetSerialize(P, Buffers, TickInfo);
@@ -240,9 +197,11 @@ struct TReplicator_DynamicSequence : public TBase
 	}
 };
 
-/** Replicates only the latest element. Does not synchronize keyframe or simulation time */
-template<typename TBufferTypes, typename TTickSettings, ENetworkSimBufferTypeId BufferId, typename TBase=TReplicatorBase<TBufferTypes, TTickSettings>>
-struct TReplicator_LatestOnly : public TBase
+// Replicates only the latest single element in the selected buffer.
+// Keyframe is not synchronized: the new element is just added to head.
+// SimTime is serialized by default (change by changing TBase)
+template<typename TBufferTypes, typename TTickSettings, ENetworkSimBufferTypeId BufferId, typename TBase=TReplicator_SimTime<TBufferTypes, TTickSettings>>
+struct TReplicator_Single : public TBase
 {
 	using TState = typename TBufferTypes::template select_type<BufferId>::type;
 
@@ -253,6 +212,8 @@ struct TReplicator_LatestOnly : public TBase
 
 	void NetSerialize(const FNetSerializeParams& P, TNetworkSimBufferContainer<TBufferTypes>& Buffers, const TSimulationTickState<TTickSettings>& TickInfo)
 	{
+		TBase::NetSerialize(P, Buffers, TickInfo);
+
 		FArchive& Ar = P.Ar;
 
 		auto& Buffer = Buffers.template Get<BufferId>();
@@ -278,51 +239,13 @@ struct TReplicator_LatestOnly : public TBase
 
 // Default Replicator for the Server
 //	-Replicates the InputBuffer client->server
-template<typename TBufferTypes, typename TTickSettings, ENetworkSimBufferTypeId BufferId=ENetworkSimBufferTypeId::Input, typename TBase=TReplicator_DynamicSequence<TBufferTypes, TTickSettings, BufferId, 3>>
+template<typename TBufferTypes, typename TTickSettings, ENetworkSimBufferTypeId BufferId=ENetworkSimBufferTypeId::Input, typename TBase=TReplicator_Sequence<TBufferTypes, TTickSettings, BufferId, 3>>
 struct TReplicator_Server : public TBase
 {
 
 };
 
-/*
-template<typename TBufferTypes, typename TTickSettings, typename TBase=TReplicator_DynamicSequence<TBufferTypes, TTickSettings, ENetworkSimBufferTypeId::Sync, 1>>
-struct TReplicator_Simulated : public TBase
-{
-	using TInputCmd = typename TBufferTypes::TInputCmd;
-	using TSyncState = typename TBufferTypes::TSyncState;
-
-	void NetSerialize(const FNetSerializeParams& P, TNetworkSimBufferContainer<TBufferTypes>& Buffers, const TSimulationTickState<TTickSettings>& TickInfo)
-	{
-		TBase::NetSerialize(P, Buffers, TickInfo);
-		bPendingReconcile = true;
-	}
-
-	template<typename T, typename TDriver>
-	void Reconcile(TDriver* Driver, TNetworkSimBufferContainer<TBufferTypes>& Buffers, TSimulationTickState<TTickSettings>& TickInfo)
-	{
-		// Fixme
-		if (bPendingReconcile==false)
-		{
-			return;
-		}
-		bPendingReconcile = false;
-
-		// Generate a new, fake, command since we just added a new sync state to head
-		Buffers.Input.GetWriteNext();
-			
-		// Set our LastProcessedInputKeyframe to fake that we handled it
-		TickInfo.LastProcessedInputKeyframe = Buffers.Input.GetHeadKeyframe();
-		TickInfo.MaxAllowedInputKeyframe = FMath::Max(TickInfo.MaxAllowedInputKeyframe, TickInfo.LastProcessedInputKeyframe);
-	}
-
-private:
-
-	bool bPendingReconcile = false;
-	
-};
-*/
-
-template<typename TBufferTypes, typename TTickSettings, typename TBase=TReplicatorBase<TBufferTypes, TTickSettings>>
+template<typename TBufferTypes, typename TTickSettings, typename TBase=TReplicator_Single<TBufferTypes, TTickSettings, ENetworkSimBufferTypeId::Sync>>
 struct TReplicator_Simulated : public TBase
 {
 	using TInputCmd = typename TBufferTypes::TInputCmd;
@@ -336,22 +259,13 @@ struct TReplicator_Simulated : public TBase
 
 	void NetSerialize(const FNetSerializeParams& P, TNetworkSimBufferContainer<TBufferTypes>& Buffers, const TSimulationTickState<TTickSettings>& TickInfo)
 	{
-		FArchive& Ar = P.Ar;
-		TSyncState* State = nullptr;
-		if (Ar.IsSaving())
-		{
-			State = Buffers.Sync.GetElementFromHead(0);
-			check(State); // We should not be here if the buffer is empty. Want to avoid serializing an "empty" flag at the top here.
-			SerializedTime = TickInfo.TotalProcessedSimulationTime;
-		}
-		else
-		{
-			State = Buffers.Sync.GetWriteNext();
+		TBase::NetSerialize(P, Buffers, TickInfo);
+
+		// Temp: can go away once reconcile is called only after receiving packet
+		if (P.Ar.IsLoading())
+		{		
 			bPendingReconciliation = true;
 		}
-
-		SerializedTime.NetSerialize(Ar);
-		State->NetSerialize(Ar);
 	}
 
 	template<typename T, typename TDriver>
@@ -363,13 +277,29 @@ struct TReplicator_Simulated : public TBase
 		bPendingReconciliation = false;
 		TInputCmd* LastCmd = Buffers.Input.GetElementFromHead(0);
 		
+		{
+			// Right now NetSerialize can be called multiple times before Reconcile can get called. Eventually we'd like to make it 1:1.
+			// Until then, lets handle multiple serializes.
+			int32 DeltaKeyframes = Buffers.Sync.GetHeadKeyframe() - Buffers.Input.GetHeadKeyframe();
+			if (!ensure(DeltaKeyframes >= 0 && DeltaKeyframes < 1000))
+			{
+				// We way out of sync if this gets hit but this will reset things
+				Buffers.Input.ResetNextHeadKeyframe( Buffers.Sync.GetHeadKeyframe() );
+			}
+		}
+		
 		// Generate a new, fake, command since we just added a new sync state to head
-		Buffers.Input.GetWriteNext();
+		while (Buffers.Input.GetHeadKeyframe() < Buffers.Sync.GetHeadKeyframe())
+		{
+			Buffers.Input.GetWriteNext();
+		}
 
-		TNetworkSimTime<TTickSettings> DeltaSimTime = TickInfo.TotalProcessedSimulationTime - SerializedTime;
-		if (DeltaSimTime > TNetworkSimTime<TTickSettings>::FromMSec(0))
+		TNetworkSimTime<TTickSettings> DeltaSimTime = TickInfo.TotalProcessedSimulationTime - this->SerializedTime;
+		if (DeltaSimTime > TNetworkSimTime<TTickSettings>::FromMSec(0) && NetworkSimulationModelCVars::EnableSimulatedReconcile())
 		{
-			// We are extrapolating ahead of the server. We will make up the difference here
+			// This is the "Simulated Reconcile" part:
+			// We have extrapolated ahead of the server. The latest network update is now "in the past" from what we rendered last frame.
+			// To avoid popping or visual artifact, we will make up the difference here since we know the extra simulation time of the server.
 			// (this can easily happen: variance in latency may delay a packet and we may get ahead)
 			TInputCmd* NewCmd = Buffers.Input.GetWriteNext();
 			*NewCmd = LastCmd ? *LastCmd : TInputCmd();
@@ -388,8 +318,9 @@ struct TReplicator_Simulated : public TBase
 		}
 		else
 		{
-			TickInfo.TotalProcessedSimulationTime = SerializedTime;
-			TickInfo.TotalAllowedSimulationTime = SerializedTime;
+			// We are taking the latest server update as head. So set our clocks to it.
+			TickInfo.TotalProcessedSimulationTime = this->SerializedTime;
+			TickInfo.TotalAllowedSimulationTime = this->SerializedTime;
 		}
 
 		// Set our LastProcessedInputKeyframe to fake that we handled it
@@ -397,32 +328,19 @@ struct TReplicator_Simulated : public TBase
 		TickInfo.MaxAllowedInputKeyframe = FMath::Max(TickInfo.MaxAllowedInputKeyframe, TickInfo.LastProcessedInputKeyframe);
 	}
 
-private:
-	
-	TNetworkSimTime<TTickSettings> SerializedTime;
-	bool bPendingReconciliation = false;
-};
-
-/*
-template<typename TBufferTypes, typename TTickSettings, typename TBase>
-struct TReplicator_Simulated<TBufferTypes, TTickSettings, TBase, true> : TReplicator_Simulated<TBufferTypes, TTickSettings, TBase, false>
-{
-	using TInputCmd = typename TBufferTypes::TInputCmd;
-	using TSyncState = typename TBufferTypes::TSyncState;
-
 	template<typename T, typename TDriver, typename TTickParameters>
 	void PreSimTick(TDriver* Driver, TNetworkSimBufferContainer<TBufferTypes>& Buffers, TSimulationTickState<TTickSettings>& TickInfo, const TTickParameters& TickParameters)
 	{
-		// Extrapolation: still needs to be a setting/state somewhere that we can change from actor code
-		if (TInputCmd* SynthesizedCmd = TickInfo.GetNextInputForWrite(Buffers))
+		if (NetworkSimulationModelCVars::EnableSimulatedExtrapolation())
 		{
-			// We will create a single cmd to simulate the remaining time
-			auto TimeToSimulate = TickInfo.GetRemaningAllowedSimulationTime();
-			SynthesizedCmd->GetFrameDeltaTime();
+			TBase::template PreSimTick<T, TDriver, TTickParameters>(Driver, Buffers, TickInfo, TickParameters);
 		}
 	}
+
+private:
+	
+	bool bPendingReconciliation = false;
 };
-*/
 
 
 /** Replicates TSyncState and does basic reconciliation. */
@@ -442,40 +360,6 @@ struct TReplicator_Autonomous : public TBase
 		return Buffers.Sync.GetDirtyCount() ^ (TBase::GetProxyDirtyCount(Buffers) << 2); ;
 	}
 
-	// --------------------------------------------------------------------
-	//	PreSimTick
-	// --------------------------------------------------------------------
-	template<typename T, typename TDriver, typename TTickParameters>
-	void PreSimTick(TDriver* Driver, TNetworkSimBufferContainer<TBufferTypes>& Buffers, TSimulationTickState<TTickSettings>& TickInfo, const TTickParameters& TickParameters)
-	{
-		// If we have a reconcile fault, we cannot continue on with the simulation until it clears itself out. This effectively drops the input time and does not sample new inputs
-		if (bReconcileFaultDetected)
-		{
-			return;
-		}
-
-		// Accumulate real time into the simulation. This will results in more Allowed Simulation Time (or will accumulate internally until we are ready for the next tick)
-		TickInfo.GiveSimulationTime(TickParameters.LocalDeltaTimeSeconds);
-
-		TNetworkSimTime<TTickSettings> DeltaSimTime = TickInfo.GetRemaningAllowedSimulationTime();
-		if (DeltaSimTime.ToRealTimeSeconds() > 0)
-		{
-			// Get next input in buffer, clear it, set delta time, then send it to driver to sample local input
-			if (auto* InputCmd = Buffers.Input.GetWriteNext())
-			{
-				*InputCmd = typename TBufferTypes::TInputCmd();
-				InputCmd->SetFrameDeltaTime(DeltaSimTime);
-				Driver->ProduceInput(DeltaSimTime, *InputCmd);
-			}
-
-			// If we are predicting, then we are allowed to process the new command
-			if (NetworkSimulationModelCVars::EnableLocalPrediction)
-			{
-				TickInfo.MaxAllowedInputKeyframe = Buffers.Input.GetHeadKeyframe();
-			}
-		}
-	}
-
 	// --------------------------------------------------------------------
 	//	NetSerialize
 	// --------------------------------------------------------------------
@@ -524,7 +408,7 @@ struct TReplicator_Autonomous : public TBase
 			// The state the client predicted that corresponds to the state the server just serialized to us
 			if (TSyncState* ClientExistingState = Buffers.Sync.FindElementByKeyframe(SerializedHeadKeyframe))
 			{
-				if (ClientExistingState->ShouldReconcile(*SerializedState) || (NetworkSimulationModelCVars::ForceReconcile > 0))
+				if (ClientExistingState->ShouldReconcile(*SerializedState) || (NetworkSimulationModelCVars::ForceReconcile() > 0))
 				{
 					UE_CLOG(!Buffers.Input.IsValidKeyframe(SerializedHeadKeyframe), LogNetworkSim, Error, TEXT("::NetSerialize: Client InputBuffer does not contain data for frame %d. {%s} {%s}"), SerializedHeadKeyframe, *Buffers.Input.GetBasicDebugStr(), *Buffers.Sync.GetBasicDebugStr());
 					bPendingReconciliation =  true;
@@ -640,6 +524,29 @@ struct TReplicator_Autonomous : public TBase
 		}
 	}
 
+	// --------------------------------------------------------------------
+	//	PreSimTick
+	// --------------------------------------------------------------------
+	template<typename T, typename TDriver, typename TTickParameters>
+	void PreSimTick(TDriver* Driver, TNetworkSimBufferContainer<TBufferTypes>& Buffers, TSimulationTickState<TTickSettings>& TickInfo, const TTickParameters& TickParameters)
+	{
+		// If we have a reconcile fault, we cannot continue on with the simulation until it clears itself out. This effectively drops the input time and does not sample new inputs
+		if (bReconcileFaultDetected)
+		{
+			return;
+		}
+
+		const int32 StartingMaxAllowedInputKeyframe = TickInfo.MaxAllowedInputKeyframe;
+
+		TBase::template PreSimTick<T, TDriver, TTickParameters>(Driver, Buffers, TickInfo, TickParameters);
+
+		// If prediction is disabled we need to reset our MaxAllowedInputKeyframe so we don't process the command(s) generated above
+		if (NetworkSimulationModelCVars::EnableLocalPrediction() == 0)
+		{
+			TickInfo.MaxAllowedInputKeyframe = StartingMaxAllowedInputKeyframe;
+		}
+	}
+
 private:
 	
 	TReplicationBuffer<TSyncState> ReconciliationBuffer;
diff --git a/Engine/Plugins/Runtime/NetworkPrediction/readme.txt b/Engine/Plugins/Runtime/NetworkPrediction/readme.txt
--- a/Engine/Plugins/Runtime/NetworkPrediction/readme.txt
+++ b/Engine/Plugins/Runtime/NetworkPrediction/readme.txt
@@ -101,6 +101,14 @@ Smoothing: Taking the output of the simulation and applying an additional layer
 // ----------------------------------------------------------------------------------------------------------
 
 
+Update (8-14-19)
+-Some more cleanup. Cvar fixes and other consolidation. Now time for vacation (2 weeks). 
+TODO still:
+-Improve debugger (has rotted a bit due to refactors)
+-Network fault handling, disable/enable local prediction causes some bad transitions
+-Initialization of system, specifically InitializeForNetworkRole can be much better now.
+	-Consider using templates to inline allocate all memory. Maybe make auto proxy buffers allocated on demant.
+
 Update (8-13-19)
 -Refactor really feeling good, things are falling into place. Still a few things to do and bugs to chase down.
 -Added simulated proxy notes above to detail current plan
diff --git a/Engine/Plugins/Runtime/NetworkPredictionExtras/Source/NetworkPredictionExtras/Private/MockNetworkSimulation.cpp b/Engine/Plugins/Runtime/NetworkPredictionExtras/Source/NetworkPredictionExtras/Private/MockNetworkSimulation.cpp
--- a/Engine/Plugins/Runtime/NetworkPredictionExtras/Source/NetworkPredictionExtras/Private/MockNetworkSimulation.cpp
+++ b/Engine/Plugins/Runtime/NetworkPredictionExtras/Source/NetworkPredictionExtras/Private/MockNetworkSimulation.cpp
@@ -142,7 +142,7 @@ IReplicationProxy* UMockNetworkSimulationComponent::InstantiateNetworkSimulation
 // Child classes should override this an initialize their NetworkSim here
 void UMockNetworkSimulationComponent::InitializeForNetworkRole(ENetRole Role)
 {	
-	NetworkSim->InitializeForNetworkRole(Role, IsLocallyControlled(), GetSimulationInitParameters(Role));
+	NetworkSim->InitializeForNetworkRole(Role, GetSimulationInitParameters(Role));
 }
 
 void UMockNetworkSimulationComponent::TickComponent(float DeltaTime, enum ELevelTick TickType, FActorComponentTickFunction *ThisTickFunction)
@@ -204,6 +204,11 @@ void UMockNetworkSimulationComponent::FinalizeFrame(const FMockSyncState& SyncSt
 	MockValue = SyncState.Total;
 }
 
+FString UMockNetworkSimulationComponent::GetDebugName() const
+{
+	return FString::Printf(TEXT("MockSim. %s. %s"), *UEnum::GetValueAsString(TEXT("Engine.ENetRole"), GetOwnerRole()), *GetName());
+}
+
 FTransform UMockNetworkSimulationComponent::GetDebugWorldTransform() const
 {
 	return GetOwner()->GetActorTransform();
diff --git a/Engine/Plugins/Runtime/NetworkPredictionExtras/Source/NetworkPredictionExtras/Private/NetworkPredictionExtrasFlyingPawn.cpp b/Engine/Plugins/Runtime/NetworkPredictionExtras/Source/NetworkPredictionExtras/Private/NetworkPredictionExtrasFlyingPawn.cpp
--- a/Engine/Plugins/Runtime/NetworkPredictionExtras/Source/NetworkPredictionExtras/Private/NetworkPredictionExtrasFlyingPawn.cpp
+++ b/Engine/Plugins/Runtime/NetworkPredictionExtras/Source/NetworkPredictionExtras/Private/NetworkPredictionExtrasFlyingPawn.cpp
@@ -149,7 +149,11 @@ void ANetworkPredictionExtrasFlyingPawn::ProduceInput(const FlyingMovement::TSim
 	// answers and will depend on the game and its specific needs. In general, at this time, I'd recommend aim assist and lock on 
 	// targeting systems to happen /outside/ of the system, i.e, here. But I can think of scenarios where that may not be ideal too.
 
-	check(Controller);
+	if (Controller == nullptr)
+	{
+		// We don't have a local controller so we can't run the code below. This is ok. Simulated proxies will just use previous input when extrapolating
+		return;
+	}
 	
 	if (USpringArmComponent* SpringComp = FindComponentByClass<USpringArmComponent>())
 	{
diff --git a/Engine/Plugins/Runtime/NetworkPredictionExtras/Source/NetworkPredictionExtras/Public/MockNetworkSimulation.h b/Engine/Plugins/Runtime/NetworkPredictionExtras/Source/NetworkPredictionExtras/Public/MockNetworkSimulation.h
--- a/Engine/Plugins/Runtime/NetworkPredictionExtras/Source/NetworkPredictionExtras/Public/MockNetworkSimulation.h
+++ b/Engine/Plugins/Runtime/NetworkPredictionExtras/Source/NetworkPredictionExtras/Public/MockNetworkSimulation.h
@@ -162,6 +162,7 @@ private:
 	
 public:
 
+	FString GetDebugName() const override;
 	void InitSyncState(FMockSyncState& OutSyncState) const override;
 	void FinalizeFrame(const FMockSyncState& SyncState) override;
 	void ProduceInput(const FMockNetworkSimulation::TSimTime& SimFrameTime, FMockInputCmd& Cmd);
