Fix warning due to pointer cast into bool

diff --git a/Engine/Plugins/Runtime/NetworkPrediction/Source/NetworkPrediction/Private/NetworkPredictionTypes.cpp b/Engine/Plugins/Runtime/NetworkPrediction/Source/NetworkPrediction/Private/NetworkPredictionTypes.cpp
--- a/Engine/Plugins/Runtime/NetworkPrediction/Source/NetworkPrediction/Private/NetworkPredictionTypes.cpp
+++ b/Engine/Plugins/Runtime/NetworkPrediction/Source/NetworkPrediction/Private/NetworkPredictionTypes.cpp
@@ -117,7 +117,7 @@ FNetSimTickParametersBase::FNetSimTickParametersBase(AActor* OwningActor)
 {
 	// Default implementation that will probably be good for most actors
 	Role = OwningActor->GetLocalRole();
-	const bool bHasNetConnection = OwningActor->GetNetConnection();
+	const bool bHasNetConnection = OwningActor->GetNetConnection() != nullptr;
 
 	// The only default case we really don;t want generating local input is on the server when there is a net owning player (auto proxy)
 	// This implies that we will (sim) extrapolate by default. To opt out of extrapolation, set to false or you can just not tick the net sim
