Fix inconsistent case for include filenames in NetworkPrediction plugin files.

diff --git a/Engine/Plugins/Runtime/NetworkPrediction/Source/NetworkPrediction/Public/NetworkSimulationModelDebugger.h b/Engine/Plugins/Runtime/NetworkPrediction/Source/NetworkPrediction/Public/NetworkSimulationModelDebugger.h
--- a/Engine/Plugins/Runtime/NetworkPrediction/Source/NetworkPrediction/Public/NetworkSimulationModelDebugger.h
+++ b/Engine/Plugins/Runtime/NetworkPrediction/Source/NetworkPrediction/Public/NetworkSimulationModelDebugger.h
@@ -15,7 +15,7 @@
 #include "Containers/Array.h"
 #include "Templates/UniquePtr.h"
 #include "NetworkSimulationModel.h"
-#include "NetworkSimulationModelCvars.h"
+#include "NetworkSimulationModelCVars.h"
 
 DEFINE_LOG_CATEGORY_STATIC(LogNetworkSimDebug, Log, All);
 
diff --git a/Engine/Plugins/Runtime/NetworkPrediction/Source/NetworkPrediction/Public/NetworkSimulationModelReplicators.h b/Engine/Plugins/Runtime/NetworkPrediction/Source/NetworkPrediction/Public/NetworkSimulationModelReplicators.h
--- a/Engine/Plugins/Runtime/NetworkPrediction/Source/NetworkPrediction/Public/NetworkSimulationModelReplicators.h
+++ b/Engine/Plugins/Runtime/NetworkPrediction/Source/NetworkPrediction/Public/NetworkSimulationModelReplicators.h
@@ -3,7 +3,7 @@
 #pragma once
 
 #include "Net/UnrealNetwork.h" // For MakeRelative
-#include "NetworkSimulationModelCvars.h"
+#include "NetworkSimulationModelCVars.h"
 
 // ----------------------------------------------------------------------------------------------------------------------------------------------------------------------
 //	CVars and compile time constants
