Network Prediction: debugger improvements
-Keep seperate Debug buffers for local and remote. This is to allow us to correlate client and server simulation advancement.
-Started reworking canvas debug hud. We now show client and server graphs and draw lines showing network traffic between the two "timelines". Still more work to be done here.

#rnx

diff --git a/Engine/Plugins/Runtime/NetworkPrediction/Source/NetworkPrediction/Public/NetworkPredictionTypes.h b/Engine/Plugins/Runtime/NetworkPrediction/Source/NetworkPrediction/Public/NetworkPredictionTypes.h
--- a/Engine/Plugins/Runtime/NetworkPrediction/Source/NetworkPrediction/Public/NetworkPredictionTypes.h
+++ b/Engine/Plugins/Runtime/NetworkPrediction/Source/NetworkPrediction/Public/NetworkPredictionTypes.h
@@ -53,12 +53,21 @@ struct FNetSimProcessedFrameDebugInfo
 	int32 LastProcessedKeyframe; // What LastProcessedKeyframe was at the end of the frame. Does NOT mean we processed it this frame!
 	int32 HeadKeyframe; // Head keyframe of the inputbuffer when the frame ended
 
+	float RemainingAllowedSimulationTimeSeconds;
+	
+	int32 LastSentInputKeyframe;
+	int32 LastReceivedInputKeyframe;
+
 	void NetSerialize(const FNetSerializeParams& P)
 	{	
 		P.Ar << LocalGFrameNumber;
 		P.Ar << LocalDeltaTimeSeconds;
 		P.Ar << LastProcessedKeyframe;
 		P.Ar << HeadKeyframe;
+		P.Ar << RemainingAllowedSimulationTimeSeconds;
+
+		P.Ar << LastSentInputKeyframe;
+		P.Ar << LastReceivedInputKeyframe;
 
 		SafeNetSerializeTArray_Default<31>(P.Ar, ProcessedKeyframes);
 	}
diff --git a/Engine/Plugins/Runtime/NetworkPrediction/Source/NetworkPrediction/Public/NetworkSimulationModel.h b/Engine/Plugins/Runtime/NetworkPrediction/Source/NetworkPrediction/Public/NetworkSimulationModel.h
--- a/Engine/Plugins/Runtime/NetworkPrediction/Source/NetworkPrediction/Public/NetworkSimulationModel.h
+++ b/Engine/Plugins/Runtime/NetworkPrediction/Source/NetworkPrediction/Public/NetworkSimulationModel.h
@@ -32,7 +32,7 @@ template <
 
 	// Defines how replication happens on these special channels, but doesn't dictate how simulation evolves
 	typename TRepProxyReplay =		TReplicator_Sequence	<TInternalBufferTypes<TUserBufferTypes, InTTickSettings>,	InTTickSettings, ENetworkSimBufferTypeId::Sync,  3>,
-	typename TRepProxyDebug =		TReplicator_Sequence	<TInternalBufferTypes<TUserBufferTypes, InTTickSettings>,	InTTickSettings, ENetworkSimBufferTypeId::Debug, 3>
+	typename TRepProxyDebug =		TReplicator_Debug		<TInternalBufferTypes<TUserBufferTypes, InTTickSettings>,	InTTickSettings>
 >
 class TNetworkedSimulationModel : public IReplicationProxy
 {
@@ -71,12 +71,40 @@ public:
 	template<typename TDriver>
 	void Tick(TDriver* Driver, const FTickParameters& Parameters)
 	{
-		TDebugState* const DebugState = GetNextDebugStateWrite();
+		// Update previous DebugState based on what we (might) have sent *after* our last Tick 
+		// (property replication and ServerRPC get sent after the tick, rather than forcing awkward callback into the NetSim post replication, we can check it here)
+		if (auto* DebugBuffer = GetLocalDebugBuffer())
+		{
+			if (TDebugState* const PrevDebugState = DebugBuffer->FindElementByKeyframe(DebugBuffer->GetHeadKeyframe()))
+			{
+				if (Parameters.Role == ROLE_AutonomousProxy)
+				{
+					PrevDebugState->LastSentInputKeyframe = RepProxy_ServerRPC.GetLastSerializedKeyframe();
+				}
+				else if (Parameters.Role == ROLE_Authority)
+				{
+					PrevDebugState->LastSentInputKeyframe = RepProxy_Autonomous.GetLastSerializedKeyframe();
+				}
+			}
+		}
+
+		// Current frame debug state
+		TDebugState* const DebugState = GetNextLocalDebugStateWrite();
 		if (DebugState)
 		{
+			*DebugState = TDebugState();
 			DebugState->LocalDeltaTimeSeconds = Parameters.LocalDeltaTimeSeconds;
 			DebugState->LocalGFrameNumber = GFrameNumber;
 			DebugState->ProcessedKeyframes.Reset();
+			
+			if (Parameters.Role == ROLE_AutonomousProxy)
+			{
+				DebugState->LastReceivedInputKeyframe = RepProxy_Autonomous.GetLastSerializedKeyframe();
+			}
+			else if (Parameters.Role == ROLE_Authority)
+			{
+				DebugState->LastReceivedInputKeyframe = RepProxy_ServerRPC.GetLastSerializedKeyframe();
+			}
 		}
 
 		// --------------------------------------------------------------------------------------------------------------------------
@@ -163,7 +191,7 @@ public:
 			if (TInputCmd* NextCmd = Buffers.Input.FindElementByKeyframe(Keyframe))
 			{
 				// We have an unprocessed command, do we have enough allotted simulation time to process it?
-				if (TickInfo.GetRemaningAllowedSimulationTime() >= NextCmd->GetFrameDeltaTime())
+				if (TickInfo.GetRemainingAllowedSimulationTime() >= NextCmd->GetFrameDeltaTime())
 				{
 					// -------------------------------------------------------------------------------------------------
 					//	The core process input command and call ::Update block!
@@ -218,6 +246,7 @@ public:
 		{
 			DebugState->LastProcessedKeyframe = TickInfo.LastProcessedInputKeyframe;
 			DebugState->HeadKeyframe = Buffers.Input.GetHeadKeyframe();
+			DebugState->RemainingAllowedSimulationTimeSeconds = (float)TickInfo.GetRemainingAllowedSimulationTime().ToRealTimeSeconds();
 		}
 
 		// Historical data recording (longer buffers for historical reference)
@@ -235,9 +264,9 @@ public:
 		Buffers.Sync.SetBufferSize(Parameters.SyncedBufferSize);
 		Buffers.Aux.SetBufferSize(Parameters.AuxBufferSize);
 
-		if (GetDebugBuffer())
+		if (GetLocalDebugBuffer())
 		{
-			GetDebugBuffer()->SetBufferSize(Parameters.DebugBufferSize);
+			GetLocalDebugBuffer()->SetBufferSize(Parameters.DebugBufferSize);
 		}
 
 		if (auto* MyHistoricBuffers = GetHistoricBuffers(true))
@@ -345,17 +374,20 @@ private:
 public:
 
 #if NETSIM_MODEL_DEBUG
-	TReplicationBuffer<TDebugState>* GetDebugBuffer() {	return &Buffers.Debug; }
-	TDebugState* GetNextDebugStateWrite() { return Buffers.Debug.GetWriteNext(); }
+	TReplicationBuffer<TDebugState>* GetLocalDebugBuffer() {	return &Buffers.Debug; }
+	TDebugState* GetNextLocalDebugStateWrite() { return Buffers.Debug.GetWriteNext(); }
 	TNetworkSimBufferContainer<TBufferTypes>* GetHistoricBuffers(bool bCreate=false)
 	{
 		if (HistoricBuffers.IsValid() == false && bCreate) { HistoricBuffers.Reset(new TNetworkSimBufferContainer<TBufferTypes>()); }
 		return HistoricBuffers.Get();
 	}
+
+	TReplicationBuffer<TDebugState>* GetRemoteDebugBuffer() {	return &RepProxy_Debug.ReceivedBuffer; }
 #else
-	TReplicationBuffer<TDebugState>* GetDebugBuffer(bool bCreate=false) {	return nullptr; }
-	TDebugState* GetNextDebugStateWrite() { return nullptr; }
+	TReplicationBuffer<TDebugState>* GetLocalDebugBuffer(bool bCreate=false) {	return nullptr; }
+	TDebugState* GetNextLocalDebugStateWrite() { return nullptr; }
 	TNetworkSimBufferContainer<TBufferTypes>* GetHistoricBuffers() { return nullptr; }
+	TReplicationBuffer<TDebugState>* GetRemoteDebugBuffer() {	return nullptr; }
 #endif
 
 private:
diff --git a/Engine/Plugins/Runtime/NetworkPrediction/Source/NetworkPrediction/Public/NetworkSimulationModelDebugger.h b/Engine/Plugins/Runtime/NetworkPrediction/Source/NetworkPrediction/Public/NetworkSimulationModelDebugger.h
--- a/Engine/Plugins/Runtime/NetworkPrediction/Source/NetworkPrediction/Public/NetworkSimulationModelDebugger.h
+++ b/Engine/Plugins/Runtime/NetworkPrediction/Source/NetworkPrediction/Public/NetworkSimulationModelDebugger.h
@@ -248,6 +248,12 @@ struct NETWORKPREDICTION_API FNetworkSimulationModelDebuggerManager: public FTic
 		CanvasItems[1].Emplace( MakeUnique<FCanvasTextItem>(ScreenPosition, FText::FromString(Str), GEngine->GetTinyFont(), Color) );
 	}
 
+	void EmitLine(FVector2D StartPosition, FVector2D EndPosition, FColor Color)
+	{
+		CanvasItems[0].Emplace( MakeUnique<FCanvasLineItem>(StartPosition, EndPosition) );
+		CanvasItems[0].Last()->SetColor(Color);
+	}
+
 private:
 
 	INetworkSimulationModelDebugger* Find(AActor* Actor)
@@ -326,6 +332,7 @@ template <typename TNetSimModel, typename TDriver>
 struct TNetworkSimulationModelDebugger : public INetworkSimulationModelDebugger
 {
 	using TSimTime = typename TNetSimModel::TSimTime;
+	using TDebugState = typename TNetSimModel::TDebugState;
 
 	TNetworkSimulationModelDebugger(TNetSimModel* InNetSim, TDriver* InDriver, AActor* OwningActor, FString InDebugName)
 	{
@@ -338,7 +345,164 @@ struct TNetworkSimulationModelDebugger : public INetworkSimulationModelDebugger
 	~TNetworkSimulationModelDebugger()
 	{
 		
-	}	
+	}
+
+	struct FCachedScreenPositionMap
+	{
+		struct FScreenPositions
+		{
+			void SetSent(const FVector2D& In) { if (SentPosition == FVector2D::ZeroVector) SentPosition = In; }
+			void SetRecv(const FVector2D& In) { if (RecvPosition == FVector2D::ZeroVector) RecvPosition = In; }
+
+			FVector2D SentPosition = FVector2D::ZeroVector;
+			FVector2D RecvPosition = FVector2D::ZeroVector;
+		};
+
+		TMap<int32, FScreenPositions> Keyframes;
+	};
+
+	void GatherDebugGraph(FNetworkSimulationModelDebuggerManager& Out, UCanvas* Canvas, TReplicationBuffer<TDebugState>* DebugBuffer, FRect DrawRect, const float MaxColumnTimeSeconds, const float MaxLocalFrameTime, const FString& Header, FCachedScreenPositionMap& SendCache, FCachedScreenPositionMap& RecvCache)
+	{
+		static float Pad = 2.f;
+		static float BaseLineYPCT = 0.8f;
+
+		auto& InputBuffer = NetworkSim->GetHistoricBuffers() ? NetworkSim->GetHistoricBuffers()->Input : NetworkSim->Buffers.Input;
+
+		if (Canvas && DebugBuffer && DebugBuffer->GetNumValidElements() > 0)
+		{
+			// Outline + Header
+			Out.EmitLine( FVector2D(DrawRect.Min.X, DrawRect.Min.Y), FVector2D(DrawRect.Min.X, DrawRect.Max.Y), FColor::White );
+			Out.EmitLine( FVector2D(DrawRect.Min.X, DrawRect.Min.Y), FVector2D(DrawRect.Max.X, DrawRect.Min.Y), FColor::White );
+			Out.EmitLine( FVector2D(DrawRect.Max.X, DrawRect.Min.Y), FVector2D(DrawRect.Max.X, DrawRect.Max.Y), FColor::White );
+			Out.EmitLine( FVector2D(DrawRect.Min.X, DrawRect.Max.Y), FVector2D(DrawRect.Max.X, DrawRect.Max.Y), FColor::White );
+
+			Out.EmitText( DrawRect.Min, FColor::White, Header );
+
+			// Frame Columns
+			const float BaseLineYPos = DrawRect.Min.Y + (BaseLineYPCT * (DrawRect.Max.Y - DrawRect.Min.Y));
+
+			const float AboveBaseLineSecondsToPixelsY = (BaseLineYPos - DrawRect.Min.Y - Pad) / MaxColumnTimeSeconds;
+			const float BelowBaseLineSecondsToPixelsY = (DrawRect.Max.Y - BaseLineYPos - Pad) / MaxLocalFrameTime;
+
+			const float SecondsToPixelsY = FMath::Min<float>(BelowBaseLineSecondsToPixelsY, AboveBaseLineSecondsToPixelsY);
+
+			Out.EmitLine( FVector2D(DrawRect.Min.X, BaseLineYPos), FVector2D(DrawRect.Max.X, BaseLineYPos), FColor::Black );
+
+
+			FTextSizingParameters TextSizing;
+			TextSizing.DrawFont = GEngine->GetTinyFont();
+			TextSizing.Scaling = FVector2D(1.f,1.f);
+			Canvas->CanvasStringSize(TextSizing, TEXT("00000"));
+
+			const float FixedWidth = TextSizing.DrawXL;
+
+			float ScreenX = DrawRect.Min.X;
+			float ScreenY = BaseLineYPos + Pad;
+
+			for (auto It = DebugBuffer->CreateIterator(); It; ++It)
+			{
+				auto* DebugState = It.Element();
+				const float FrameHeight = SecondsToPixelsY * DebugState->LocalDeltaTimeSeconds;
+				
+				// Green local frame time (below baseline)
+				FColor Color = FColor::Green;
+				Out.EmitQuad(FVector2D( ScreenX, ScreenY), FVector2D( FixedWidth, FrameHeight), Color);
+				Out.EmitText(FVector2D( ScreenX, ScreenY), FColor::Black, FString::Printf(TEXT("%.2f"), (DebugState->LocalDeltaTimeSeconds * 1000.f)));
+
+				// Processed InputcmdsKeyframes (above baseline)
+				float ClientX = ScreenX;
+				float ClientY = ScreenY - Pad;
+
+				for (int32 Keyframe : DebugState->ProcessedKeyframes)
+				{
+					auto* Cmd = InputBuffer.FindElementByKeyframe(Keyframe);
+					if (Cmd)
+					{
+						float ClientSizeX = FixedWidth;
+						float ClientSizeY = SecondsToPixelsY * Cmd->GetFrameDeltaTime().ToRealTimeSeconds();
+
+						FVector2D ScreenPos(ClientX, ClientY - ClientSizeY);
+						Out.EmitQuad(ScreenPos, FVector2D( ClientSizeX, ClientSizeY), FColor::Blue);
+						Out.EmitText(ScreenPos, FColor::White, LexToString(Keyframe));
+						ClientY -= (ClientSizeY + Pad);
+					}
+				}
+
+				// Unprocessed InputCmds (above processed)				
+				for (int32 Keyframe = DebugState->LastProcessedKeyframe+1; Keyframe <= DebugState->HeadKeyframe; ++Keyframe)
+				{
+					if (auto* Cmd = InputBuffer.FindElementByKeyframe(Keyframe))
+					{
+						float ClientSizeX = FixedWidth;
+						float ClientSizeY = SecondsToPixelsY * Cmd->GetFrameDeltaTime().ToRealTimeSeconds();
+						FVector2D ScreenPos(ClientX, ClientY - ClientSizeY);
+						Out.EmitQuad(ScreenPos, FVector2D( ClientSizeX, ClientSizeY), FColor::Red);
+						Out.EmitText(ScreenPos, FColor::White, LexToString(Keyframe));
+						ClientY -= (ClientSizeY + Pad);
+					}
+				}
+
+				// Cache Screen Positions based on keyframe
+				RecvCache.Keyframes.FindOrAdd(DebugState->LastReceivedInputKeyframe).SetRecv(FVector2D(ScreenX, BaseLineYPos));
+				
+				// Advance 
+				ScreenX += FixedWidth + Pad;
+
+				// Send cache
+				SendCache.Keyframes.FindOrAdd(DebugState->LastSentInputKeyframe).SetSent(FVector2D(ScreenX, BaseLineYPos));
+			}
+
+
+			// Remaining Simulation Time
+			TOptional<FVector2D> LastLinePos;
+			FVector2D LinePos(DrawRect.Min.X, BaseLineYPos);
+			for (auto It = DebugBuffer->CreateIterator(); It; ++It)
+			{
+				auto* DebugState = It.Element();
+				
+				LinePos.X += FixedWidth + Pad;
+				LinePos.Y = BaseLineYPos - (DebugState->RemainingAllowedSimulationTimeSeconds * SecondsToPixelsY);
+				
+				if (LastLinePos.IsSet())
+				{
+					Out.EmitLine(LastLinePos.GetValue(), LinePos, FColor::White);
+				}
+
+				LastLinePos = LinePos;
+				
+			}
+		}
+	}
+
+	void CalcMaxColumnFrameTime(TReplicationBuffer<TDebugState>* DebugBuffer, float& MaxInputTime, float& MaxLocalFrameTime)
+	{
+		auto& InputBuffer = NetworkSim->GetHistoricBuffers() ? NetworkSim->GetHistoricBuffers()->Input : NetworkSim->Buffers.Input;
+
+		for (auto It = DebugBuffer->CreateIterator(); It; ++It)
+		{
+			float ColumnTime = 0.f;
+
+			auto* DebugState = It.Element();
+			for (int32 Keyframe : DebugState->ProcessedKeyframes)
+			{
+				auto* Cmd = InputBuffer.FindElementByKeyframe(Keyframe);
+				if (Cmd)
+				{
+					ColumnTime += (float)Cmd->GetFrameDeltaTime().ToRealTimeSeconds();
+				}
+			}
+			for (int32 Keyframe = DebugState->LastProcessedKeyframe+1; Keyframe <= DebugState->HeadKeyframe; ++Keyframe)
+			{
+				if (auto* Cmd = InputBuffer.FindElementByKeyframe(Keyframe))
+				{
+					ColumnTime += (float)Cmd->GetFrameDeltaTime().ToRealTimeSeconds();
+				}
+			}
+
+			MaxInputTime = FMath::Max<float>(ColumnTime, MaxInputTime);
+			MaxLocalFrameTime = FMath::Max<float>(DebugState->LocalDeltaTimeSeconds, MaxLocalFrameTime);
+		}
+	}
 
 	void GatherCurrent(FNetworkSimulationModelDebuggerManager& Out, UCanvas* Canvas) override
 	{
@@ -385,7 +549,9 @@ struct TNetworkSimulationModelDebugger : public INetworkSimulationModelDebugger
 			FString SimulationTimeString = FString::Printf(TEXT("Local SimulationTime: %s. SerialisedSimulationTime: %s. Difference MS: %s"), *NetworkSim->TickInfo.TotalProcessedSimulationTime.ToString(),
 				*NetworkSim->RepProxy_Autonomous.GetLastSerializedSimTime().ToString(), *(NetworkSim->TickInfo.TotalProcessedSimulationTime - NetworkSim->RepProxy_Autonomous.GetLastSerializedSimTime()).ToString());
 			Out.Emit(*SimulationTimeString, Color);
-			
+
+			FString AllowedSimulationTimeString = FString::Printf(TEXT("Allowed Simulation Time: %s. Keyframe: %d/%d/%d"), *NetworkSim->TickInfo.GetRemainingAllowedSimulationTime().ToString(), NetworkSim->TickInfo.MaxAllowedInputKeyframe, NetworkSim->TickInfo.LastProcessedInputKeyframe, NetworkSim->Buffers.Input.GetHeadKeyframe());
+			Out.Emit(*AllowedSimulationTimeString, Color);
 		}
 
 		auto EmitBuffer = [&Out](FString BufferName, auto& Buffer)
@@ -404,89 +570,57 @@ struct TNetworkSimulationModelDebugger : public INetworkSimulationModelDebugger
 		//	Canvas
 		// ------------------------------------------------------------------------------------------------------------------------------------------------
 
-		auto* DebugBuffer = NetworkSim->GetDebugBuffer();
-		if (Canvas && DebugBuffer && DebugBuffer->GetNumValidElements() > 0)
+		if (Canvas)
 		{
-			auto& InputBuffer = NetworkSim->GetHistoricBuffers() ? NetworkSim->GetHistoricBuffers()->Input : NetworkSim->Buffers.Input;
-
-			static float StartPctX = 0.3f;
-			static float StartPctY = 0.6f;
-			static float LocalFrameHeightPct = 0.01f;
+			FRect ServerRect;
+			ServerRect.Min.X = 0.30f * (float)Canvas->SizeX;
+			ServerRect.Max.X = 0.95f * (float)Canvas->SizeX;
+			ServerRect.Min.Y = 0.05f * (float)Canvas->SizeY;
+			ServerRect.Max.Y = 0.45f * (float)Canvas->SizeY;
 
-			float ScreenX = StartPctX * (float)Canvas->SizeX;
-			float ScreenY = StartPctY * (float)Canvas->SizeY;
-			float LocalFrameHeight = LocalFrameHeightPct * (float)Canvas->SizeY;
+			FRect ClientRect;
+			ClientRect.Min.X = 0.30f * (float)Canvas->SizeX;
+			ClientRect.Max.X = 0.95f * (float)Canvas->SizeX;
+			ClientRect.Min.Y = 0.55f * (float)Canvas->SizeY;
+			ClientRect.Max.Y = 0.95f * (float)Canvas->SizeY;
 
-			static float LocalFrameTimeGreen = 1/30.f;
-			static float LocalFrameTimeRed = 1/10.f;
+			float MaxColumnTime = 1.f/60.f;
+			float MaxLocalFrameTime = 1.f/60.f;
+			CalcMaxColumnFrameTime(NetworkSim->GetRemoteDebugBuffer(), MaxColumnTime, MaxLocalFrameTime);
+			CalcMaxColumnFrameTime(NetworkSim->GetLocalDebugBuffer(), MaxColumnTime, MaxLocalFrameTime);
 
-			static float Thickness = 1.f;
-			static float ClientOffsetY = 2.f;
+			FCachedScreenPositionMap ServerToClientCache;
+			FCachedScreenPositionMap ClientToServerCache;
 
-			FTextSizingParameters TextSizing;
-			TextSizing.DrawFont = GEngine->GetTinyFont();
-			TextSizing.Scaling = FVector2D(1.f,1.f);
-			Canvas->CanvasStringSize(TextSizing, TEXT("00000"));
+			GatherDebugGraph(Out, Canvas, NetworkSim->GetRemoteDebugBuffer(), ServerRect, MaxColumnTime, MaxLocalFrameTime, TEXT("Server"), ServerToClientCache, ClientToServerCache);
+			GatherDebugGraph(Out, Canvas, NetworkSim->GetLocalDebugBuffer(), ClientRect, MaxColumnTime, MaxLocalFrameTime, TEXT("Client"), ClientToServerCache, ServerToClientCache);
 
-			float MinWidth = TextSizing.DrawXL;
-			float MinHeight = TextSizing.DrawYL;
-			static float MinHeightSeconds = 1/60.f;	// 60hz frame is drawn at MinHeight.
+			// Network Send/Recv lines
 
-			auto CalcHeight = [&](float RealTimeSeconds)
+			for (auto& It : ServerToClientCache.Keyframes)
 			{
-				const float Ratio = MinHeight / MinHeightSeconds;
-				return RealTimeSeconds * Ratio;
-			};
-
-			for (auto It = DebugBuffer->CreateIterator(); It; ++It)
-			{
-				auto* DebugState = It.Element();
-
-				
-				float ScreenWidth = MinWidth;
-				float FramePct = FMath::Clamp<float>((DebugState->LocalDeltaTimeSeconds - LocalFrameTimeRed) / ( LocalFrameTimeGreen - LocalFrameTimeRed ), 0.f, 1.f);
-				FColor Color = FColor::MakeRedToGreenColorFromScalar( FramePct );
-				float ServerHeight = CalcHeight(DebugState->LocalDeltaTimeSeconds);
-
-				Out.EmitQuad(FVector2D( ScreenX, ScreenY), FVector2D( ScreenWidth, ServerHeight), Color);
-				Out.EmitText(FVector2D( ScreenX, ScreenY), FColor::Black, FString::Printf(TEXT("%.2f"), (DebugState->LocalDeltaTimeSeconds * 1000.f)));
-
-				TSimTime ClientSimTime;
-				float ClientX = ScreenX;
-				float ClientY = ScreenY - ClientOffsetY;
-				for (int32 Keyframe : DebugState->ProcessedKeyframes)
+				const int32 KeyFrame = It.Key;
+				auto& Positions = It.Value;
+				if (KeyFrame != 0 && Positions.RecvPosition != FVector2D::ZeroVector && Positions.SentPosition != FVector2D::ZeroVector)
 				{
-					auto* Cmd = InputBuffer.FindElementByKeyframe(Keyframe);
-					if (Cmd)
-					{
-						ClientSimTime += Cmd->GetFrameDeltaTime();
-						
-						float ClientSizeX = MinWidth;
-						float ClientSizeY =  CalcHeight(Cmd->GetFrameDeltaTime().ToRealTimeSeconds());
+					Out.EmitLine(Positions.SentPosition, Positions.RecvPosition, FColor::Purple);
 
-						FVector2D ScreenPos(ClientX, ClientY - ClientSizeY);
-						Out.EmitQuad(ScreenPos, FVector2D( ClientSizeX, ClientSizeY), FColor::Blue);
-						Out.EmitText(ScreenPos, FColor::White, LexToString(Keyframe));
-						ClientY -= (ClientSizeY + ClientOffsetY);
-					}
+					FVector2D TextPos = Positions.SentPosition + (0.25f * (Positions.RecvPosition - Positions.SentPosition));
+					Out.EmitText(TextPos, FColor::Purple, LexToString(KeyFrame));
 				}
+			}
 
-				for (int32 Keyframe = DebugState->LastProcessedKeyframe+1; Keyframe <= DebugState->HeadKeyframe; ++Keyframe)
+			for (auto& It : ClientToServerCache.Keyframes)
+			{
+				const int32 KeyFrame = It.Key;
+				auto& Positions = It.Value;
+				if (KeyFrame != 0 && Positions.RecvPosition != FVector2D::ZeroVector && Positions.SentPosition != FVector2D::ZeroVector)
 				{
-					if (auto* Cmd = InputBuffer.FindElementByKeyframe(Keyframe))
-					{
-						float ClientSizeX = MinWidth;
-						float ClientSizeY = CalcHeight(Cmd->GetFrameDeltaTime().ToRealTimeSeconds());
-						FVector2D ScreenPos(ClientX, ClientY - ClientSizeY);
-						Out.EmitQuad(ScreenPos, FVector2D( ClientSizeX, ClientSizeY), FColor::Red);
-						Out.EmitText(ScreenPos, FColor::White, LexToString(Keyframe));
-						ClientY -= (ClientSizeY + ClientOffsetY);
-					}
-				}
+					Out.EmitLine(Positions.SentPosition, Positions.RecvPosition, FColor::Orange);
 
-					
-				// -------------------------------------------
-				ScreenX += ScreenWidth + 2.f;
+					FVector2D TextPos = Positions.SentPosition + (0.25f * (Positions.RecvPosition - Positions.SentPosition));
+					Out.EmitText(TextPos, FColor::Orange, LexToString(KeyFrame));
+				}
 			}
 		}
 	}
diff --git a/Engine/Plugins/Runtime/NetworkPrediction/Source/NetworkPrediction/Public/NetworkSimulationModelReplicators.h b/Engine/Plugins/Runtime/NetworkPrediction/Source/NetworkPrediction/Public/NetworkSimulationModelReplicators.h
--- a/Engine/Plugins/Runtime/NetworkPrediction/Source/NetworkPrediction/Public/NetworkSimulationModelReplicators.h
+++ b/Engine/Plugins/Runtime/NetworkPrediction/Source/NetworkPrediction/Public/NetworkSimulationModelReplicators.h
@@ -97,7 +97,7 @@ struct TReplicatorBase
 		TickInfo.GiveSimulationTime(TickParameters.LocalDeltaTimeSeconds);
 
 		// See if we have sim time to spend (variable tick will always go through. fixed step will miss frames while accumulating)
-		TNetworkSimTime<TTickSettings> DeltaSimTime = TickInfo.GetRemaningAllowedSimulationTime();
+		TNetworkSimTime<TTickSettings> DeltaSimTime = TickInfo.GetRemainingAllowedSimulationTime();
 		if (DeltaSimTime.ToRealTimeSeconds() > 0)
 		{
 			if (TickParameters.bGenerateLocalInputCmds)
@@ -194,7 +194,15 @@ struct TReplicator_Sequence : public TBase
 			auto* Cmd = Ar.IsLoading() ? Buffer.GetWriteNext() : Buffer.FindElementByKeyframe(Keyframe);
 			Cmd->NetSerialize(P);
 		}
+
+		LastSerializedKeyframe = HeadKeyframe;
 	}
+
+	int32 GetLastSerializedKeyframe() const { return LastSerializedKeyframe; }
+
+protected:
+
+	int32 LastSerializedKeyframe = 0;
 };
 
 // Replicates only the latest single element in the selected buffer.
@@ -461,14 +469,22 @@ struct TReplicator_Autonomous : public TBase
 		// Resimulate
 		// -------------------------------------------------------------------------------------------------------------------------
 
-		TSyncState* ClientSyncState = Buffers.Sync.FindElementByKeyframe( ReconciliationKeyframe );	
+		TSyncState* ClientSyncState = Buffers.Sync.FindElementByKeyframe( ReconciliationKeyframe );
 
-		ServerState->VisualLog( FVisualLoggingParameters(EVisualLoggingContext::LastConfirmed, ReconciliationKeyframe, EVisualLoggingLifetime::Persistent), Driver, Driver);
+		const bool bDoVisualLog = NetworkSimulationModelCVars::EnableLocalPrediction() > 0; // don't visual log if we have prediction disabled
+		
+		if (bDoVisualLog)
+		{
+			ServerState->VisualLog( FVisualLoggingParameters(EVisualLoggingContext::LastConfirmed, ReconciliationKeyframe, EVisualLoggingLifetime::Persistent), Driver, Driver);
+		}
 
 		if (ClientSyncState)
 		{
 			// Existing ClientSyncState, log it before overwriting it
-			ClientSyncState->VisualLog( FVisualLoggingParameters(EVisualLoggingContext::FirstMispredicted, ReconciliationKeyframe, EVisualLoggingLifetime::Persistent), Driver, Driver);
+			if (bDoVisualLog)
+			{
+				ClientSyncState->VisualLog( FVisualLoggingParameters(EVisualLoggingContext::FirstMispredicted, ReconciliationKeyframe, EVisualLoggingLifetime::Persistent), Driver, Driver);
+			}
 		}
 		else
 		{
@@ -555,4 +571,58 @@ private:
 	int32 LastSerializedKeyframe = -1;
 	bool bPendingReconciliation = false;	// Reconciliation is pending: we need to reconcile state from the server that differs from the locally predicted state
 	bool bReconcileFaultDetected = false;	// A fault was detected: we received state from the server that we are unable to reconcile with locally predicted state
+};
+
+/** Special replicator for debug buffer, this preserves the local buffer and receives into a replicator-owned buffer (we want these buffers to be distinct/not merged) */
+template<typename TBufferTypes, typename TTickSettings, int32 MaxNumElements=5, typename TBase=TReplicator_SimTime<TBufferTypes, TTickSettings>>
+struct TReplicator_Debug : public TBase
+{
+	using TDebugState = typename TBufferTypes::TDebugState;
+
+	int32 GetProxyDirtyCount(TNetworkSimBufferContainer<TBufferTypes>& Buffers) const 
+	{
+		return Buffers.Debug.GetDirtyCount() ^ (TBase::GetProxyDirtyCount(Buffers) << 2); 
+	}
+
+	void NetSerialize(const FNetSerializeParams& P, TNetworkSimBufferContainer<TBufferTypes>& Buffers, const TSimulationTickState<TTickSettings>& TickInfo)
+	{
+		TBase::NetSerialize(P, Buffers, TickInfo);
+		FArchive& Ar = P.Ar;
+
+		TReplicationBuffer<TDebugState>& Buffer = Ar.IsSaving() ? Buffers.Debug : ReceivedBuffer;
+		
+		uint8 SerializedNumElements = FMath::Min<uint8>(MaxNumElements, Buffer.GetNumValidElements());
+		Ar << SerializedNumElements;
+
+		const int32 HeadKeyframe = FNetworkSimulationSerialization::SerializeKeyframe(Ar, Buffer.GetHeadKeyframe());
+		const int32 StartingKeyframe = FMath::Max(0, HeadKeyframe - SerializedNumElements + 1);
+
+		if (Ar.IsLoading())
+		{
+			// Lazy init on recieve
+			if (ReceivedBuffer.GetMaxNumElements() != Buffers.Debug.GetMaxNumElements())
+			{
+				ReceivedBuffer.SetBufferSize(Buffers.Debug.GetMaxNumElements());
+			}
+
+			const int32 PrevHead = Buffer.GetHeadKeyframe();
+			if (PrevHead < StartingKeyframe && PrevHead >= 0)
+			{
+				// There is a gap in the stream. In some cases, we want this to be a "fault" and bubble up. We may want to synthesize state or maybe we just skip ahead.
+				UE_LOG(LogNetworkSim, Warning, TEXT("Fault: gap in received buffer. PrevHead: %d. Received: %d-%d. Reseting previous buffer contents"), PrevHead, StartingKeyframe, HeadKeyframe);
+			}
+
+			Buffer.ResetNextHeadKeyframe(StartingKeyframe);
+		}
+
+		for (int32 Keyframe = StartingKeyframe; Keyframe <= HeadKeyframe; ++Keyframe)
+		{
+			// This, as is, is bad. The intention is that these functions serialize multiple items in some delta compressed fashion.
+			// As is, we are just serializing the elements individually.
+			auto* Cmd = Ar.IsLoading() ? Buffer.GetWriteNext() : Buffer.FindElementByKeyframe(Keyframe);
+			Cmd->NetSerialize(P);
+		}
+	}
+
+	TReplicationBuffer<TDebugState> ReceivedBuffer;
 };
\ No newline at end of file
diff --git a/Engine/Plugins/Runtime/NetworkPrediction/Source/NetworkPrediction/Public/NetworkSimulationModelTypes.h b/Engine/Plugins/Runtime/NetworkPrediction/Source/NetworkPrediction/Public/NetworkSimulationModelTypes.h
--- a/Engine/Plugins/Runtime/NetworkPrediction/Source/NetworkPrediction/Public/NetworkSimulationModelTypes.h
+++ b/Engine/Plugins/Runtime/NetworkPrediction/Source/NetworkPrediction/Public/NetworkSimulationModelTypes.h
@@ -286,7 +286,7 @@ struct TSimulationTickState
 	}
 
 	// How much granted simulation time is left to process
-	TNetworkSimTime<TSettings> GetRemaningAllowedSimulationTime() const
+	TNetworkSimTime<TSettings> GetRemainingAllowedSimulationTime() const
 	{
 		return TotalAllowedSimulationTime - TotalProcessedSimulationTime;
 	}
