NetworkPrediction plugin
-Fix several issues around simulation advancement and latency + local prediction being toggled on/off. Local prediction being disabled will now use an indendent time tracker to manage when new input cmds are sampled. Server more resiliant to latency and packet drops, avoids getting "stuck" not processing input cmds.
-Minor tweaks to the canvas debugger

#rnx

diff --git a/Engine/Plugins/Runtime/NetworkPrediction/Source/NetworkPrediction/Private/Movement/BaseMovementComponent.cpp b/Engine/Plugins/Runtime/NetworkPrediction/Source/NetworkPrediction/Private/Movement/BaseMovementComponent.cpp
--- a/Engine/Plugins/Runtime/NetworkPrediction/Source/NetworkPrediction/Private/Movement/BaseMovementComponent.cpp
+++ b/Engine/Plugins/Runtime/NetworkPrediction/Source/NetworkPrediction/Private/Movement/BaseMovementComponent.cpp
@@ -385,7 +385,7 @@ void UBaseMovementComponent::DrawDebug(const IBaseMovementDriver::FDrawDebugPara
 			if (BaseMovementCVars::UseDrawDebug)
 			{
 				static const float Thickness = 2.f;
-				DrawDebugCapsule(Params.DebugWorld, Params.Transform.GetLocation(), HalfHeight, Radius, Params.Transform.GetRotation(), Params.DrawColor, PersistentLines, LifetimeSeconds, 0.f, Thickness);
+				DrawDebugCapsule(Params.DebugWorld, Params.Transform.GetLocation(), HalfHeight, Radius, Params.Transform.GetRotation(), Params.DrawColor, false, LifetimeSeconds, 0.f, Thickness);
 			}
 
 			if (BaseMovementCVars::UseVLogger)
diff --git a/Engine/Plugins/Runtime/NetworkPrediction/Source/NetworkPrediction/Public/NetworkSimulationModelDebugger.h b/Engine/Plugins/Runtime/NetworkPrediction/Source/NetworkPrediction/Public/NetworkSimulationModelDebugger.h
--- a/Engine/Plugins/Runtime/NetworkPrediction/Source/NetworkPrediction/Public/NetworkSimulationModelDebugger.h
+++ b/Engine/Plugins/Runtime/NetworkPrediction/Source/NetworkPrediction/Public/NetworkSimulationModelDebugger.h
@@ -22,6 +22,7 @@ DEFINE_LOG_CATEGORY_STATIC(LogNetworkSimDebug, Log, All);
 namespace NetworkSimulationModelDebugCVars
 {
 	NETSIM_DEVCVAR_SHIPCONST_INT(DrawKeyframes, 1, "nsm.debug.DrawKeyFrames", "Draws keyframe data (text) in debug graphs");
+	NETSIM_DEVCVAR_SHIPCONST_INT(DrawNetworkSendLines, 1, "nsm.debug.DrawNetworkSendLines", "Draws lines representing network traffic in debugger");
 	NETSIM_DEVCVAR_SHIPCONST_INT(GatherServerSidePIE, 1, "nsm.debug.GatherServerSide", "Whenever we gather debug info from a client side actor, also gather server side equivelent. Only works in PIE.");
 }
 
@@ -248,10 +249,11 @@ struct NETWORKPREDICTION_API FNetworkSimulationModelDebuggerManager: public FTic
 		CanvasItems[1].Emplace( MakeUnique<FCanvasTextItem>(ScreenPosition, FText::FromString(Str), GEngine->GetTinyFont(), Color) );
 	}
 
-	void EmitLine(FVector2D StartPosition, FVector2D EndPosition, FColor Color)
+	void EmitLine(FVector2D StartPosition, FVector2D EndPosition, FColor Color, float Thickness=1.f)
 	{
 		CanvasItems[0].Emplace( MakeUnique<FCanvasLineItem>(StartPosition, EndPosition) );
 		CanvasItems[0].Last()->SetColor(Color);
+		((FCanvasLineItem*)CanvasItems[0].Last().Get())->LineThickness = Thickness;
 	}
 
 private:
@@ -462,10 +464,22 @@ struct TNetworkSimulationModelDebugger : public INetworkSimulationModelDebugger
 				
 				LinePos.X += FixedWidth + Pad;
 				LinePos.Y = BaseLineYPos - (DebugState->RemainingAllowedSimulationTimeSeconds * SecondsToPixelsY);
+
+				FColor LineColor =  FColor::White;
+				if (LinePos.Y < DrawRect.Min.Y)
+				{
+					LinePos.Y = DrawRect.Min.Y;
+					LineColor = FColor::Red;
+				}
+				if (LinePos.Y > DrawRect.Max.Y)
+				{
+					LinePos.Y = DrawRect.Max.Y;
+					LineColor = FColor::Red;
+				}
 				
 				if (LastLinePos.IsSet())
 				{
-					Out.EmitLine(LastLinePos.GetValue(), LinePos, FColor::White);
+					Out.EmitLine(LastLinePos.GetValue(), LinePos, LineColor, 2.f);
 				}
 
 				LastLinePos = LinePos;
@@ -596,30 +610,32 @@ struct TNetworkSimulationModelDebugger : public INetworkSimulationModelDebugger
 			GatherDebugGraph(Out, Canvas, NetworkSim->GetLocalDebugBuffer(), ClientRect, MaxColumnTime, MaxLocalFrameTime, TEXT("Client"), ClientToServerCache, ServerToClientCache);
 
 			// Network Send/Recv lines
-
-			for (auto& It : ServerToClientCache.Keyframes)
+			if (NetworkSimulationModelDebugCVars::DrawNetworkSendLines() > 0)
 			{
-				const int32 KeyFrame = It.Key;
-				auto& Positions = It.Value;
-				if (KeyFrame != 0 && Positions.RecvPosition != FVector2D::ZeroVector && Positions.SentPosition != FVector2D::ZeroVector)
+				for (auto& It : ServerToClientCache.Keyframes)
 				{
-					Out.EmitLine(Positions.SentPosition, Positions.RecvPosition, FColor::Purple);
+					const int32 KeyFrame = It.Key;
+					auto& Positions = It.Value;
+					if (KeyFrame != 0 && Positions.RecvPosition != FVector2D::ZeroVector && Positions.SentPosition != FVector2D::ZeroVector)
+					{
+						Out.EmitLine(Positions.SentPosition, Positions.RecvPosition, FColor::Purple);
 
-					FVector2D TextPos = Positions.SentPosition + (0.25f * (Positions.RecvPosition - Positions.SentPosition));
-					Out.EmitText(TextPos, FColor::Purple, LexToString(KeyFrame));
+						FVector2D TextPos = Positions.SentPosition + (0.25f * (Positions.RecvPosition - Positions.SentPosition));
+						Out.EmitText(TextPos, FColor::Purple, LexToString(KeyFrame));
+					}
 				}
-			}
 
-			for (auto& It : ClientToServerCache.Keyframes)
-			{
-				const int32 KeyFrame = It.Key;
-				auto& Positions = It.Value;
-				if (KeyFrame != 0 && Positions.RecvPosition != FVector2D::ZeroVector && Positions.SentPosition != FVector2D::ZeroVector)
+				for (auto& It : ClientToServerCache.Keyframes)
 				{
-					Out.EmitLine(Positions.SentPosition, Positions.RecvPosition, FColor::Orange);
+					const int32 KeyFrame = It.Key;
+					auto& Positions = It.Value;
+					if (KeyFrame != 0 && Positions.RecvPosition != FVector2D::ZeroVector && Positions.SentPosition != FVector2D::ZeroVector)
+					{
+						Out.EmitLine(Positions.SentPosition, Positions.RecvPosition, FColor::Orange);
 
-					FVector2D TextPos = Positions.SentPosition + (0.25f * (Positions.RecvPosition - Positions.SentPosition));
-					Out.EmitText(TextPos, FColor::Orange, LexToString(KeyFrame));
+						FVector2D TextPos = Positions.SentPosition + (0.25f * (Positions.RecvPosition - Positions.SentPosition));
+						Out.EmitText(TextPos, FColor::Orange, LexToString(KeyFrame));
+					}
 				}
 			}
 		}
diff --git a/Engine/Plugins/Runtime/NetworkPrediction/Source/NetworkPrediction/Public/NetworkSimulationModelReplicators.h b/Engine/Plugins/Runtime/NetworkPrediction/Source/NetworkPrediction/Public/NetworkSimulationModelReplicators.h
--- a/Engine/Plugins/Runtime/NetworkPrediction/Source/NetworkPrediction/Public/NetworkSimulationModelReplicators.h
+++ b/Engine/Plugins/Runtime/NetworkPrediction/Source/NetworkPrediction/Public/NetworkSimulationModelReplicators.h
@@ -95,34 +95,6 @@ struct TReplicatorBase
 	{
 		// Accumulate local delta time into TickInfo
 		TickInfo.GiveSimulationTime(TickParameters.LocalDeltaTimeSeconds);
-
-		// See if we have sim time to spend (variable tick will always go through. fixed step will miss frames while accumulating)
-		TNetworkSimTime<TTickSettings> DeltaSimTime = TickInfo.GetRemainingAllowedSimulationTime();
-		if (DeltaSimTime.ToRealTimeSeconds() > 0)
-		{
-			if (TickParameters.bGenerateLocalInputCmds)
-			{
-				if (TInputCmd* InputCmd = Buffers.Input.GetWriteNext())
-				{
-					*InputCmd = TInputCmd();
-					InputCmd->SetFrameDeltaTime(DeltaSimTime);
-					Driver->ProduceInput(DeltaSimTime, *InputCmd);
-					TickInfo.MaxAllowedInputKeyframe++;
-				}
-			}
-			else
-			{
-				// Allowed to process all remaining commands (we will still check for frame time in the core input processing loop to prevent speed hacks)
-				TickInfo.MaxAllowedInputKeyframe = Buffers.Input.GetHeadKeyframe();
-
-				// Check for gaps in commands
-				if ( TickInfo.LastProcessedInputKeyframe+1 < Buffers.Input.GetTailKeyframe() )
-				{
-					UE_LOG(LogNetworkSim, Warning, TEXT("::Tick missing inputcmds. LastProcessedInputKeyframe: %d. %s"), TickInfo.LastProcessedInputKeyframe, *Buffers.Input.GetBasicDebugStr());
-					TickInfo.LastProcessedInputKeyframe = Buffers.Input.GetTailKeyframe()+1;
-				}
-			}
-		}
 	}
 };
 
@@ -181,7 +153,7 @@ struct TReplicator_Sequence : public TBase
 			if (PrevHead < StartingKeyframe && PrevHead >= 0)
 			{
 				// There is a gap in the stream. In some cases, we want this to be a "fault" and bubble up. We may want to synthesize state or maybe we just skip ahead.
-				UE_LOG(LogNetworkSim, Warning, TEXT("Fault: gap in received buffer. PrevHead: %d. Received: %d-%d. Reseting previous buffer contents"), PrevHead, StartingKeyframe, HeadKeyframe);
+				UE_LOG(LogNetworkSim, Warning, TEXT("Fault: gap in received %s buffer. PrevHead: %d. Received: %d-%d. Reseting previous buffer contents"), *LexToString(BufferId), PrevHead, StartingKeyframe, HeadKeyframe);
 			}
 
 			Buffer.ResetNextHeadKeyframe(StartingKeyframe);
@@ -242,17 +214,54 @@ struct TReplicator_Single : public TBase
 };
 
 // -------------------------------------------------------------------------------------------------------
-//	Role based Replicators: these replicators are meant to server specific roles
+//	Role based Replicators: these replicators are meant to serve specific roles wrt how the simulation evolves
 // -------------------------------------------------------------------------------------------------------
 
-// Default Replicator for the Server
-//	-Replicates the InputBuffer client->server
-template<typename TBufferTypes, typename TTickSettings, ENetworkSimBufferTypeId BufferId=ENetworkSimBufferTypeId::Input, typename TBase=TReplicator_Sequence<TBufferTypes, TTickSettings, BufferId, 3>>
+// Default Replicator for the Server: Replicates the InputBuffer client->server
+template<typename TBufferTypes, typename TTickSettings, typename TBase=TReplicator_Sequence<TBufferTypes, TTickSettings, ENetworkSimBufferTypeId::Input, 3>>
 struct TReplicator_Server : public TBase
 {
+	using TInputCmd = typename TBufferTypes::TInputCmd;
+
+	template<typename T, typename TDriver>
+	void Reconcile(TDriver* Driver, TNetworkSimBufferContainer<TBufferTypes>& Buffers, TSimulationTickState<TTickSettings>& TickInfo)
+	{
+		// After receiving input, server may process up to the latest received frames.
+		// (If we needed to buffer input server side for whatever reason, we would do it here)
+		// (also note that we will implicitly guard against speed hacks int he core update loop by not processing cmds past what we have been "allowed")
+		TickInfo.MaxAllowedInputKeyframe = Buffers.Input.GetHeadKeyframe();
+
+		// Check for gaps in commands
+		if ( TickInfo.LastProcessedInputKeyframe+1 < Buffers.Input.GetTailKeyframe() )
+		{
+			UE_LOG(LogNetworkSim, Warning, TEXT("TReplicator_Server::Reconcile missing inputcmds. LastProcessedInputKeyframe: %d. %s"), TickInfo.LastProcessedInputKeyframe, *Buffers.Input.GetBasicDebugStr());
+			TickInfo.LastProcessedInputKeyframe = Buffers.Input.GetTailKeyframe()+1;
+		}
+	}
+
+	template<typename T, typename TDriver, typename TTickParameters>
+	void PreSimTick(TDriver* Driver, TNetworkSimBufferContainer<TBufferTypes>& Buffers, TSimulationTickState<TTickSettings>& TickInfo, const TTickParameters& TickParameters)
+	{
+		TickInfo.GiveSimulationTime(TickParameters.LocalDeltaTimeSeconds);
 
+		if (TickParameters.bGenerateLocalInputCmds)
+		{
+			TNetworkSimTime<TTickSettings> DeltaSimTime = TickInfo.GetRemainingAllowedSimulationTime();
+			if (DeltaSimTime.IsPositive())
+			{
+				if (TInputCmd* InputCmd = Buffers.Input.GetWriteNext())
+				{
+					*InputCmd = TInputCmd();
+					InputCmd->SetFrameDeltaTime(DeltaSimTime);
+					Driver->ProduceInput(DeltaSimTime, *InputCmd);
+					TickInfo.MaxAllowedInputKeyframe = Buffers.Input.GetHeadKeyframe();
+				}
+			}
+		}
+	}
 };
 
+/** Simulated: "non locally controlled" simulations. We support "Simulation Extrapolation" here (using the sim to fake inputs to advance the sim)  */
 template<typename TBufferTypes, typename TTickSettings, typename TBase=TReplicator_Single<TBufferTypes, TTickSettings, ENetworkSimBufferTypeId::Sync>>
 struct TReplicator_Simulated : public TBase
 {
@@ -339,9 +348,26 @@ struct TReplicator_Simulated : public TBase
 	template<typename T, typename TDriver, typename TTickParameters>
 	void PreSimTick(TDriver* Driver, TNetworkSimBufferContainer<TBufferTypes>& Buffers, TSimulationTickState<TTickSettings>& TickInfo, const TTickParameters& TickParameters)
 	{
-		if (NetworkSimulationModelCVars::EnableSimulatedExtrapolation())
+		if (NetworkSimulationModelCVars::EnableSimulatedExtrapolation() == 0)
+		{
+			return;
+		}
+		
+		TickInfo.GiveSimulationTime(TickParameters.LocalDeltaTimeSeconds);
+
+		if (TickParameters.bGenerateLocalInputCmds)
 		{
-			TBase::template PreSimTick<T, TDriver, TTickParameters>(Driver, Buffers, TickInfo, TickParameters);
+			TNetworkSimTime<TTickSettings> DeltaSimTime = TickInfo.GetRemainingAllowedSimulationTime();
+			if (DeltaSimTime.IsPositive())
+			{
+				if (TInputCmd* InputCmd = Buffers.Input.GetWriteNext())
+				{
+					*InputCmd = TInputCmd();
+					InputCmd->SetFrameDeltaTime(DeltaSimTime);
+					Driver->ProduceInput(DeltaSimTime, *InputCmd);
+					TickInfo.MaxAllowedInputKeyframe = Buffers.Input.GetHeadKeyframe();
+				}
+			}
 		}
 	}
 
@@ -389,12 +415,6 @@ struct TReplicator_Autonomous : public TBase
 		}
 		else
 		{
-			// Its possible we process two packets in between ticks. This isn't really a problem but should be rare.
-			if (bPendingReconciliation)
-			{
-				UE_LOG(LogNetworkSim, Warning, TEXT("bPendingReconciliation while in ::NetSerialize. LastSerializedKeyframe: %d. New SerializedHeadKeyframe: %d."), LastSerializedKeyframe, SerializedHeadKeyframe);
-			}
-
 			// Lazy init the reconciliation buffer. We don't need this on the server/writing side. (Fixme, with templated approach we could easily specialize this to an inline element instead of buffer?)
 			if (ReconciliationBuffer.GetMaxNumElements() == 0)
 			{
@@ -501,6 +521,12 @@ struct TReplicator_Autonomous : public TBase
 		TickInfo.TotalProcessedSimulationTime = SerializedTime;
 		TickInfo.LastProcessedInputKeyframe = ReconciliationKeyframe;
 		TickInfo.MaxAllowedInputKeyframe = FMath::Max(TickInfo.MaxAllowedInputKeyframe, TickInfo.LastProcessedInputKeyframe); // Make sure this doesn't lag behind. This is the only place we should need to do this.
+
+		if (NetworkSimulationModelCVars::EnableLocalPrediction() == 0)
+		{
+			// If we aren't predicting at all, then we advanced the allowed sim time here, (since we aren't doing it in PreSimTick). This just keeps us constantly falling behind and not being able to toggle prediction on/off for debugging.
+			TickInfo.TotalAllowedSimulationTime = SerializedTime;
+		}
 		
 		// Resimulate all user commands 
 		const int32 LastKeyframeToProcess = TickInfo.MaxAllowedInputKeyframe;
@@ -513,7 +539,12 @@ struct TReplicator_Autonomous : public TBase
 			
 			check(ResimulateCmd);
 			check(PrevMotionState);
-			check(NextMotionState);
+			if (NextMotionState == nullptr)
+			{
+				// This should only happen if we transition from no local prediction to local prediction, since we will traverse through out input buffer but not have predicted sync states to look at.
+				NextMotionState = Buffers.Sync.GetWriteNext();
+				check(Keyframe == Buffers.Sync.GetHeadKeyframe());
+			}
 
 			// TEMP (Aux buffer not fully plumbed through the system yet)
 			// ------------------------------------------------------
@@ -552,21 +583,52 @@ struct TReplicator_Autonomous : public TBase
 			return;
 		}
 
-		const int32 StartingMaxAllowedInputKeyframe = TickInfo.MaxAllowedInputKeyframe;
-
-		TBase::template PreSimTick<T, TDriver, TTickParameters>(Driver, Buffers, TickInfo, TickParameters);
-
-		// If prediction is disabled we need to reset our MaxAllowedInputKeyframe so we don't process the command(s) generated above
-		if (NetworkSimulationModelCVars::EnableLocalPrediction() == 0)
+		if (TickParameters.bGenerateLocalInputCmds)
 		{
-			TickInfo.MaxAllowedInputKeyframe = StartingMaxAllowedInputKeyframe;
+			if (NetworkSimulationModelCVars::EnableLocalPrediction() > 0)
+			{
+				// Prediction: add simulation time and generate new commands
+				TickInfo.GiveSimulationTime(TickParameters.LocalDeltaTimeSeconds);
+				const TNetworkSimTime<TTickSettings> DeltaSimTime = TickInfo.GetRemainingAllowedSimulationTime();
+				if (DeltaSimTime.IsPositive())
+				{
+					if (TInputCmd* InputCmd = Buffers.Input.GetWriteNext())
+					{
+						*InputCmd = TInputCmd();
+						InputCmd->SetFrameDeltaTime(DeltaSimTime);
+						Driver->ProduceInput(DeltaSimTime, *InputCmd);
+						TickInfo.MaxAllowedInputKeyframe = Buffers.Input.GetHeadKeyframe(); // Allow the new command to be processed by the local simulation
+					}
+				}
+			}
+			else
+			{
+				// Local prediction disabled: we must use a separate time accumulator to figure out when we should add more input cmds.
+				// Since we aren't processing the simulation locally, our core simulation time will only advance from network updates.
+				// (still, we need *something* to tell us when to generate a new command and what the delta time should be)
+
+				TNetworkSimTime<TTickSettings> NonPredictedInputTime;
+				NonPredictedInputTimeAccumulator.Accumulate(NonPredictedInputTime, TickParameters.LocalDeltaTimeSeconds);
+				if (NonPredictedInputTime.IsPositive())
+				{
+					if (TInputCmd* InputCmd = Buffers.Input.GetWriteNext())
+					{
+						*InputCmd = TInputCmd();
+						InputCmd->SetFrameDeltaTime(NonPredictedInputTime);
+						Driver->ProduceInput(NonPredictedInputTime, *InputCmd);
+					}
+				}
+			}
 		}
 	}
 
 private:
+
+
 	
 	TReplicationBuffer<TSyncState> ReconciliationBuffer;
 	TNetworkSimTime<TTickSettings> SerializedTime; // last serialized time keeper
+	TRealTimeAccumulator<TTickSettings> NonPredictedInputTimeAccumulator; // for tracking input time in the non predictive case
 
 	int32 LastSerializedKeyframe = -1;
 	bool bPendingReconciliation = false;	// Reconciliation is pending: we need to reconcile state from the server that differs from the locally predicted state
@@ -609,7 +671,7 @@ struct TReplicator_Debug : public TBase
 			if (PrevHead < StartingKeyframe && PrevHead >= 0)
 			{
 				// There is a gap in the stream. In some cases, we want this to be a "fault" and bubble up. We may want to synthesize state or maybe we just skip ahead.
-				UE_LOG(LogNetworkSim, Warning, TEXT("Fault: gap in received buffer. PrevHead: %d. Received: %d-%d. Reseting previous buffer contents"), PrevHead, StartingKeyframe, HeadKeyframe);
+				UE_LOG(LogNetworkSim, Warning, TEXT("Fault: gap in received Debug buffer. PrevHead: %d. Received: %d-%d. Reseting previous buffer contents"), PrevHead, StartingKeyframe, HeadKeyframe);
 			}
 
 			Buffer.ResetNextHeadKeyframe(StartingKeyframe);
diff --git a/Engine/Plugins/Runtime/NetworkPrediction/Source/NetworkPrediction/Public/NetworkSimulationModelTypes.h b/Engine/Plugins/Runtime/NetworkPrediction/Source/NetworkPrediction/Public/NetworkSimulationModelTypes.h
--- a/Engine/Plugins/Runtime/NetworkPrediction/Source/NetworkPrediction/Public/NetworkSimulationModelTypes.h
+++ b/Engine/Plugins/Runtime/NetworkPrediction/Source/NetworkPrediction/Public/NetworkSimulationModelTypes.h
@@ -17,6 +17,18 @@ enum ENetworkSimBufferTypeId
 	Debug
 };
 
+inline FString LexToString(ENetworkSimBufferTypeId A)
+{
+	switch(A)
+	{
+		case Input: return TEXT("Input");
+		case Sync: return TEXT("Sync");
+		case Aux: return TEXT("Aux");
+		case Debug: return TEXT("Debug");
+	};
+	return TEXT("Unknown");
+}
+
 // Helper needed to specialize TNetworkSimBufferTypes::Get (must be done outside of templated struct)
 template<typename TBufferTypes, ENetworkSimBufferTypeId BufferId> 
 struct TSelectTypeHelper
@@ -195,6 +207,10 @@ struct TNetworkSimTime
 	TRealTime ToRealTimeSeconds() const { return (Time * TickSettings::GetSimToRealFactor()); }
 	FString ToString() const { return LexToString(this->Time); }
 
+	bool IsPositive() const { return (Time > 0); }
+	bool IsNegative() const { return (Time < 0); }
+	void Reset() { Time = 0; }
+
 	// FIXME
 	void NetSerialize(FArchive& Ar) { Ar << Time; }
 
