Network Prediction WIP
-Interpolation code stood up
-Fixes to Simulated Extrapolation
-Debugger improvements

#rnx

diff --git a/Engine/Plugins/Runtime/NetworkPrediction/Source/NetworkPrediction/Private/Movement/BaseMovementComponent.cpp b/Engine/Plugins/Runtime/NetworkPrediction/Source/NetworkPrediction/Private/Movement/BaseMovementComponent.cpp
--- a/Engine/Plugins/Runtime/NetworkPrediction/Source/NetworkPrediction/Private/Movement/BaseMovementComponent.cpp
+++ b/Engine/Plugins/Runtime/NetworkPrediction/Source/NetworkPrediction/Private/Movement/BaseMovementComponent.cpp
@@ -363,7 +363,7 @@ void UBaseMovementComponent::DrawDebug(const IBaseMovementDriver::FDrawDebugPara
 		if (BaseMovementCVars::UseDrawDebug)
 		{
 			static float PointSize = 3.f;
-			DrawDebugPoint(Params.DebugWorld, Params.Transform.GetLocation(), PointSize, Params.DrawColor, PersistentLines, LifetimeSeconds);
+			DrawDebugPoint(Params.DebugWorld, Params.Transform.GetLocation(), PointSize, Params.DrawColor, false, LifetimeSeconds);
 		}
 
 		if (BaseMovementCVars::UseVLogger)
@@ -408,7 +408,7 @@ void UBaseMovementComponent::DrawDebug(const IBaseMovementDriver::FDrawDebugPara
 				FVector ActorExtent;
 				LocalSpaceBox.GetCenterAndExtents(ActorOrigin, ActorExtent);
 				ActorExtent *= Params.Transform.GetScale3D();
-				DrawDebugBox(Params.DebugWorld, Params.Transform.GetLocation(), ActorExtent, Params.Transform.GetRotation(), Params.DrawColor, PersistentLines, LifetimeSeconds, 0, Thickness);
+				DrawDebugBox(Params.DebugWorld, Params.Transform.GetLocation(), ActorExtent, Params.Transform.GetRotation(), Params.DrawColor, false, LifetimeSeconds, 0, Thickness);
 			}
 
 			if (BaseMovementCVars::UseVLogger)
diff --git a/Engine/Plugins/Runtime/NetworkPrediction/Source/NetworkPrediction/Private/Movement/ParametricMovement.cpp b/Engine/Plugins/Runtime/NetworkPrediction/Source/NetworkPrediction/Private/Movement/ParametricMovement.cpp
--- a/Engine/Plugins/Runtime/NetworkPrediction/Source/NetworkPrediction/Private/Movement/ParametricMovement.cpp
+++ b/Engine/Plugins/Runtime/NetworkPrediction/Source/NetworkPrediction/Private/Movement/ParametricMovement.cpp
@@ -2,6 +2,7 @@
 
 #include "Movement/ParametricMovement.h"
 #include "NetworkSimulationModelDebugger.h"
+#include "HAL/IConsoleManager.h"
 
 DEFINE_LOG_CATEGORY_STATIC(LogParametricMovement, Log, All);
 
@@ -19,10 +20,6 @@ static float DrawDebugDefaultLifeTime = 30.f;
 static FAutoConsoleVariableRef CVarDrawDebugDefaultLifeTime(TEXT("parametricmover.Debug.UseDrawDebug.DefaultLifeTime"),
 	DrawDebugDefaultLifeTime, TEXT("Use built in DrawDebug* functions for visual logging"), ECVF_Default);
 
-static int32 ForceNetUpdate = 0;
-static FAutoConsoleVariableRef CVarForceNetUpdate(TEXT("parametricmover.ForceNetUpdate"),
-	DrawDebugDefaultLifeTime, TEXT("Toggles calling ForceNetUpdate each tick on Parametric movers. This is only meant for testing/debugging."), ECVF_Default);
-
 }
 
 // -------------------------------------------------------------------------------------------------------------------------------
@@ -64,7 +61,7 @@ namespace ParametricMovement
 		}
 	}
 
-	void FMoveState::VisualLog(const FVisualLoggingParameters& Parameters, IParametricMovementDriver* Driver, IParametricMovementDriver* LogDriver)
+	void FMoveState::VisualLog(const FVisualLoggingParameters& Parameters, IParametricMovementDriver* Driver, IParametricMovementDriver* LogDriver) const
 	{
 		IBaseMovementDriver& BaseMovementDriver = Driver->GetBaseMovementDriver();
 		IBaseMovementDriver& LogMovementDriver = LogDriver->GetBaseMovementDriver();
@@ -103,6 +100,13 @@ IReplicationProxy* UParametricMovementComponent::InstantiateNetworkSimulation()
 #if NETSIM_MODEL_DEBUG
 	FNetworkSimulationModelDebuggerManager::Get().RegisterNetworkSimulationModel(NetworkSim.Get(), (IParametricMovementDriver*)this, GetOwner(), TEXT("ParametricMovement"));
 #endif
+
+	if (bEnableInterpolation)
+	{
+		NetworkInterpolator.Reset(new TNetworkSimulationModelInterpolator<ParametricMovement::FMovementSystem, IParametricMovementDriver>(NetworkSim.Get(), (IParametricMovementDriver*)this));
+		NetworkSim->RepProxy_Simulated.bAllowSimulatedExtrapolation = false;
+	}
+
 	return NetworkSim.Get();
 }
 
@@ -117,7 +121,7 @@ FNetworkSimulationModelInitParameters UParametricMovementComponent::GetSimulatio
 	// These are reasonable defaults but may not be right for everyone
 	FNetworkSimulationModelInitParameters InitParams;
 	InitParams.InputBufferSize = 32; //Role != ROLE_SimulatedProxy ? 32 : 32;  // Fixme.. not good
-	InitParams.SyncedBufferSize = Role != ROLE_AutonomousProxy ? 2 : 32;
+	InitParams.SyncedBufferSize = Role != ROLE_AutonomousProxy ? 4 : 32;
 	InitParams.AuxBufferSize = Role != ROLE_AutonomousProxy ? 2 : 32;
 	InitParams.DebugBufferSize = 32;
 	InitParams.HistoricBufferSize = 128;
@@ -194,7 +198,15 @@ void UParametricMovementComponent::TickComponent(float DeltaTime, enum ELevelTic
 {
 	Super::TickComponent(DeltaTime, TickType, ThisTickFunction);
 
-	TickSimulation(DeltaTime); // fixme
+	if (!bDisableParametricMovementSimulation)
+	{
+		TickSimulation(DeltaTime); // fixme
+
+		if (NetworkInterpolator.IsValid() && GetOwnerRole() == ROLE_SimulatedProxy)
+		{
+			NetworkInterpolator->Tick(DeltaTime);
+		}
+	}
 }
 
 void UParametricMovementComponent::Reconcile()
@@ -216,8 +228,12 @@ void UParametricMovementComponent::TickSimulation(float DeltaTimeSeconds)
 		NetworkSim->Tick((IParametricMovementDriver*)this, Parameters);
 	}
 
-	// TEMP
-	if (ParametricMoverCVars::ForceNetUpdate)
+	if (ParentNetUpdateFrequency > 0.f)
+	{
+		GetOwner()->NetUpdateFrequency = ParentNetUpdateFrequency;
+	}
+	
+	if (bEnableForceNetUpdate)
 	{
 		GetOwner()->ForceNetUpdate();
 	}
diff --git a/Engine/Plugins/Runtime/NetworkPrediction/Source/NetworkPrediction/Private/NetworkSimulationModelDebugger.cpp b/Engine/Plugins/Runtime/NetworkPrediction/Source/NetworkPrediction/Private/NetworkSimulationModelDebugger.cpp
--- a/Engine/Plugins/Runtime/NetworkPrediction/Source/NetworkPrediction/Private/NetworkSimulationModelDebugger.cpp
+++ b/Engine/Plugins/Runtime/NetworkPrediction/Source/NetworkPrediction/Private/NetworkSimulationModelDebugger.cpp
@@ -10,6 +10,7 @@
 #include "Engine/NetConnection.h"
 #include "Engine/NetDriver.h"
 #include "EngineUtils.h"
+#include "ProfilingDebugging/ScopedTimers.h"
 
 FNetworkSimulationModelDebuggerManager& FNetworkSimulationModelDebuggerManager::Get()
 {
@@ -60,10 +61,10 @@ UObject* FindReplicatedObjectOnPIEServer(UObject* ClientObject)
 }
 
 // ------------------------------------------------------------------------------------------------------------------------
-//
-//
+//	Debug functions for toggling debugger on specific actors
 // ------------------------------------------------------------------------------------------------------------------------
 
+// Debug the first locally controlled pawn
 FAutoConsoleCommandWithWorldAndArgs NetworkSimulationModelDebugCmd(TEXT("nms.Debug.LocallyControlledPawn"), TEXT(""),
 	FConsoleCommandWithWorldAndArgsDelegate::CreateLambda([](const TArray< FString >& InArgs, UWorld* World) 
 {
@@ -83,12 +84,34 @@ FAutoConsoleCommandWithWorldAndArgs NetworkSimulationModelDebugCmd(TEXT("nms.Deb
 	FNetworkSimulationModelDebuggerManager::Get().ToggleDebuggerActive(Pawn);
 }));
 
+// Toggles continous updates to debugger for locally controlled player
+FAutoConsoleCommandWithWorldAndArgs NetworkSimulationModelDebugToggleContinousCmd(TEXT("nms.Debug.ToggleContinous"), TEXT(""),
+	FConsoleCommandWithWorldAndArgsDelegate::CreateLambda([](const TArray< FString >& InArgs, UWorld* World) 
+{
+
+	if (!World || !World->GetFirstLocalPlayerFromController())
+	{
+		return;
+	}
+	ULocalPlayer* Player = World->GetFirstLocalPlayerFromController();
+	if (!Player || !Player->GetPlayerController(World) || !Player->GetPlayerController(World)->GetPawn())
+	{
+		UE_LOG(LogNetworkSimDebug, Error, TEXT("Could not find valid locally controlled pawn. "));
+		return;
+	}
+
+	APawn* Pawn = Player->GetPlayerController(World)->GetPawn();
+	FNetworkSimulationModelDebuggerManager::Get().SetDebuggerActive(Pawn, true);
+	FNetworkSimulationModelDebuggerManager::Get().ToggleContinousGather();
+}));
+
+// Debug actors by class filter. 2nd parameter can filter by name
 FAutoConsoleCommandWithWorldAndArgs NetworkSimulationModelDebugClassCmd(TEXT("nms.Debug.Class"), TEXT(""),
 	FConsoleCommandWithWorldAndArgsDelegate::CreateLambda([](const TArray< FString >& Args, UWorld* World) 
 {
 	if (Args.Num() <= 0)
 	{
-		UE_LOG(LogNetworkSimDebug, Display, TEXT("Usage: nms.Debug.Class <Class>"));
+		UE_LOG(LogNetworkSimDebug, Display, TEXT("Usage: nms.Debug.Class <Class> <Name>"));
 		return;
 	}
 
@@ -120,31 +143,11 @@ FAutoConsoleCommandWithWorldAndArgs NetworkSimulationModelDebugClassCmd(TEXT("nm
 	{
 		if (It->GetClass()->IsChildOf(Class))
 		{
-			UE_LOG(LogNetworkSimDebug, Display, TEXT("Toggling NetworkSim debugger for %s"), *It->GetName());
-			FNetworkSimulationModelDebuggerManager::Get().ToggleDebuggerActive(*It);
-			break;
+			if (Args.Num() < 2 || It->GetName().Contains(Args[1]))
+			{
+				UE_LOG(LogNetworkSimDebug, Display, TEXT("Toggling NetworkSim debugger for %s"), *It->GetName());
+				FNetworkSimulationModelDebuggerManager::Get().ToggleDebuggerActive(*It);
+			}
 		}
 	}
-	
-}));
-
-
-FAutoConsoleCommandWithWorldAndArgs NetworkSimulationModelDebugToggleContinousCmd(TEXT("nms.Debug.ToggleContinous"), TEXT(""),
-	FConsoleCommandWithWorldAndArgsDelegate::CreateLambda([](const TArray< FString >& InArgs, UWorld* World) 
-{
-
-	if (!World || !World->GetFirstLocalPlayerFromController())
-	{
-		return;
-	}
-	ULocalPlayer* Player = World->GetFirstLocalPlayerFromController();
-	if (!Player || !Player->GetPlayerController(World) || !Player->GetPlayerController(World)->GetPawn())
-	{
-		UE_LOG(LogNetworkSimDebug, Error, TEXT("Could not find valid locally controlled pawn. "));
-		return;
-	}
-
-	APawn* Pawn = Player->GetPlayerController(World)->GetPawn();
-	FNetworkSimulationModelDebuggerManager::Get().SetDebuggerActive(Pawn, true);
-	FNetworkSimulationModelDebuggerManager::Get().ToggleContinousGather();
 }));
\ No newline at end of file
diff --git a/Engine/Plugins/Runtime/NetworkPrediction/Source/NetworkPrediction/Public/Movement/ParametricMovement.h b/Engine/Plugins/Runtime/NetworkPrediction/Source/NetworkPrediction/Public/Movement/ParametricMovement.h
--- a/Engine/Plugins/Runtime/NetworkPrediction/Source/NetworkPrediction/Public/Movement/ParametricMovement.h
+++ b/Engine/Plugins/Runtime/NetworkPrediction/Source/NetworkPrediction/Public/Movement/ParametricMovement.h
@@ -4,6 +4,7 @@
 
 #include "NetworkSimulationModel.h"
 #include "BaseMovementComponent.h"
+#include "NetworkInterpolation.h"
 
 #include "ParametricMovement.generated.h"
 
@@ -71,7 +72,13 @@ namespace ParametricMovement
 			}
 		}
 
-		void VisualLog(const FVisualLoggingParameters& Parameters, IParametricMovementDriver* Driver, IParametricMovementDriver* LogDriver);
+		void VisualLog(const FVisualLoggingParameters& Parameters, IParametricMovementDriver* Driver, IParametricMovementDriver* LogDriver) const;
+
+		static void Interpolate(const FMoveState& From, const FMoveState& To, const float PCT, FMoveState& OutDest)
+		{
+			OutDest.Position = From.Position + ((To.Position - From.Position) * PCT);
+			OutDest.PlayRate = From.PlayRate + ((To.PlayRate - From.PlayRate) * PCT);
+		}
 	};
 
 	// Auxiliary state that is input into the simulation. Doesn't change during the simulation tick.
@@ -168,6 +175,9 @@ protected:
 
 	TUniquePtr<ParametricMovement::FMovementSystem> NetworkSim;	
 
+	// Temp, will be moved
+	TUniquePtr< TNetworkSimulationModelInterpolator<ParametricMovement::FMovementSystem, IParametricMovementDriver> > NetworkInterpolator;
+
 	// ------------------------------------------------------------------------
 	// Temp Parametric movement example
 	//	The essence of this movement simulation is to map some Time value to a transform. That is it.
@@ -175,6 +185,11 @@ protected:
 	//	What is below is just a simple C++ implementation to stand things up. Most likely we would 
 	//	do additional subclasses to vary the way this is implemented)
 	// ------------------------------------------------------------------------
+	
+	/** Disables starting the simulation. For development/testing ease of use */
+	UPROPERTY(EditAnywhere, BlueprintReadOnly, Category=ParametricMovement)
+	bool bDisableParametricMovementSimulation = false;
+	
 	UPROPERTY(EditAnywhere, BlueprintReadOnly, Category=ParametricMovement)
 	FVector ParametricDelta = FVector(0.f, 0.f, 500.f);
 
@@ -184,6 +199,17 @@ protected:
 	UPROPERTY(EditAnywhere, BlueprintReadOnly, Category=ParametricMovement)
 	float MaxTime = 1.f;
 
+	UPROPERTY(EditAnywhere, BlueprintReadOnly, Category=ParametricMovementNetworking)
+	bool bEnableInterpolation = true;
+
+	/** Calls ForceNetUpdate every frame. Has slightly different behavior than a very high NetUpdateFrequency */
+	UPROPERTY(EditAnywhere, BlueprintReadOnly, Category=ParametricMovementNetworking)
+	bool bEnableForceNetUpdate = false;
+
+	/** Sets NetUpdateFrequency on parent. This is editable on the component and really just meant for use during development/test maps */
+	UPROPERTY(EditAnywhere, BlueprintReadOnly, Category=ParametricMovementNetworking)
+	float ParentNetUpdateFrequency = 0.f;
+
 	FTransform CachedStartingTransform;
 
 	TOptional<float> PendingPlayRate = 1.f;
diff --git a/Engine/Plugins/Runtime/NetworkPrediction/Source/NetworkPrediction/Public/NetworkInterpolation.h b/Engine/Plugins/Runtime/NetworkPrediction/Source/NetworkPrediction/Public/NetworkInterpolation.h
--- a/Engine/Plugins/Runtime/NetworkPrediction/Source/NetworkPrediction/Public/NetworkInterpolation.h
+++ b/Engine/Plugins/Runtime/NetworkPrediction/Source/NetworkPrediction/Public/NetworkInterpolation.h
@@ -1,16 +1,208 @@
 // Copyright 1998-2019 Epic Games, Inc. All Rights Reserved.
 
 #pragma once
+#include "NetworkSimulationModelCVars.h"
 
-class UNetworkInterpolationManager
+namespace NetworkInterpolationDebugCVars
+{
+	NETSIM_DEVCVAR_SHIPCONST_INT(Disable, 0, "ni.Disable", "Disables Network Interpolation");
+}
+
+class INetworkInterpolator
 {
+public:
 
+	virtual ~INetworkInterpolator() { }
+	virtual void Tick(float DeltaSeconds) = 0;
 };
 
+template<typename TNetworkedSimulationModel, typename TDriver>
+class TNetworkSimulationModelInterpolator : public INetworkInterpolator
+{
+public:
+
+	using TSimTime = typename TNetworkedSimulationModel::TSimTime;
+	using TRealTime = typename TNetworkedSimulationModel::TRealTime;
+	using TSyncState = typename TNetworkedSimulationModel::TSyncState;
+
+	TNetworkSimulationModelInterpolator(TNetworkedSimulationModel* Source, TDriver* InDriver)
+	{
+		NetworkSim = Source;
+		Driver = InDriver;
+	}
+	virtual ~TNetworkSimulationModelInterpolator()
+	{
+
+	}
+
+	void Tick(float DeltaSeconds) override
+	{
+		if (NetworkInterpolationDebugCVars::Disable() > 0)
+		{			
+			if (const TSyncState* HeadState = NetworkSim->Buffers.Sync.GetElementFromHead(0))
+			{
+				Driver->FinalizeFrame(*HeadState);
+			}
+			return;
+		}
+
+		if (NetworkSim->TickInfo.SimulationTimeBuffer.GetNumValidElements() <= 1)
+		{
+			// Cant interpolate yet	
+			return;
+		}
+
+		InterpolationTime += DeltaSeconds;
+
+		// SimulatinoTimeBuffer holds the simulation time stamps of each SyncState in the SyncBuffer
+		// (Note that things are not going to be evenly spaced time wise)
+		auto& SimulationTimeBuffer = NetworkSim->TickInfo.SimulationTimeBuffer;		
+		const int32 HeadKeyframe = SimulationTimeBuffer.GetHeadKeyframe();
+
+		// Find which keyframes we should be interpolating between
+		int32 InterpolationKeyframe = INDEX_NONE;
+		TSimTime FromInterpolationTime;
+		TSimTime ToInterpolationTime;
+
+		for (int32 Keyframe = SimulationTimeBuffer.GetTailKeyframe(); Keyframe < HeadKeyframe; ++Keyframe)
+		{
+			TSimTime ElementSimTime = *SimulationTimeBuffer.FindElementByKeyframe(Keyframe);
+			if (InterpolationTime > ElementSimTime.ToRealTimeSeconds())
+			{
+				InterpolationKeyframe = Keyframe;
+				FromInterpolationTime = ElementSimTime;
+			}
+			else
+			{
+				break;
+			}
+		}
+
+		auto ResetInterpolationTimeToMidpoint = [&](int32 Keyframe)
+		{
+			check(SimulationTimeBuffer.IsValidKeyframe(Keyframe));
+			check(SimulationTimeBuffer.IsValidKeyframe(Keyframe+1));
+
+			InterpolationKeyframe = Keyframe;
+			
+			FromInterpolationTime = *SimulationTimeBuffer.FindElementByKeyframe(Keyframe);
+			ToInterpolationTime = *SimulationTimeBuffer.FindElementByKeyframe(Keyframe+1);
+			InterpolationTime = (FromInterpolationTime.ToRealTimeSeconds() + ToInterpolationTime.ToRealTimeSeconds()) / 2.f;
+		};
+
+		if (InterpolationKeyframe == INDEX_NONE)
+		{
+			//UE_LOG(LogTemp, Warning, TEXT("FELL BEHIND. %s < %s"), *LexToString(InterpolationTime), *LexToString(SimulationTimeBuffer.GetElementFromTail(0)->ToRealTimeSeconds()));
+			//ResetInterpolationTimeToMidpoint(SimulationTimeBuffer.GetTailKeyframe());
+			ResetInterpolationTimeToMidpoint((SimulationTimeBuffer.GetTailKeyframe() + SimulationTimeBuffer.GetHeadKeyframe()) / 2);
+		}
+		else
+		{
+			ToInterpolationTime = *SimulationTimeBuffer.FindElementByKeyframe(InterpolationKeyframe + 1);
+			if (InterpolationTime > ToInterpolationTime.ToRealTimeSeconds())
+			{
+				//UE_LOG(LogTemp, Warning, TEXT("GOT AHEAD. %s > %s"), *LexToString(InterpolationTime), *LexToString(ToInterpolationTime.ToRealTimeSeconds()));
+				//ResetInterpolationTimeToMidpoint(SimulationTimeBuffer.GetHeadKeyframe()-1);
+				ResetInterpolationTimeToMidpoint((SimulationTimeBuffer.GetTailKeyframe() + SimulationTimeBuffer.GetHeadKeyframe()) / 2);
+			}
+		}
+
+
+		const TSyncState* FromState = NetworkSim->Buffers.Sync.FindElementByKeyframe(InterpolationKeyframe);
+		const TSyncState* ToState = NetworkSim->Buffers.Sync.FindElementByKeyframe(InterpolationKeyframe+1);
+
+		if (FromState && ToState)
+		{
+			const TRealTime FromRealTime = FromInterpolationTime.ToRealTimeSeconds();
+			const TRealTime ToRealTime = ToInterpolationTime.ToRealTimeSeconds();
+			const TRealTime InterpolationInterval = ToRealTime - FromRealTime;
+		
+			if (ensure(FMath::Abs(InterpolationInterval) > 0.f))
+			{
+				const float InterpolationPCT = (InterpolationTime - FromRealTime) / InterpolationInterval;
+				ensureMsgf(InterpolationPCT >= 0.f && InterpolationPCT <= 1.f, TEXT("Calculated InterpolationPCT not in expected range. InterpolationTime: %s. From: %s. To: %s"), *LexToString(InterpolationTime), *LexToString(FromRealTime), *LexToString(ToRealTime));
+
+				TSyncState InterpolatedState;
+				TSyncState::Interpolate(*FromState, *ToState, InterpolationPCT, InterpolatedState);
+
+				//UE_LOG(LogTemp, Warning, TEXT("Interp. [%d] %s"), ToInterpolationData->Keyframe,  *LexToString(InterpolationTime));
 
+				Driver->FinalizeFrame(InterpolatedState);
+			}
+		}
+		else
+		{
+			ensure(false);
+		}
 
-template <typename TState>
-class TNetworkInterpolationHelper
+		//UE_LOG(LogTemp, Warning, TEXT("NumElements: [%d, %d]"), NetworkSim->Buffers.Sync.GetNumValidElements(), SimulationTimeBuffer.GetNumValidElements());
+
+		/*
+		InterpolationTime += DeltaSeconds;
+
+		static TRealTime LastInterpolationTime = 0.f;		
+		static TRealTime LastSimTime = 0.f;
+		static TSimTime LastSimTimeMSec;
+
+		TRealTime DeltaSim = NetworkSim->TickInfo.TotalProcessedSimulationTime.ToRealTimeSeconds() - LastSimTime;
+		TRealTime DeltaInt = InterpolationTime - LastInterpolationTime;
+		TSimTime DeltaSimMSec = NetworkSim->TickInfo.TotalProcessedSimulationTime - LastSimTimeMSec;
+		
+
+		TRealTime Delta = NetworkSim->TickInfo.TotalProcessedSimulationTime.ToRealTimeSeconds() - InterpolationTime;
+		UE_LOG(LogTemp, Warning, TEXT("%s (%s)... %s (%s) - %s (%s). Delta: %s"), *NetworkSim->TickInfo.TotalProcessedSimulationTime.ToString(), *DeltaSimMSec.ToString(),
+			*LexToString(NetworkSim->TickInfo.TotalProcessedSimulationTime.ToRealTimeSeconds()), *LexToString(DeltaSim), *LexToString(InterpolationTime), *LexToString(DeltaInt), *LexToString(Delta));
+
+		if (InterpolationTime > NetworkSim->TickInfo.TotalProcessedSimulationTime.ToRealTimeSeconds())
+		{
+			InterpolationTime = NetworkSim->TickInfo.TotalProcessedSimulationTime.ToRealTimeSeconds();
+			UE_LOG(LogTemp, Warning, TEXT("CLAMP! %s"), *LexToString(InterpolationTime));
+		}
+
+		LastInterpolationTime = InterpolationTime;
+		LastSimTime = NetworkSim->TickInfo.TotalProcessedSimulationTime.ToRealTimeSeconds();
+		LastSimTimeMSec = NetworkSim->TickInfo.TotalProcessedSimulationTime;
+
+		if (true)
+			return;
+
+		*/
+	}
+
+private:
+
+	const TNetworkedSimulationModel* NetworkSim;
+	TDriver* Driver;
+	
+	TRealTime InterpolationTime = 0.f; // SimTime we are currently interpolating at
+
+	TRealTime	MinBufferedTime = 1/120.f;
+	TRealTime	MaxBufferedTime = 1.f;
+
+};
+
+
+class UNetworkInterpolationManager
 {
+public:
+
+	virtual ~UNetworkInterpolationManager() { }
+
+
+	virtual void Tick(float DeltaSeconds);
+	
+	void Register(INetworkInterpolator* Interpolater)
+	{
+		Interpolaters.Add(Interpolater);
+	}
+
+	
+	void UnregisterNetworkSim(INetworkInterpolator* Interpolater)
+	{
+		Interpolaters.Remove(Interpolater);
+	}
+
+private:
 
+	TArray<INetworkInterpolator*> Interpolaters;
 };
\ No newline at end of file
diff --git a/Engine/Plugins/Runtime/NetworkPrediction/Source/NetworkPrediction/Public/NetworkSimulationModel.h b/Engine/Plugins/Runtime/NetworkPrediction/Source/NetworkPrediction/Public/NetworkSimulationModel.h
--- a/Engine/Plugins/Runtime/NetworkPrediction/Source/NetworkPrediction/Public/NetworkSimulationModel.h
+++ b/Engine/Plugins/Runtime/NetworkPrediction/Source/NetworkPrediction/Public/NetworkSimulationModel.h
@@ -148,93 +148,100 @@ public:
 		}
 
 		// -------------------------------------------------------------------------------------------------------------------------------------------------
-		//	PreInput processing
-		//	We are just about to start mechanically processing input. Lets just make sure we have a valid sync state to start with
+		//												Input Processing & Simulation Update
 		// -------------------------------------------------------------------------------------------------------------------------------------------------
-
-		// The SyncedState buffer needs to be in sync here:
-		//	-We want it to have a SyncedState, but it may not on the first frame through (thats ok).
-		//  -Its HeadKeyframe should be one behind the Keyframe we are about to process.
-		//
-		// Note, InputCmds start @ Keyframe=1. The first SyncedState that Update produces will go in KeyFrame=1.
-		// (E.g, InputCmd @ keyframe=X is used to generate MotionState @ keyframe=X)
-		// This means that SyncedState @ keyframe=0 is always created here via InitSyncState.
-		// This also means that we never actually process InputCmd @ keyframe=0. Which is why LastProcessedInputKeyframe is initialized to 0 ("already processed")
-		// and the buffer has an empty element inserted in InitializeForNetworkRole.
-
-		if (Buffers.Sync.GetHeadKeyframe() != TickInfo.LastProcessedInputKeyframe)
+		if (Buffers.Input.GetHeadKeyframe() > Buffers.Sync.GetHeadKeyframe())
 		{
-			if (TickInfo.LastProcessedInputKeyframe != 0)
+			// -------------------------------------------------------------------------------------------------------------
+			//	
+			//	The SyncedState buffer needs to be in sync here:
+			//		-We want it to have a SyncedState, but it may not on the first frame through (thats ok).
+			//		-Its HeadKeyframe should be one behind the Keyframe we are about to process.
+			//
+			//	Note, InputCmds start @ Keyframe=1. The first SyncedState that Update produces will go in KeyFrame=1.
+			//	(E.g, InputCmd @ keyframe=X is used to generate MotionState @ keyframe=X)
+			//	This means that SyncedState @ keyframe=0 is always created here via InitSyncState.
+			//	This also means that we never actually process InputCmd @ keyframe=0. Which is why LastProcessedInputKeyframe is initialized to 0 ("already processed")
+			//	and the buffer has an empty element inserted in InitializeForNetworkRole.
+			// -------------------------------------------------------------------------------------------------------------
+
+			if (Buffers.Sync.GetHeadKeyframe() != TickInfo.LastProcessedInputKeyframe)
 			{
-				// This shouldn't happen, but is not fatal. We are reseting the sync state buffer.
-				UE_LOG(LogNetworkSim, Warning, TEXT("%s. Break in SyncState continuity. LastProcessedInputKeyframe: %d. SyncBuffer.GetHeadKeyframe(): %d."), *Driver->GetDebugName() , TickInfo.LastProcessedInputKeyframe, Buffers.Sync.GetHeadKeyframe());
-			}
-
-			// We need an initial/current state. Get this from the sim driver
-			Buffers.Sync.ResetNextHeadKeyframe(TickInfo.LastProcessedInputKeyframe);
-			TSyncState* StartingState = Buffers.Sync.GetWriteNext();
-			Driver->InitSyncState(*StartingState);
-		}
+				if (TickInfo.LastProcessedInputKeyframe != 0)
+				{
+					// This shouldn't happen, but is not fatal. We are reseting the sync state buffer.
+					UE_LOG(LogNetworkSim, Warning, TEXT("%s. Break in SyncState continuity. LastProcessedInputKeyframe: %d. SyncBuffer.GetHeadKeyframe(): %d."), *Driver->GetDebugName() , TickInfo.LastProcessedInputKeyframe, Buffers.Sync.GetHeadKeyframe());
+				}
 
-		// -------------------------------------------------------------------------------------------------------------------------------------------------
-		//														Input Processing
-		// -------------------------------------------------------------------------------------------------------------------------------------------------
+				// We need an initial/current state. Get this from the sim driver
+				Buffers.Sync.ResetNextHeadKeyframe(TickInfo.LastProcessedInputKeyframe);
+				TSyncState* StartingState = Buffers.Sync.GetWriteNext();
+				Driver->InitSyncState(*StartingState);
 
-		while(true)
-		{
-			const int32 Keyframe = TickInfo.LastProcessedInputKeyframe+1;
-			if (Keyframe > TickInfo.MaxAllowedInputKeyframe)
-			{
-				break;
+				// Reset time tracking buffer too
+				TickInfo.SetTotalProcessedSimulationTime(TickInfo.GetTotalProcessedSimulationTime(), Buffers.Sync.GetHeadKeyframe());
 			}
-
-			if (TInputCmd* NextCmd = Buffers.Input.FindElementByKeyframe(Keyframe))
+		
+			// -------------------------------------------------------------------------------------------------------------
+			// Process Input
+			// -------------------------------------------------------------------------------------------------------------
+			while(true)
 			{
-				// We have an unprocessed command, do we have enough allotted simulation time to process it?
-				if (TickInfo.GetRemainingAllowedSimulationTime() >= NextCmd->GetFrameDeltaTime())
+				const int32 Keyframe = TickInfo.LastProcessedInputKeyframe+1;
+				if (Keyframe > TickInfo.MaxAllowedInputKeyframe)
 				{
-					// -------------------------------------------------------------------------------------------------
-					//	The core process input command and call ::Update block!
-					// -------------------------------------------------------------------------------------------------
-					TSyncState* PrevSyncState = Buffers.Sync.FindElementByKeyframe(TickInfo.LastProcessedInputKeyframe);
-					TSyncState* NextSyncState = Buffers.Sync.GetWriteNext();
-
-					check(PrevSyncState != nullptr);
-					check(NextSyncState != nullptr);
-					check(Buffers.Sync.GetHeadKeyframe() == Keyframe);
-				
-					if (DebugState)
+					break;
+				}
+
+				if (TInputCmd* NextCmd = Buffers.Input.FindElementByKeyframe(Keyframe))
+				{
+					// We have an unprocessed command, do we have enough allotted simulation time to process it?
+					if (TickInfo.GetRemainingAllowedSimulationTime() >= NextCmd->GetFrameDeltaTime())
 					{
-						DebugState->ProcessedKeyframes.Add(Keyframe);
-					}
+						// -------------------------------------------------------------------------------------------------
+						//	The core process input command and call ::Update block!
+						// -------------------------------------------------------------------------------------------------
+						TSyncState* PrevSyncState = Buffers.Sync.FindElementByKeyframe(TickInfo.LastProcessedInputKeyframe);
+						TSyncState* NextSyncState = Buffers.Sync.GetWriteNext();
+
+						check(PrevSyncState != nullptr);
+						check(NextSyncState != nullptr);
+						check(Buffers.Sync.GetHeadKeyframe() == Keyframe);
+				
+						if (DebugState)
+						{
+							DebugState->ProcessedKeyframes.Add(Keyframe);
+						}
 					
-					TAuxState AuxState; // Temp: aux buffer not implemented yet
+						TAuxState AuxState; // Temp: aux buffer not implemented yet
 
-					T::Update(Driver, NextCmd->GetFrameDeltaTime(), *NextCmd, *PrevSyncState, *NextSyncState, AuxState);
+						T::Update(Driver, NextCmd->GetFrameDeltaTime(), *NextCmd, *PrevSyncState, *NextSyncState, AuxState);
 					
-					TickInfo.TotalProcessedSimulationTime += NextCmd->GetFrameDeltaTime();
-					TickInfo.LastProcessedInputKeyframe = Keyframe;
+						TickInfo.IncrementTotalProcessedSimulationTime(NextCmd->GetFrameDeltaTime(), Keyframe);
+						TickInfo.LastProcessedInputKeyframe = Keyframe;
+					}
+					else
+					{
+						break;
+					}
 				}
 				else
 				{
 					break;
 				}
 			}
-			else
-			{
-				break;
-			}
-		}
 
-		// FIXME: this needs to be sorted out. We really want to check if there is new sync state and then call this here.
-		// Call into the driver to sync to the latest state if we processed any input
-		//if (NumProcessed > 0)
-		{
-			//check(SyncBuffer.GetNumValidElements() > 0);
-			if (Buffers.Sync.GetNumValidElements() > 0)
+			// FIXME: this needs to be sorted out. We really want to check if there is new sync state and then call this here.
+			// Call into the driver to sync to the latest state if we processed any input
+			//if (NumProcessed > 0)	
 			{
-				Driver->FinalizeFrame(*Buffers.Sync.GetElementFromHead(0));
+				//check(SyncBuffer.GetNumValidElements() > 0);
+				if (Buffers.Sync.GetNumValidElements() > 0)
+				{
+					Driver->FinalizeFrame(*Buffers.Sync.GetElementFromHead(0));
+				}
 			}
+
 		}
 
 		// -------------------------------------------------------------------------------------------------------------------------------------------------
@@ -274,7 +281,9 @@ public:
 			MyHistoricBuffers->Input.SetBufferSize(Parameters.HistoricBufferSize);
 			MyHistoricBuffers->Sync.SetBufferSize(Parameters.HistoricBufferSize);
 			MyHistoricBuffers->Aux.SetBufferSize(Parameters.HistoricBufferSize);
-		}		
+		}
+
+		TickInfo.InitSimulationTimeBuffer(Parameters.SyncedBufferSize);
 
 		// We want to start with an empty command in the input buffer. The sync buffer will be populated @ frame 0 with the "current" state when we actually sim. This keeps them in sync
 		*Buffers.Input.GetWriteNext() = TInputCmd();
@@ -282,27 +291,23 @@ public:
 
 	void NetSerializeProxy(EReplicationProxyTarget Target, const FNetSerializeParams& Params)
 	{
-		// You are not allowed to change the simulations TickInfo while NetSerializing. Since all rep proxies are templated
-		// and there is not interface/base class, we are enforcing the constness at the call site.
-		const TSimulationTickState<TTickSettings>& ConstTickInfo = const_cast<const TSimulationTickState<TTickSettings>&>(TickInfo);
-
 		switch(Target)
 		{
 		case EReplicationProxyTarget::ServerRPC:
-			RepProxy_ServerRPC.NetSerialize(Params, Buffers, ConstTickInfo);
+			RepProxy_ServerRPC.NetSerialize(Params, Buffers, TickInfo);
 			break;
 		case EReplicationProxyTarget::AutonomousProxy:
-			RepProxy_Autonomous.NetSerialize(Params, Buffers, ConstTickInfo);
+			RepProxy_Autonomous.NetSerialize(Params, Buffers, TickInfo);
 			break;
 		case EReplicationProxyTarget::SimulatedProxy:
-			RepProxy_Simulated.NetSerialize(Params, Buffers, ConstTickInfo);
+			RepProxy_Simulated.NetSerialize(Params, Buffers, TickInfo);
 			break;
 		case EReplicationProxyTarget::Replay:
-			RepProxy_Replay.NetSerialize(Params, Buffers, ConstTickInfo);
+			RepProxy_Replay.NetSerialize(Params, Buffers, TickInfo);
 			break;
 		case EReplicationProxyTarget::Debug:
 #if NETSIM_MODEL_DEBUG
-			RepProxy_Debug.NetSerialize(Params, Buffers, ConstTickInfo);
+			RepProxy_Debug.NetSerialize(Params, Buffers, TickInfo);
 			break;
 #endif
 		default:
diff --git a/Engine/Plugins/Runtime/NetworkPrediction/Source/NetworkPrediction/Public/NetworkSimulationModelBuffer.h b/Engine/Plugins/Runtime/NetworkPrediction/Source/NetworkPrediction/Public/NetworkSimulationModelBuffer.h
--- a/Engine/Plugins/Runtime/NetworkPrediction/Source/NetworkPrediction/Public/NetworkSimulationModelBuffer.h
+++ b/Engine/Plugins/Runtime/NetworkPrediction/Source/NetworkPrediction/Public/NetworkSimulationModelBuffer.h
@@ -54,37 +54,27 @@ struct TReplicationBuffer
 	}
 
 	// Gets the oldest element in the buffer (or null)
-	T* GetElementFromTail(int32 OffsetFromTail)
+	const T* GetElementFromTail(int32 OffsetFromTail) const
 	{
-		check(OffsetFromTail >= 0); // Must always pass in valid offset
+		return GetElementFromTailImpl(OffsetFromTail);
+	}
 
-		// Out of range. This is not fatal, just don't return the wrong element.
-		if (NumValidElements <= OffsetFromTail)
-		{
-			return nullptr;
-		}
-		
-		const int32 Tail = GetTailKeyframe();
-		const int32 Position = Tail + OffsetFromTail;
-		const int32 Offset = (Position % Data.Num());
-		return &Data[Offset];
+	T* GetElementFromTail(int32 OffsetFromTail)
+	{
+		const T* Element = const_cast<TReplicationBuffer<T>*>(this)->GetElementFromTailImpl(OffsetFromTail);
+		return const_cast<T*>(Element);
 	}
 
 	// Gets the newest element in the buffer (or null)
-	T* GetElementFromHead(int32 OffsetFromHead)
+	const T* GetElementFromHead(int32 OffsetFromHead) const
 	{
-		check(OffsetFromHead >= 0); // Must always pass in valid offset
+		return GetElementFromHeadImpl(OffsetFromHead);
+	}
 
-		// Out of range. This is not fatal, just don't return the wrong element.
-		if (NumValidElements <= OffsetFromHead)
-		{
-			return nullptr;
-		}
-		
-		const int32 Position = Head - OffsetFromHead;
-		check(Position >= 0);
-		const int32 Offset = (Position % Data.Num());
-		return &Data[Offset];
+	T* GetElementFromHead(int32 OffsetFromHead)
+	{
+		const T* Element = const_cast<TReplicationBuffer<T>*>(this)->GetElementFromHeadImpl(OffsetFromHead);
+		return const_cast<T*>(Element);
 	}
 	
 	const T* FindElementByKeyframe(int32 Key) const
@@ -202,6 +192,34 @@ private:
 		return &Data[Key % Data.Num()];
 	}
 
+	const T* GetElementFromTailImpl(int32 OffsetFromTail) const
+	{
+		check(OffsetFromTail >= 0); // Must always pass in valid offset	
+		if (NumValidElements <= OffsetFromTail)
+		{
+			return nullptr;
+		}
+		
+		const int32 Tail = GetTailKeyframe();
+		const int32 Position = Tail + OffsetFromTail;
+		const int32 Offset = (Position % Data.Num());
+		return &Data[Offset];
+	}
+
+	const T* GetElementFromHeadImpl(int32 OffsetFromHead) const
+	{
+		check(OffsetFromHead >= 0); // Must always pass in valid offset
+		if (NumValidElements <= OffsetFromHead)
+		{
+			return nullptr;
+		}
+		
+		const int32 Position = Head - OffsetFromHead;
+		check(Position >= 0);
+		const int32 Offset = (Position % Data.Num());
+		return &Data[Offset];
+	}
+
 	int32 Head = INDEX_NONE;
 	int32 DirtyCount = 0;
 	int32 NumValidElements = 0;
diff --git a/Engine/Plugins/Runtime/NetworkPrediction/Source/NetworkPrediction/Public/NetworkSimulationModelDebugger.h b/Engine/Plugins/Runtime/NetworkPrediction/Source/NetworkPrediction/Public/NetworkSimulationModelDebugger.h
--- a/Engine/Plugins/Runtime/NetworkPrediction/Source/NetworkPrediction/Public/NetworkSimulationModelDebugger.h
+++ b/Engine/Plugins/Runtime/NetworkPrediction/Source/NetworkPrediction/Public/NetworkSimulationModelDebugger.h
@@ -301,6 +301,9 @@ private:
 						}
 					}
 				}
+
+				// Only gather first active debugger (it would be great to have more control over this when debugging multiples)
+				break;
 			}
 		}
 	}
@@ -532,8 +535,8 @@ struct TNetworkSimulationModelDebugger : public INetworkSimulationModelDebugger
 
 		Out.Emit(FString::Printf(TEXT("%s - %s"), *Owner->GetName(), *UEnum::GetValueAsString(TEXT("Engine.ENetRole"), Owner->Role)), FColor::Yellow);
 		Out.Emit(FString::Printf(TEXT("LastProcessedInputKeyframe: %d (%d Buffered)"), NetworkSim->TickInfo.LastProcessedInputKeyframe, NetworkSim->Buffers.Input.GetHeadKeyframe() - NetworkSim->TickInfo.LastProcessedInputKeyframe));
-
-		// Autorproxy
+				
+		if (Owner->Role == ROLE_AutonomousProxy)
 		{			
 			FColor Color = FColor::White;
 			const bool FaultDetected = NetworkSim->RepProxy_Autonomous.IsReconcileFaultDetected();
@@ -560,13 +563,19 @@ struct TNetworkSimulationModelDebugger : public INetworkSimulationModelDebugger
 
 			Out.Emit(*ConfirmedFrameStr, Color);
 
-			FString SimulationTimeString = FString::Printf(TEXT("Local SimulationTime: %s. SerialisedSimulationTime: %s. Difference MS: %s"), *NetworkSim->TickInfo.TotalProcessedSimulationTime.ToString(),
-				*NetworkSim->RepProxy_Autonomous.GetLastSerializedSimTime().ToString(), *(NetworkSim->TickInfo.TotalProcessedSimulationTime - NetworkSim->RepProxy_Autonomous.GetLastSerializedSimTime()).ToString());
+			FString SimulationTimeString = FString::Printf(TEXT("Local SimulationTime: %s. SerializedSimulationTime: %s. Difference MS: %s"), *NetworkSim->TickInfo.GetTotalProcessedSimulationTime().ToString(),
+				*NetworkSim->RepProxy_Autonomous.GetLastSerializedSimTime().ToString(), *(NetworkSim->TickInfo.GetTotalProcessedSimulationTime() - NetworkSim->RepProxy_Autonomous.GetLastSerializedSimTime()).ToString());
 			Out.Emit(*SimulationTimeString, Color);
 
 			FString AllowedSimulationTimeString = FString::Printf(TEXT("Allowed Simulation Time: %s. Keyframe: %d/%d/%d"), *NetworkSim->TickInfo.GetRemainingAllowedSimulationTime().ToString(), NetworkSim->TickInfo.MaxAllowedInputKeyframe, NetworkSim->TickInfo.LastProcessedInputKeyframe, NetworkSim->Buffers.Input.GetHeadKeyframe());
 			Out.Emit(*AllowedSimulationTimeString, Color);
 		}
+		else if (Owner->Role == ROLE_SimulatedProxy)
+		{
+			FColor Color = FColor::White;
+			FString TimeString = FString::Printf(TEXT("Total Processed Simulation Time: %s. Last Serialized Simulation Time: %s. Delta: %s"), *NetworkSim->TickInfo.GetTotalProcessedSimulationTime().ToString(), *NetworkSim->RepProxy_Simulated.GetLastSerializedSimulationTime().ToString(), *(NetworkSim->RepProxy_Simulated.GetLastSerializedSimulationTime() - NetworkSim->TickInfo.GetTotalProcessedSimulationTime()).ToString());
+			Out.Emit(*TimeString, Color);
+		}
 
 		auto EmitBuffer = [&Out](FString BufferName, auto& Buffer)
 		{
diff --git a/Engine/Plugins/Runtime/NetworkPrediction/Source/NetworkPrediction/Public/NetworkSimulationModelReplicators.h b/Engine/Plugins/Runtime/NetworkPrediction/Source/NetworkPrediction/Public/NetworkSimulationModelReplicators.h
--- a/Engine/Plugins/Runtime/NetworkPrediction/Source/NetworkPrediction/Public/NetworkSimulationModelReplicators.h
+++ b/Engine/Plugins/Runtime/NetworkPrediction/Source/NetworkPrediction/Public/NetworkSimulationModelReplicators.h
@@ -52,8 +52,6 @@ struct FNetworkSimulationSerialization
 // ----------------------------------------------------------------------------------------------------------------------------------------------------------------------
 //	Templated Replicators
 //
-//	
-//
 // ----------------------------------------------------------------------------------------------------------------------------------------------------------------------
 
 template<typename TBufferTypes, typename TTickSettings>
@@ -106,10 +104,10 @@ struct TReplicatorBase
 template<typename TBufferTypes, typename TTickSettings, typename TBase=TReplicatorBase<TBufferTypes, TTickSettings>, bool Enabled=true>
 struct TReplicator_SimTime : public TBase
 {
-	void NetSerialize(const FNetSerializeParams& P, TNetworkSimBufferContainer<TBufferTypes>& Buffers, const TSimulationTickState<TTickSettings>& TickInfo)
+	void NetSerialize(const FNetSerializeParams& P, TNetworkSimBufferContainer<TBufferTypes>& Buffers, TSimulationTickState<TTickSettings>& TickInfo)
 	{
 		TBase::NetSerialize(P, Buffers, TickInfo);
-		SerializedTime = TickInfo.TotalProcessedSimulationTime;
+		SerializedTime = TickInfo.GetTotalProcessedSimulationTime();
 		SerializedTime.NetSerialize(P.Ar);
 	}
 
@@ -135,7 +133,7 @@ struct TReplicator_Sequence : public TBase
 		return Buffers.template Get<BufferId>().GetDirtyCount() ^ (TBase::GetProxyDirtyCount(Buffers) << 2); 
 	}
 
-	void NetSerialize(const FNetSerializeParams& P, TNetworkSimBufferContainer<TBufferTypes>& Buffers, const TSimulationTickState<TTickSettings>& TickInfo)
+	void NetSerialize(const FNetSerializeParams& P, TNetworkSimBufferContainer<TBufferTypes>& Buffers, TSimulationTickState<TTickSettings>& TickInfo)
 	{
 		TBase::NetSerialize(P, Buffers, TickInfo);
 
@@ -190,7 +188,7 @@ struct TReplicator_Single : public TBase
 		return Buffers.template Get<BufferId>().GetDirtyCount() ^ (TBase::GetProxyDirtyCount(Buffers) << 2); ; 
 	}
 
-	void NetSerialize(const FNetSerializeParams& P, TNetworkSimBufferContainer<TBufferTypes>& Buffers, const TSimulationTickState<TTickSettings>& TickInfo)
+	void NetSerialize(const FNetSerializeParams& P, TNetworkSimBufferContainer<TBufferTypes>& Buffers, TSimulationTickState<TTickSettings>& TickInfo)
 	{
 		TBase::NetSerialize(P, Buffers, TickInfo);
 
@@ -262,118 +260,205 @@ struct TReplicator_Server : public TBase
 };
 
 /** Simulated: "non locally controlled" simulations. We support "Simulation Extrapolation" here (using the sim to fake inputs to advance the sim)  */
-template<typename TBufferTypes, typename TTickSettings, typename TBase=TReplicator_Single<TBufferTypes, TTickSettings, ENetworkSimBufferTypeId::Sync>>
+template<typename TBufferTypes, typename TTickSettings, typename TBase=TReplicatorEmpty<TBufferTypes, TTickSettings>>
 struct TReplicator_Simulated : public TBase
 {
 	using TInputCmd = typename TBufferTypes::TInputCmd;
 	using TSyncState = typename TBufferTypes::TSyncState;
 	using TAuxState = typename TBufferTypes::TAuxState;
+
+	// Instance flag for enabling simulated extrapolation
+	bool bAllowSimulatedExtrapolation = true;
+
+	TNetworkSimTime<TTickSettings> GetLastSerializedSimulationTime() const { return LastSerializedSimulationTime; }
 	
 	int32 GetProxyDirtyCount(TNetworkSimBufferContainer<TBufferTypes>& Buffers) const
 	{
 		return Buffers.Sync.GetDirtyCount() ^ (TBase::GetProxyDirtyCount(Buffers) << 2); 
 	}
 
-	void NetSerialize(const FNetSerializeParams& P, TNetworkSimBufferContainer<TBufferTypes>& Buffers, const TSimulationTickState<TTickSettings>& TickInfo)
+	void NetSerialize(const FNetSerializeParams& P, TNetworkSimBufferContainer<TBufferTypes>& Buffers, TSimulationTickState<TTickSettings>& TickInfo)
 	{
-		TBase::NetSerialize(P, Buffers, TickInfo);
+		FArchive& Ar = P.Ar;
+
+		TNetworkSimTime<TTickSettings> PrevLastSerializedSimulationTime = LastSerializedSimulationTime;
+
+		// Serialize latest simulation time
+		LastSerializedSimulationTime = TickInfo.GetTotalProcessedSimulationTime();
+		LastSerializedSimulationTime.NetSerialize(P.Ar);
+
+		// Serialize latest element
+		TSyncState* State = nullptr;
+		
+		if (Ar.IsSaving())
+		{
+			State = Buffers.Sync.GetElementFromHead(0);
+			check(State); // We should not be here if the buffer is empty. Want to avoid serializing an "empty" flag at the top here.
+		}
+		else
+		{
+			check(TickInfo.SimulationTimeBuffer.GetHeadKeyframe() == Buffers.Sync.GetHeadKeyframe());
+			check(TickInfo.GetTotalProcessedSimulationTime() <= TickInfo.GetTotalAllowedSimulationTime());
+
+			// Cache off our "starting" time before possibly overwriting it. We will use this in Reconcile to catch back up in some cases.
+			if (TickInfo.GetTotalProcessedSimulationTime() > ReconcileSimulationTime)
+			{
+				ReconcileSimulationTime = TickInfo.GetTotalProcessedSimulationTime();
+			}
+
+			// Find out where this should go in the local buffer based on the serialized time
+			int32 DestinationKeyframe = INDEX_NONE;
+			if (LastSerializedSimulationTime > TickInfo.GetTotalProcessedSimulationTime())
+			{
+				// We are getting new state that is ahead of what we have locally, so it can safety go right to head
+				DestinationKeyframe = TickInfo.SimulationTimeBuffer.GetHeadKeyframe()+1;
+
+				//UE_LOG(LogNetworkSim, Warning, TEXT("@@@"));
+			}
+			else
+			{
+				// We are getting state that is behind what we have locally
+				for (int32 Keyframe = TickInfo.SimulationTimeBuffer.GetHeadKeyframe(); Keyframe >= TickInfo.SimulationTimeBuffer.GetTailKeyframe(); --Keyframe)
+				{
+					if (LastSerializedSimulationTime > *TickInfo.SimulationTimeBuffer.FindElementByKeyframe(Keyframe))
+					{
+						//UE_LOG(LogNetworkSim, Warning, TEXT("..."));
+						DestinationKeyframe = Keyframe+1;
+						break;
+					}
+				}
+
+				if (DestinationKeyframe == INDEX_NONE)
+				{
+					TNetworkSimTime<TTickSettings> TotalTimeAhead = *TickInfo.SimulationTimeBuffer.GetElementFromHead(0) - LastSerializedSimulationTime;
+					TNetworkSimTime<TTickSettings> SerializeDelta = LastSerializedSimulationTime - PrevLastSerializedSimulationTime;
+
+					// We are way far ahead of the server... we will need to clear out sync buffers, take what they gave us, and catch up in reconcile
+					//UE_LOG(LogNetworkSim, Warning, TEXT("!!! TReplicator_Simulated. Large gap detected. SerializedTime: %s. Buffer time: [%s-%s]. %d Elements. DeltaFromHead: %s. DeltaSerialize: %s"), *LastSerializedSimulationTime.ToString(), 
+					//	*TickInfo.SimulationTimeBuffer.GetElementFromTail(0)->ToString(), *TickInfo.SimulationTimeBuffer.GetElementFromHead(0)->ToString(), TickInfo.SimulationTimeBuffer.GetNumValidElements(), *TotalTimeAhead.ToString(), *SerializeDelta.ToString());
+					DestinationKeyframe = TickInfo.SimulationTimeBuffer.GetHeadKeyframe()+2; // (Skip ahead 2 to force a break in continuity)
+				}
+			}
+
+			check(DestinationKeyframe != INDEX_NONE);
 
-		// Temp: can go away once reconcile is called only after receiving packet
-		if (P.Ar.IsLoading())
-		{		
-			bPendingReconciliation = true;
+			// "Finalize" our buffers and time keeping such that we serialize the latest state from the server in the right spot			
+			Buffers.Sync.ResetNextHeadKeyframe(DestinationKeyframe);
+			State = Buffers.Sync.GetWriteNext();
+
+			// Update tick info
+			TickInfo.SetTotalProcessedSimulationTime(LastSerializedSimulationTime, DestinationKeyframe);
+			if (TickInfo.GetTotalAllowedSimulationTime() < LastSerializedSimulationTime)
+			{
+				TickInfo.SetTotalAllowedSimulationTime(LastSerializedSimulationTime);
+			}
+
+			check(TickInfo.GetTotalProcessedSimulationTime() <= TickInfo.GetTotalAllowedSimulationTime());
+
+			TickInfo.LastProcessedInputKeyframe = DestinationKeyframe;
+			TickInfo.MaxAllowedInputKeyframe = DestinationKeyframe;
+
+			Buffers.Input.ResetNextHeadKeyframe(DestinationKeyframe+1);
 		}
+
+		check(State);
+		State->NetSerialize(Ar);
 	}
 
 	template<typename T, typename TDriver>
 	void Reconcile(TDriver* Driver, TNetworkSimBufferContainer<TBufferTypes>& Buffers, TSimulationTickState<TTickSettings>& TickInfo)
 	{
-		if (!bPendingReconciliation)
+		if (ReconcileSimulationTime.IsPositive() == false)
+		{
 			return;
+		}
 
-		bPendingReconciliation = false;
-		TInputCmd* LastCmd = Buffers.Input.GetElementFromHead(0);
-		
+		check(TickInfo.GetTotalProcessedSimulationTime() <= TickInfo.GetTotalAllowedSimulationTime());
+
+		if (bAllowSimulatedExtrapolation && NetworkSimulationModelCVars::EnableSimulatedExtrapolation() && NetworkSimulationModelCVars::EnableSimulatedReconcile())
 		{
-			// Right now NetSerialize can be called multiple times before Reconcile can get called. Eventually we'd like to make it 1:1.
-			// Until then, lets handle multiple serializes.
-			int32 DeltaKeyframes = Buffers.Sync.GetHeadKeyframe() - Buffers.Input.GetHeadKeyframe();
-			if (!ensure(DeltaKeyframes >= 0 && DeltaKeyframes < 1000))
+			TInputCmd* LastCmd = Buffers.Input.GetElementFromHead(0);
+
+			// Simulated Reconcile requires the input buffer to be kept up to date with the Sync buffer
+			// Generate a new, fake, command since we just added a new sync state to head
+			while (Buffers.Input.GetHeadKeyframe() < Buffers.Sync.GetHeadKeyframe())
 			{
-				// We way out of sync if this gets hit but this will reset things
-				Buffers.Input.ResetNextHeadKeyframe( Buffers.Sync.GetHeadKeyframe() );
+				TInputCmd* Next = Buffers.Input.GetWriteNext();
+				*Next = LastCmd ? *LastCmd : TInputCmd();
 			}
-		}
-		
-		// Generate a new, fake, command since we just added a new sync state to head
-		while (Buffers.Input.GetHeadKeyframe() < Buffers.Sync.GetHeadKeyframe())
-		{
-			Buffers.Input.GetWriteNext();
-		}
 
-		TNetworkSimTime<TTickSettings> DeltaSimTime = TickInfo.TotalProcessedSimulationTime - this->SerializedTime;
-		if (DeltaSimTime > TNetworkSimTime<TTickSettings>::FromMSec(0) && NetworkSimulationModelCVars::EnableSimulatedReconcile())
-		{
-			// This is the "Simulated Reconcile" part:
-			// We have extrapolated ahead of the server. The latest network update is now "in the past" from what we rendered last frame.
-			// To avoid popping or visual artifact, we will make up the difference here since we know the extra simulation time of the server.
-			// (this can easily happen: variance in latency may delay a packet and we may get ahead)
-			TInputCmd* NewCmd = Buffers.Input.GetWriteNext();
-			*NewCmd = LastCmd ? *LastCmd : TInputCmd();
-			NewCmd->SetFrameDeltaTime(DeltaSimTime);
+			// Do we have time to make up? We may have extrapolated ahead of the server (totally fine - can happen with small amount of latency variance)
+			TNetworkSimTime<TTickSettings> DeltaSimTime = ReconcileSimulationTime - TickInfo.GetTotalProcessedSimulationTime();
+			if (DeltaSimTime.IsPositive() && NetworkSimulationModelCVars::EnableSimulatedReconcile())
+			{
+				// We have extrapolated ahead of the server. The latest network update is now "in the past" from what we rendered last frame.
+				// We will insert a new keyframe to make up the difference from the last known state to where we want to be in the now.
+				
+				TInputCmd* NewCmd = Buffers.Input.GetWriteNext();
+				*NewCmd = LastCmd ? *LastCmd : TInputCmd();
+				NewCmd->SetFrameDeltaTime(DeltaSimTime);
 
-			TSyncState* PrevSyncState = Buffers.Sync.GetElementFromHead(0);
-			TSyncState* NextSyncState = Buffers.Sync.GetWriteNext();
+				TSyncState* PrevSyncState = Buffers.Sync.GetElementFromHead(0);
+				TSyncState* NextSyncState = Buffers.Sync.GetWriteNext();
 
-			TAuxState Junk;
+				TAuxState Junk;
 
-			// TODO: log this guy
-			//NextMotionState->VisualLog( FVisualLoggingParameters(Keyframe == LastKeyframeToProcess ? EVisualLoggingContext::LastMispredicted : EVisualLoggingContext::OtherMispredicted, Keyframe, EVisualLoggingLifetime::Persistent), Driver, Driver);
+				// TODO: log this guy
+				//NextMotionState->VisualLog( FVisualLoggingParameters(Keyframe == LastKeyframeToProcess ? EVisualLoggingContext::LastMispredicted : EVisualLoggingContext::OtherMispredicted, Keyframe, EVisualLoggingLifetime::Persistent), Driver, Driver);
 
-			// Do the actual update
-			T::Update(Driver, NewCmd->GetFrameDeltaTime(), *NewCmd, *PrevSyncState, *NextSyncState, Junk);
-		}
-		else
-		{
-			// We are taking the latest server update as head. So set our clocks to it.
-			TickInfo.TotalProcessedSimulationTime = this->SerializedTime;
-			TickInfo.TotalAllowedSimulationTime = this->SerializedTime;
-		}
+				//UE_LOG(LogNetworkSim, Warning, TEXT("   >>> Reconcile Update: %s"), *DeltaSimTime.ToString());
+
+				// Do the actual update
+				T::Update(Driver, NewCmd->GetFrameDeltaTime(), *NewCmd, *PrevSyncState, *NextSyncState, Junk);
+				TickInfo.IncrementTotalProcessedSimulationTime(NewCmd->GetFrameDeltaTime(), Buffers.Sync.GetHeadKeyframe());
 
-		// Set our LastProcessedInputKeyframe to fake that we handled it
-		TickInfo.LastProcessedInputKeyframe = Buffers.Input.GetHeadKeyframe();
-		TickInfo.MaxAllowedInputKeyframe = FMath::Max(TickInfo.MaxAllowedInputKeyframe, TickInfo.LastProcessedInputKeyframe);
+				// Set our LastProcessedInputKeyframe to fake that we handled it
+				TickInfo.LastProcessedInputKeyframe = Buffers.Input.GetHeadKeyframe();
+				TickInfo.MaxAllowedInputKeyframe = Buffers.Input.GetHeadKeyframe();
+			}
+		}
+		
+		check(TickInfo.GetTotalProcessedSimulationTime() <= TickInfo.GetTotalAllowedSimulationTime());
+		ReconcileSimulationTime.Reset();
 	}
 
 	template<typename T, typename TDriver, typename TTickParameters>
 	void PreSimTick(TDriver* Driver, TNetworkSimBufferContainer<TBufferTypes>& Buffers, TSimulationTickState<TTickSettings>& TickInfo, const TTickParameters& TickParameters)
 	{
-		if (NetworkSimulationModelCVars::EnableSimulatedExtrapolation() == 0)
+		if (bAllowSimulatedExtrapolation && NetworkSimulationModelCVars::EnableSimulatedExtrapolation())
 		{
-			return;
-		}
-		
-		TickInfo.GiveSimulationTime(TickParameters.LocalDeltaTimeSeconds);
+			// Don't start this simulation until you've gotten at least one update from the server
+			if (TickInfo.GetTotalProcessedSimulationTime().IsPositive())
+			{
+				TickInfo.GiveSimulationTime(TickParameters.LocalDeltaTimeSeconds);
+			}
 
-		if (TickParameters.bGenerateLocalInputCmds)
-		{
-			TNetworkSimTime<TTickSettings> DeltaSimTime = TickInfo.GetRemainingAllowedSimulationTime();
-			if (DeltaSimTime.IsPositive())
+			if (TickParameters.bGenerateLocalInputCmds)
 			{
-				if (TInputCmd* InputCmd = Buffers.Input.GetWriteNext())
+				TNetworkSimTime<TTickSettings> DeltaSimTime = TickInfo.GetRemainingAllowedSimulationTime();
+				if (DeltaSimTime.IsPositive())
 				{
-					*InputCmd = TInputCmd();
-					InputCmd->SetFrameDeltaTime(DeltaSimTime);
-					Driver->ProduceInput(DeltaSimTime, *InputCmd);
-					TickInfo.MaxAllowedInputKeyframe = Buffers.Input.GetHeadKeyframe();
+					if (TInputCmd* InputCmd = Buffers.Input.GetWriteNext())
+					{
+						//UE_LOG(LogNetworkSim, Warning, TEXT("   Extrapolating %s"), *DeltaSimTime.ToString());
+
+						*InputCmd = TInputCmd();
+						InputCmd->SetFrameDeltaTime(DeltaSimTime);
+						Driver->ProduceInput(DeltaSimTime, *InputCmd);
+						TickInfo.MaxAllowedInputKeyframe = Buffers.Input.GetHeadKeyframe();
+					}
 				}
 			}
 		}
 	}
 
+
+
 private:
 	
-	bool bPendingReconciliation = false;
+	TNetworkSimTime<TTickSettings> ReconcileSimulationTime;
+	TNetworkSimTime<TTickSettings> LastSerializedSimulationTime;
+	
 };
 
 
@@ -397,7 +482,7 @@ struct TReplicator_Autonomous : public TBase
 	// --------------------------------------------------------------------
 	//	NetSerialize
 	// --------------------------------------------------------------------
-	void NetSerialize(const FNetSerializeParams& P, TNetworkSimBufferContainer<TBufferTypes>& Buffers, const TSimulationTickState<TTickSettings>& TickInfo)
+	void NetSerialize(const FNetSerializeParams& P, TNetworkSimBufferContainer<TBufferTypes>& Buffers, TSimulationTickState<TTickSettings>& TickInfo)
 	{
 		FArchive& Ar = P.Ar;
 
@@ -405,7 +490,7 @@ struct TReplicator_Autonomous : public TBase
 		TSyncState* SerializedState = nullptr;
 
 		// Serialize total simulation time. This isn't really necessary since we have the keyframe above. 
-		SerializedTime = TickInfo.TotalProcessedSimulationTime;
+		SerializedTime = TickInfo.GetTotalProcessedSimulationTime();
 		SerializedTime.NetSerialize(P.Ar);
 
 		if (Ar.IsSaving())
@@ -518,14 +603,14 @@ struct TReplicator_Autonomous : public TBase
 		*ClientSyncState = *ServerState;
 
 		// Set the canonical simulation time to what we received (we will advance it as we resimulate)
-		TickInfo.TotalProcessedSimulationTime = SerializedTime;
+		TickInfo.SetTotalProcessedSimulationTime(SerializedTime, ReconciliationKeyframe);
 		TickInfo.LastProcessedInputKeyframe = ReconciliationKeyframe;
 		TickInfo.MaxAllowedInputKeyframe = FMath::Max(TickInfo.MaxAllowedInputKeyframe, TickInfo.LastProcessedInputKeyframe); // Make sure this doesn't lag behind. This is the only place we should need to do this.
 
 		if (NetworkSimulationModelCVars::EnableLocalPrediction() == 0)
 		{
 			// If we aren't predicting at all, then we advanced the allowed sim time here, (since we aren't doing it in PreSimTick). This just keeps us constantly falling behind and not being able to toggle prediction on/off for debugging.
-			TickInfo.TotalAllowedSimulationTime = SerializedTime;
+			TickInfo.SetTotalAllowedSimulationTime(SerializedTime);
 		}
 		
 		// Resimulate all user commands 
@@ -563,7 +648,7 @@ struct TReplicator_Autonomous : public TBase
 			T::Update(Driver, ResimulateCmd->GetFrameDeltaTime(), *ResimulateCmd, *PrevMotionState, *NextMotionState, *AuxState);
 			
 			// Update TickInfo
-			TickInfo.TotalProcessedSimulationTime += ResimulateCmd->GetFrameDeltaTime();
+			TickInfo.IncrementTotalProcessedSimulationTime(ResimulateCmd->GetFrameDeltaTime(), Keyframe);
 			TickInfo.LastProcessedInputKeyframe = Keyframe;
 
 			// Log out the newly predicted state that we got.
@@ -636,7 +721,7 @@ private:
 };
 
 /** Special replicator for debug buffer, this preserves the local buffer and receives into a replicator-owned buffer (we want these buffers to be distinct/not merged) */
-template<typename TBufferTypes, typename TTickSettings, int32 MaxNumElements=5, typename TBase=TReplicator_SimTime<TBufferTypes, TTickSettings>>
+template<typename TBufferTypes, typename TTickSettings, int32 MaxNumElements=5, typename TBase=TReplicatorEmpty<TBufferTypes, TTickSettings>>
 struct TReplicator_Debug : public TBase
 {
 	using TDebugState = typename TBufferTypes::TDebugState;
@@ -646,7 +731,7 @@ struct TReplicator_Debug : public TBase
 		return Buffers.Debug.GetDirtyCount() ^ (TBase::GetProxyDirtyCount(Buffers) << 2); 
 	}
 
-	void NetSerialize(const FNetSerializeParams& P, TNetworkSimBufferContainer<TBufferTypes>& Buffers, const TSimulationTickState<TTickSettings>& TickInfo)
+	void NetSerialize(const FNetSerializeParams& P, TNetworkSimBufferContainer<TBufferTypes>& Buffers, TSimulationTickState<TTickSettings>& TickInfo)
 	{
 		TBase::NetSerialize(P, Buffers, TickInfo);
 		FArchive& Ar = P.Ar;
diff --git a/Engine/Plugins/Runtime/NetworkPrediction/Source/NetworkPrediction/Public/NetworkSimulationModelTypes.h b/Engine/Plugins/Runtime/NetworkPrediction/Source/NetworkPrediction/Public/NetworkSimulationModelTypes.h
--- a/Engine/Plugins/Runtime/NetworkPrediction/Source/NetworkPrediction/Public/NetworkSimulationModelTypes.h
+++ b/Engine/Plugins/Runtime/NetworkPrediction/Source/NetworkPrediction/Public/NetworkSimulationModelTypes.h
@@ -169,6 +169,9 @@ struct TNetworkSimBufferContainer
 // ----------------------------------------------------------------------------------------------------------------------------------------------
 
 // Holds all settings for a network sim related to ticking
+// Main point are we want to distinguish between real time coming in from the outside in seconds. Float by default.
+// "SimTime" meaning time type we when advancing the simulation or talking about how long it has been running. This is stored in int32 milliseconds by default.
+// (This can be changed by the InRealToSimFactor parameter. Though code will always refer to is as MSec for clarity/consistency).
 template<int32 InFixedStepMS=0, int32 InMaxStepMS=0, typename TUnderlyingSimTimeType=int32, typename TUnderlyingRealTimeType=float, int32 InRealToSimFactor=1000>
 struct TNetworkSimTickSettings
 {
@@ -203,8 +206,12 @@ struct TNetworkSimTime
 	// Things get confusing with templated types and overloaded functions. To avoid that, use these funcs to construct from either msec or real time
 	static inline TNetworkSimTime FromMSec(const TTime& InTime) { return TNetworkSimTime(InTime); } 
 	static inline TNetworkSimTime FromRealTimeSeconds(const TRealTime& InRealTime) { return TNetworkSimTime( static_cast<TTime>(InRealTime * TickSettings::GetRealToSimFactor())); } 
-	
 	TRealTime ToRealTimeSeconds() const { return (Time * TickSettings::GetSimToRealFactor()); }
+
+	// Direct casts to "real time MS" which should be rarely used in practice (TRealTimeAccumulator only current case). All other cases of "real time" imply seconds.
+	static inline TNetworkSimTime FromRealTimeMS(const TRealTime& InRealTime) { return TNetworkSimTime( static_cast<TTime>(InRealTime)); } 
+	TRealTime ToRealTimeMS() const { return static_cast<TRealTime>(Time); }
+
 	FString ToString() const { return LexToString(this->Time); }
 
 	bool IsPositive() const { return (Time > 0); }
@@ -245,11 +252,25 @@ struct TRealTimeAccumulator
 	using TRealTime = typename TickSettings::TRealTime;
 	void Accumulate(TNetworkSimTime<TickSettings>& NetworkSimTime, const TRealTime RealTimeSeconds)
 	{
-		// Non fixed step: just accumulate the time directly
-		// Note that MaxStepMS enforcement does NOT belong here. Dropping time due to MaxStepMS would just make the sim run slower. MaxStepMS is used at the input processing level.
-		// (wondering: does converting down to MSec make sense in variable rates? We aren't really gaining anything here. Specializing to avoid the multiply probably not worth the complexity though)
-		NetworkSimTime += TNetworkSimTime<TickSettings>::FromRealTimeSeconds(RealTimeSeconds);
+		// Even though we are variable tick, we still want to truncate down to an even msec. This keeps sim steps as whole integer values that serialize better, don't have denormals or other floating point weirdness.
+		// (If we ever wanted to just fully let floats pass through and be used as the msec sim time, this could be done through another specialization for float/float time).		
+		// Also note that MaxStepMS enforcement does NOT belong here. Dropping time due to MaxStepMS would just make the sim run slower. MaxStepMS is used at the input processing level.
+		
+		AccumulatedTimeMS += RealTimeSeconds * TickSettings::GetRealToSimFactor(); // convert input seconds -> msec
+		const TNetworkSimTime<TickSettings> AccumulatedSimTimeMS = TNetworkSimTime<TickSettings>::FromRealTimeMS(AccumulatedTimeMS);	// truncate (float) accumulated msec to (int32) sim time msec
+
+		NetworkSimTime += AccumulatedSimTimeMS;
+		AccumulatedTimeMS -= AccumulatedSimTimeMS.ToRealTimeMS(); // subtract out the "whole" msec, we are left with the remainder msec
 	}
+
+	void Reset()
+	{
+		AccumulatedTimeMS = 0.f;
+	}
+
+private:
+
+	TRealTime AccumulatedTimeMS = 0.f;
 };
 
 // Specialized version of FixedTicking. This accumulates real time that spills over into NetworkSimTime as it crosses the FixStep threshold
@@ -275,6 +296,14 @@ struct TRealTimeAccumulator<TickSettings, true>
 			NetworkSimTime += TNetworkSimTime<TickSettings>::FromMSec(NumFrames * TickSettings::FixedStepMS);
 		}
 	}
+
+	void Reset()
+	{
+		AccumulatedTime = 0.f;
+	}
+
+private:
+
 	TRealTime AccumulatedTime;
 };
 
@@ -291,8 +320,39 @@ struct TSimulationTickState
 	int32 LastProcessedInputKeyframe;	// The last input keyframe that we processed
 	int32 MaxAllowedInputKeyframe;		// The max input keyframe that we are allowed to process (e.g, don't process input past this keyframe yet)
 
-	TNetworkSimTime<TSettings> TotalAllowedSimulationTime;	// Total time we have been "given" to process. We cannot process more simulation time than this: doing so would be speed hacking.
-	TNetworkSimTime<TSettings> TotalProcessedSimulationTime;	// How much time we've actually processed. The only way to increment this is to process user commands or receive authoritative state from the network.
+	void SetTotalAllowedSimulationTime(const TNetworkSimTime<TSettings>& SimTime)
+	{
+		TotalAllowedSimulationTime = SimTime;
+		RealtimeAccumulator.Reset();
+	}
+
+	TNetworkSimTime<TSettings> GetTotalProcessedSimulationTime() const 
+	{ 
+		return TotalProcessedSimulationTime; 
+	}
+
+	void SetTotalProcessedSimulationTime(const TNetworkSimTime<TSettings>& SimTime, int32 Keyframe)
+	{
+		TotalProcessedSimulationTime = SimTime;
+		SimulationTimeBuffer.ResetNextHeadKeyframe(Keyframe);
+		*SimulationTimeBuffer.GetWriteNext() = SimTime;
+	}
+
+	void IncrementTotalProcessedSimulationTime(const TNetworkSimTime<TSettings>& DeltaSimTime, int32 Keyframe)
+	{
+		TotalProcessedSimulationTime += DeltaSimTime;
+		*SimulationTimeBuffer.GetWriteNext() = TotalProcessedSimulationTime;
+		ensure(SimulationTimeBuffer.GetHeadKeyframe() == Keyframe);
+	}
+
+	void InitSimulationTimeBuffer(int32 Size)
+	{
+		SimulationTimeBuffer.SetBufferSize(Size);
+		//*SimulationTimeBuffer.GetWriteNext() = TNetworkSimTime<TSettings>::FromMSec(0);
+	}
+		
+	// Historic tracking of simulation time. This allows us to timestamp sync data as its produced
+	TReplicationBuffer<TNetworkSimTime<TSettings>>	SimulationTimeBuffer;
 
 	// "Grants" allowed simulation time to this tick state. That is, we are now allowed to advance the simulation by this amount the next time the sim ticks.
 	// Note the input is RealTime in SECONDS. This is what the rest of the engine uses when dealing with float delta time.
@@ -307,8 +367,16 @@ struct TSimulationTickState
 		return TotalAllowedSimulationTime - TotalProcessedSimulationTime;
 	}
 
+	TNetworkSimTime<TSettings> GetTotalAllowedSimulationTime() const
+	{
+		return TotalAllowedSimulationTime;
+	}
+
 private:
 
+	TNetworkSimTime<TSettings> TotalAllowedSimulationTime;	// Total time we have been "given" to process. We cannot process more simulation time than this: doing so would be speed hacking.
+	TNetworkSimTime<TSettings> TotalProcessedSimulationTime;	// How much time we've actually processed. The only way to increment this is to process user commands or receive authoritative state from the network.
+
 	TRealTimeAccumulator<TSettings>	RealtimeAccumulator;
 };
 
diff --git a/Engine/Plugins/Runtime/NetworkPrediction/readme.txt b/Engine/Plugins/Runtime/NetworkPrediction/readme.txt
--- a/Engine/Plugins/Runtime/NetworkPrediction/readme.txt
+++ b/Engine/Plugins/Runtime/NetworkPrediction/readme.txt
@@ -65,7 +65,7 @@ TODO: Major missing elements
 -No optimizations (bandwidth, cpu) have been done.
 -[Movement specific]: must lock down SetActorLocation/Rotation API so that movement state can be changed out from underneath the system.
  
-Short term road map:
+High level focus:
 -The focus right now is on the generalized prediction system, rather than actual movement code.
 -We are exploring "forward predicting non autonomous proxies" with the parametric mover. Unclear how deep we will go with this.
 -Then Aux buffer and "events"
@@ -75,6 +75,14 @@ Short term road map:
 -E.g, We don't intend all users of the engine to write their own network simulation model movement system.
 -We hope to provide a generalized movement system that is easy to extend in itself without knowing all the details of the NetworkPrediction plugin.
 -We intend to support RootMotion (Animation) / RootMotionSources (non anim) in the new movement system.
+
+Road Map:
+-Interpolation/smoothing layer
+-Dependent simulations
+-Aux Buffer
+-Event System
+-Ability System integration
+-New Movement System
  
  
 // ----------------------------------------------------------------------------------------------------------
