Fix static analysis failure
#rnx

diff --git a/Engine/Plugins/Runtime/NetworkPrediction/Source/NetworkPrediction/Private/NetworkPredictionTests.cpp b/Engine/Plugins/Runtime/NetworkPrediction/Source/NetworkPrediction/Private/NetworkPredictionTests.cpp
--- a/Engine/Plugins/Runtime/NetworkPrediction/Source/NetworkPrediction/Private/NetworkPredictionTests.cpp
+++ b/Engine/Plugins/Runtime/NetworkPrediction/Source/NetworkPrediction/Private/NetworkPredictionTests.cpp
@@ -6,8 +6,10 @@
 //	Testing Fixed vs Variable tick rate struct declration
 // ------------------------------------------------------------------------------------------------------------
 
-struct FDummy { };
-using FDummyBufferTypes = TNetworkSimBufferTypes<FDummy, FDummy, FDummy>;
+
+
+struct FNetSimDummy { };
+using FDummyBufferTypes = TNetworkSimBufferTypes<FNetSimDummy, FNetSimDummy, FNetSimDummy>;
 
 template<typename TTickSettings>
 void TickSettingsTest(TTickSettings& Settings)
