Continued work on network interpolator. Handles fluctuating conditions a bit better but still not perfecct. Better debugging.

#rnx

diff --git a/Engine/Plugins/Runtime/NetworkPrediction/Source/NetworkPrediction/Private/Movement/NetworkInterpolation.cpp b/Engine/Plugins/Runtime/NetworkPrediction/Source/NetworkPrediction/Private/Movement/NetworkInterpolation.cpp
--- a/Engine/Plugins/Runtime/NetworkPrediction/Source/NetworkPrediction/Private/Movement/NetworkInterpolation.cpp
+++ b/Engine/Plugins/Runtime/NetworkPrediction/Source/NetworkPrediction/Private/Movement/NetworkInterpolation.cpp
@@ -1,2 +1,3 @@
 // Copyright 1998-2019 Epic Games, Inc. All Rights Reserved.
-
+#include "NetworkInterpolation.h"
+#include "HAL/IConsoleManager.h"
diff --git a/Engine/Plugins/Runtime/NetworkPrediction/Source/NetworkPrediction/Private/Movement/ParametricMovement.cpp b/Engine/Plugins/Runtime/NetworkPrediction/Source/NetworkPrediction/Private/Movement/ParametricMovement.cpp
--- a/Engine/Plugins/Runtime/NetworkPrediction/Source/NetworkPrediction/Private/Movement/ParametricMovement.cpp
+++ b/Engine/Plugins/Runtime/NetworkPrediction/Source/NetworkPrediction/Private/Movement/ParametricMovement.cpp
@@ -20,6 +20,11 @@ static float DrawDebugDefaultLifeTime = 30.f;
 static FAutoConsoleVariableRef CVarDrawDebugDefaultLifeTime(TEXT("parametricmover.Debug.UseDrawDebug.DefaultLifeTime"),
 	DrawDebugDefaultLifeTime, TEXT("Use built in DrawDebug* functions for visual logging"), ECVF_Default);
 
+
+static int32 SimulatedProxyBufferSize = 4;
+static FAutoConsoleVariableRef CVarSimulatedProxyBufferSize(TEXT("parametricmover.SimulatedProxyBufferSize"),
+	DrawDebugDefaultLifeTime, TEXT(""), ECVF_Default);
+
 }
 
 // -------------------------------------------------------------------------------------------------------------------------------
@@ -121,7 +126,7 @@ FNetworkSimulationModelInitParameters UParametricMovementComponent::GetSimulatio
 	// These are reasonable defaults but may not be right for everyone
 	FNetworkSimulationModelInitParameters InitParams;
 	InitParams.InputBufferSize = 32; //Role != ROLE_SimulatedProxy ? 32 : 32;  // Fixme.. not good
-	InitParams.SyncedBufferSize = Role != ROLE_AutonomousProxy ? 4 : 32;
+	InitParams.SyncedBufferSize = Role != ROLE_AutonomousProxy ? ParametricMoverCVars::SimulatedProxyBufferSize : 32;
 	InitParams.AuxBufferSize = Role != ROLE_AutonomousProxy ? 2 : 32;
 	InitParams.DebugBufferSize = 32;
 	InitParams.HistoricBufferSize = 128;
@@ -204,7 +209,7 @@ void UParametricMovementComponent::TickComponent(float DeltaTime, enum ELevelTic
 
 		if (NetworkInterpolator.IsValid() && GetOwnerRole() == ROLE_SimulatedProxy)
 		{
-			NetworkInterpolator->Tick(DeltaTime);
+			NetworkInterpolator->Tick(DeltaTime, GetOwner());
 		}
 	}
 }
diff --git a/Engine/Plugins/Runtime/NetworkPrediction/Source/NetworkPrediction/Private/NetworkPredictionTypes.cpp b/Engine/Plugins/Runtime/NetworkPrediction/Source/NetworkPrediction/Private/NetworkPredictionTypes.cpp
--- a/Engine/Plugins/Runtime/NetworkPrediction/Source/NetworkPrediction/Private/NetworkPredictionTypes.cpp
+++ b/Engine/Plugins/Runtime/NetworkPrediction/Source/NetworkPrediction/Private/NetworkPredictionTypes.cpp
@@ -14,6 +14,15 @@ FColor FVisualLoggingParameters::DebugColors[(int32)EVisualLoggingContext::MAX]
 	FColor(255,81,101),	// OtherMispredicted
 	FColor(255,91,91),	// LastMispredicted
 	FColor::Purple,		// CurrentServerPIE
+
+	FColor(189, 195, 199),	// InterpolationBufferHead
+	FColor(94, 96, 99),		// InterpolationBufferTail
+	FColor(150, 150, 150),	// InterpolationFrom
+	FColor(200, 200, 200),	// InterpolationTo
+
+	FColor(0, 255, 0),	// InterpolationLatest
+	FColor(255, 0, 0),	// InterpolationWaiting
+	FColor(0, 0, 255),	// InterpolationSpeedUp
 };
 
 // -------------------------------------------------------------------------------------------------------------------------------
diff --git a/Engine/Plugins/Runtime/NetworkPrediction/Source/NetworkPrediction/Public/NetworkInterpolation.h b/Engine/Plugins/Runtime/NetworkPrediction/Source/NetworkPrediction/Public/NetworkInterpolation.h
--- a/Engine/Plugins/Runtime/NetworkPrediction/Source/NetworkPrediction/Public/NetworkInterpolation.h
+++ b/Engine/Plugins/Runtime/NetworkPrediction/Source/NetworkPrediction/Public/NetworkInterpolation.h
@@ -2,20 +2,43 @@
 
 #pragma once
 #include "NetworkSimulationModelCVars.h"
+#include "VisualLogger/VisualLogger.h"
+
+DEFINE_LOG_CATEGORY_STATIC(LogNetInterpolation, Log, All);
 
 namespace NetworkInterpolationDebugCVars
 {
 	NETSIM_DEVCVAR_SHIPCONST_INT(Disable, 0, "ni.Disable", "Disables Network Interpolation");
+	NETSIM_DEVCVAR_SHIPCONST_INT(VLog, 1, "ni.VLog", "Enables Network Interpolation VLog ");
+
+	NETSIM_DEVCVAR_SHIPCONST_FLOAT(WaitSlack, 0.05, "ni.WaitSlack", "How much slack to wait for when waiting");
+
+	NETSIM_DEVCVAR_SHIPCONST_FLOAT(CatchUpThreshold, 0.200, "ni.CatchUpThreshold", "When we start catching up (seconds from head)");
+	NETSIM_DEVCVAR_SHIPCONST_FLOAT(CatchUpGoal, 0.100, "ni.CatchUpGoal", "When we stop cathcing up (seconds from head)");
+	NETSIM_DEVCVAR_SHIPCONST_FLOAT(CatchUpFactor, 1.25, "ni.CatchUpFactor", "Factor we use to catch up");
 }
 
+
 class INetworkInterpolator
 {
 public:
 
 	virtual ~INetworkInterpolator() { }
-	virtual void Tick(float DeltaSeconds) = 0;
+	virtual void Tick(float DeltaSeconds, AActor* LogOwner) = 0;
 };
 
+// ----------------------------------------------------------------------------------------------------------------------------------------------------------------------
+//	TNetworkSimulationModelInterpolator
+//	The interpolator is responsible for smoothly interpolating between known Sync States. The output of the interpolator does NOT feed back into the simulation!
+//	This is still a WIP and not where we want it to be. Some notes about future improvements:
+//	-We should take a data oriented approach and have an interpolation manager than can more efficiently bulk process interpolations
+//	-We should have interpolation settings instead of cvars and in general improve the dynamic-ness of the interpolation algorithm.
+//	-The tricky part here is that we cannot rely on consistent replication frequency and frame rates. We want the interpolator to handle any network conditions and scale accordingly.
+//	-E.g, smooth motion, but minimize added application latency.
+//
+//	Also note: the intention is not for every actor using NetworkPrediction to necessarily be interpolated. Simulated Extrapolation is still a fine option.
+//
+// ----------------------------------------------------------------------------------------------------------------------------------------------------------------------
 template<typename TNetworkedSimulationModel, typename TDriver>
 class TNetworkSimulationModelInterpolator : public INetworkInterpolator
 {
@@ -25,6 +48,8 @@ public:
 	using TRealTime = typename TNetworkedSimulationModel::TRealTime;
 	using TSyncState = typename TNetworkedSimulationModel::TSyncState;
 
+	bool bEnableVisualLog = true;
+
 	TNetworkSimulationModelInterpolator(TNetworkedSimulationModel* Source, TDriver* InDriver)
 	{
 		NetworkSim = Source;
@@ -32,11 +57,28 @@ public:
 	}
 	virtual ~TNetworkSimulationModelInterpolator()
 	{
+		if (NetworkInterpolationDebugCVars::Disable() > 0)
+		{			
+			if (const TSyncState* HeadState = NetworkSim->Buffers.Sync.GetElementFromHead(0))
+			{
+				Driver->FinalizeFrame(*HeadState);
+			}
+			return;
+		}
+		
+		if (NetworkSim->TickInfo.SimulationTimeBuffer.GetNumValidElements() <= 1)
+		{
+			// Cant interpolate yet	
+			return;
+		}
+
 
 	}
 
-	void Tick(float DeltaSeconds) override
+	void Tick(float DeltaSeconds, AActor* LogOwner) override
 	{
+		const bool bDoVLog = NetworkInterpolationDebugCVars::VLog() && bEnableVisualLog;
+
 		if (NetworkInterpolationDebugCVars::Disable() > 0)
 		{			
 			if (const TSyncState* HeadState = NetworkSim->Buffers.Sync.GetElementFromHead(0))
@@ -45,128 +87,163 @@ public:
 			}
 			return;
 		}
-
+		
 		if (NetworkSim->TickInfo.SimulationTimeBuffer.GetNumValidElements() <= 1)
 		{
 			// Cant interpolate yet	
 			return;
 		}
 
-		InterpolationTime += DeltaSeconds;
 
-		// SimulatinoTimeBuffer holds the simulation time stamps of each SyncState in the SyncBuffer
-		// (Note that things are not going to be evenly spaced time wise)
-		auto& SimulationTimeBuffer = NetworkSim->TickInfo.SimulationTimeBuffer;		
-		const int32 HeadKeyframe = SimulationTimeBuffer.GetHeadKeyframe();
+		auto& SimulationTimeBuffer = NetworkSim->TickInfo.SimulationTimeBuffer;
 
-		// Find which keyframes we should be interpolating between
-		int32 InterpolationKeyframe = INDEX_NONE;
-		TSimTime FromInterpolationTime;
-		TSimTime ToInterpolationTime;
+		// Starting off: start at the tail end
+		if (InterpolationTime <= 0.f)
+		{
+			InterpolationTime = SimulationTimeBuffer.GetElementFromTail(0)->ToRealTimeSeconds();
+			InterpolationKeyframe = SimulationTimeBuffer.GetTailKeyframe();
+			InterpolationState = *NetworkSim->Buffers.Sync.GetElementFromTail(0);
+		}
 
-		for (int32 Keyframe = SimulationTimeBuffer.GetTailKeyframe(); Keyframe < HeadKeyframe; ++Keyframe)
+		// Wait if we were too far ahead
+		if (WaitUntilTime > 0.f)
 		{
-			TSimTime ElementSimTime = *SimulationTimeBuffer.FindElementByKeyframe(Keyframe);
-			if (InterpolationTime > ElementSimTime.ToRealTimeSeconds())
+			if (WaitUntilTime <= SimulationTimeBuffer.GetElementFromHead(0)->ToRealTimeSeconds())
 			{
-				InterpolationKeyframe = Keyframe;
-				FromInterpolationTime = ElementSimTime;
+				// done waiting, we can continue
+				WaitUntilTime = 0.f;
+				UE_VLOG(LogOwner, LogNetInterpolation, Log, TEXT("Done Waiting! Head: %s"), *LexToString(SimulationTimeBuffer.GetElementFromHead(0)->ToRealTimeSeconds()));
 			}
 			else
 			{
-				break;
+				if (bDoVLog)
+				{
+					// Still waiting, return
+					UE_VLOG(LogOwner, LogNetInterpolation, Log, TEXT("Still Waiting! %s < %s"), *LexToString(WaitUntilTime), *LexToString(SimulationTimeBuffer.GetElementFromHead(0)->ToRealTimeSeconds()));
+
+					const FName LocalInterpolationTimeName("Local Interpolation Time");
+					FVector2D LocalTimeVsInterpolationTime(LogOwner->GetWorld()->GetTimeSeconds(), (int32)(InterpolationTime * 1000.f));
+					UE_VLOG_HISTOGRAM(LogOwner, LogNetInterpolation, Log, "ServerSimulationTimeGraph", LocalInterpolationTimeName, LocalTimeVsInterpolationTime);
+				}
+				return;
 			}
 		}
 
-		auto ResetInterpolationTimeToMidpoint = [&](int32 Keyframe)
-		{
-			check(SimulationTimeBuffer.IsValidKeyframe(Keyframe));
-			check(SimulationTimeBuffer.IsValidKeyframe(Keyframe+1));
+		EVisualLoggingContext LoggingContext = EVisualLoggingContext::InterpolationLatest;
 
-			InterpolationKeyframe = Keyframe;
+		// Calc new interpolation time
+		TRealTime NewInterpolationTime = InterpolationTime;
+		{
+			TRealTime Step = DeltaSeconds;
 			
-			FromInterpolationTime = *SimulationTimeBuffer.FindElementByKeyframe(Keyframe);
-			ToInterpolationTime = *SimulationTimeBuffer.FindElementByKeyframe(Keyframe+1);
-			InterpolationTime = (FromInterpolationTime.ToRealTimeSeconds() + ToInterpolationTime.ToRealTimeSeconds()) / 2.f;
-		};
+			// Speed up if we are too far behind
+			TRealTime CatchUpThreshold = SimulationTimeBuffer.GetElementFromHead(0)->ToRealTimeSeconds() - NetworkInterpolationDebugCVars::CatchUpThreshold();
+			if (CatchUpUntilTime <= 0.f && InterpolationTime < CatchUpThreshold)
+			{
+				CatchUpUntilTime = SimulationTimeBuffer.GetElementFromHead(0)->ToRealTimeSeconds() - NetworkInterpolationDebugCVars::CatchUpGoal();
+			}
 
-		if (InterpolationKeyframe == INDEX_NONE)
-		{
-			//UE_LOG(LogTemp, Warning, TEXT("FELL BEHIND. %s < %s"), *LexToString(InterpolationTime), *LexToString(SimulationTimeBuffer.GetElementFromTail(0)->ToRealTimeSeconds()));
-			//ResetInterpolationTimeToMidpoint(SimulationTimeBuffer.GetTailKeyframe());
-			ResetInterpolationTimeToMidpoint((SimulationTimeBuffer.GetTailKeyframe() + SimulationTimeBuffer.GetHeadKeyframe()) / 2);
-		}
-		else
-		{
-			ToInterpolationTime = *SimulationTimeBuffer.FindElementByKeyframe(InterpolationKeyframe + 1);
-			if (InterpolationTime > ToInterpolationTime.ToRealTimeSeconds())
+			if (CatchUpUntilTime > 0.f)
+			{
+				if (InterpolationTime  < CatchUpUntilTime)
+				{
+					Step *= NetworkInterpolationDebugCVars::CatchUpFactor();
+					LoggingContext = EVisualLoggingContext::InterpolationSpeedUp;
+
+					UE_VLOG(LogOwner, LogNetInterpolation, Log, TEXT("Catching up! %s < %s"), *LexToString(InterpolationTime), *LexToString(CatchUpUntilTime));
+					UE_LOG(LogNetInterpolation, Warning, TEXT("Catching up! %s < %s"), *LexToString(InterpolationTime), *LexToString(CatchUpUntilTime));
+				}
+				else
+				{
+					CatchUpUntilTime = 0.f;
+				}
+			}
+
+			NewInterpolationTime += Step;
+
+			// Did this put us too far ahead, and now we need to start waiting?
+			if (NewInterpolationTime > SimulationTimeBuffer.GetElementFromHead(0)->ToRealTimeSeconds())
 			{
-				//UE_LOG(LogTemp, Warning, TEXT("GOT AHEAD. %s > %s"), *LexToString(InterpolationTime), *LexToString(ToInterpolationTime.ToRealTimeSeconds()));
-				//ResetInterpolationTimeToMidpoint(SimulationTimeBuffer.GetHeadKeyframe()-1);
-				ResetInterpolationTimeToMidpoint((SimulationTimeBuffer.GetTailKeyframe() + SimulationTimeBuffer.GetHeadKeyframe()) / 2);
+				UE_VLOG(LogOwner, LogNetInterpolation, Log, TEXT("Too far ahead! Starting to wait! Head: %s"), *LexToString(SimulationTimeBuffer.GetElementFromHead(0)->ToRealTimeSeconds()));
+				WaitUntilTime = SimulationTimeBuffer.GetElementFromHead(0)->ToRealTimeSeconds() + NetworkInterpolationDebugCVars::WaitSlack();
+				NewInterpolationTime = SimulationTimeBuffer.GetElementFromHead(0)->ToRealTimeSeconds();
 			}
 		}
 
+		// Find "To" keyframe
+		const TSyncState* ToState = nullptr;
+		TRealTime ToTime = 0.f;
 
-		const TSyncState* FromState = NetworkSim->Buffers.Sync.FindElementByKeyframe(InterpolationKeyframe);
-		const TSyncState* ToState = NetworkSim->Buffers.Sync.FindElementByKeyframe(InterpolationKeyframe+1);
+		for (int32 Keyframe = SimulationTimeBuffer.GetTailKeyframe(); Keyframe <= SimulationTimeBuffer.GetHeadKeyframe(); ++Keyframe)
+		{
+			TSimTime ElementSimTime = *SimulationTimeBuffer.FindElementByKeyframe(Keyframe);
+			if (NewInterpolationTime <= ElementSimTime.ToRealTimeSeconds())
+			{
+				InterpolationKeyframe = Keyframe;
+				ToTime = ElementSimTime.ToRealTimeSeconds();
+				ToState = NetworkSim->Buffers.Sync.FindElementByKeyframe(Keyframe);
+				break;
+			}
+		}
 
-		if (FromState && ToState)
+		if (ensure(ToState))
 		{
-			const TRealTime FromRealTime = FromInterpolationTime.ToRealTimeSeconds();
-			const TRealTime ToRealTime = ToInterpolationTime.ToRealTimeSeconds();
+			const TRealTime FromRealTime = InterpolationTime;
+			const TRealTime ToRealTime = ToTime;
 			const TRealTime InterpolationInterval = ToRealTime - FromRealTime;
 		
 			if (ensure(FMath::Abs(InterpolationInterval) > 0.f))
 			{
-				const float InterpolationPCT = (InterpolationTime - FromRealTime) / InterpolationInterval;
-				ensureMsgf(InterpolationPCT >= 0.f && InterpolationPCT <= 1.f, TEXT("Calculated InterpolationPCT not in expected range. InterpolationTime: %s. From: %s. To: %s"), *LexToString(InterpolationTime), *LexToString(FromRealTime), *LexToString(ToRealTime));
+				const float InterpolationPCT = (NewInterpolationTime - FromRealTime) / InterpolationInterval;
+				ensureMsgf(InterpolationPCT >= 0.f && InterpolationPCT <= 1.f, TEXT("Calculated InterpolationPCT not in expected range. NewInterpolationTime: %s. From: %s. To: %s"), *LexToString(NewInterpolationTime), *LexToString(FromRealTime), *LexToString(ToRealTime));
 
-				TSyncState InterpolatedState;
-				TSyncState::Interpolate(*FromState, *ToState, InterpolationPCT, InterpolatedState);
+				TSyncState NewInterpolatedState;
+				TSyncState::Interpolate(InterpolationState, *ToState, InterpolationPCT, NewInterpolatedState);
 
-				//UE_LOG(LogTemp, Warning, TEXT("Interp. [%d] %s"), ToInterpolationData->Keyframe,  *LexToString(InterpolationTime));
+				Driver->FinalizeFrame(NewInterpolatedState);
+				
+				if (bDoVLog)
+				{
+					UE_VLOG(LogOwner, LogNetInterpolation, Log, TEXT("%s - %s - %s.  InterpolationPCT: %f"), *LexToString(FromRealTime), *LexToString(NewInterpolationTime), *LexToString(ToRealTime), InterpolationPCT);
 
-				Driver->FinalizeFrame(InterpolatedState);
-			}
-		}
-		else
-		{
-			ensure(false);
-		}
+					// Graph Interpolation Time vs Buffer Head/Tail times
+					const FName ServerSimulationGraphName("ServerSimulationTimeGraph");
+					const FName ServerSimTimeName("Server Simulation Time");
+					FVector2D LocalTimeVsServerSimTime(LogOwner->GetWorld()->GetTimeSeconds(), (int32)NetworkSim->TickInfo.SimulationTimeBuffer.GetElementFromHead(0)->ToRealTimeMS());
+					UE_VLOG_HISTOGRAM(LogOwner, LogNetInterpolation, Log, ServerSimulationGraphName, ServerSimTimeName, LocalTimeVsServerSimTime);
 
-		//UE_LOG(LogTemp, Warning, TEXT("NumElements: [%d, %d]"), NetworkSim->Buffers.Sync.GetNumValidElements(), SimulationTimeBuffer.GetNumValidElements());
+					const FName BufferTailSimTimeName("Buffer Tail Simulation Time");
+					FVector2D LocalTimeVsBufferTailSim(LogOwner->GetWorld()->GetTimeSeconds(), (int32)NetworkSim->TickInfo.SimulationTimeBuffer.GetElementFromTail(0)->ToRealTimeMS());
+					UE_VLOG_HISTOGRAM(LogOwner, LogNetInterpolation, Log, ServerSimulationGraphName, BufferTailSimTimeName, LocalTimeVsBufferTailSim);
 
-		/*
-		InterpolationTime += DeltaSeconds;
+					const FName LocalInterpolationTimeName("Local Interpolation Time");
+					FVector2D LocalTimeVsInterpolationTime(LogOwner->GetWorld()->GetTimeSeconds(), (int32)(NewInterpolationTime * 1000.f));
+					UE_VLOG_HISTOGRAM(LogOwner, LogNetInterpolation, Log, ServerSimulationGraphName, LocalInterpolationTimeName, LocalTimeVsInterpolationTime);
+					
+					FVector2D LocalTimeVsCatchUpThreshold(LogOwner->GetWorld()->GetTimeSeconds(), (SimulationTimeBuffer.GetElementFromHead(0)->ToRealTimeSeconds() - NetworkInterpolationDebugCVars::CatchUpThreshold()) * 1000.f);
+					UE_VLOG_HISTOGRAM(LogOwner, LogNetInterpolation, Log, ServerSimulationGraphName, "Catch Up Threshold", LocalTimeVsCatchUpThreshold);
 
-		static TRealTime LastInterpolationTime = 0.f;		
-		static TRealTime LastSimTime = 0.f;
-		static TSimTime LastSimTimeMSec;
+					FVector2D LocalTimeVsCatchUpGoal(LogOwner->GetWorld()->GetTimeSeconds(), (SimulationTimeBuffer.GetElementFromHead(0)->ToRealTimeSeconds() - NetworkInterpolationDebugCVars::CatchUpGoal()) * 1000.f);
+					UE_VLOG_HISTOGRAM(LogOwner, LogNetInterpolation, Log, ServerSimulationGraphName, "Catch Up Goal", LocalTimeVsCatchUpGoal);
 
-		TRealTime DeltaSim = NetworkSim->TickInfo.TotalProcessedSimulationTime.ToRealTimeSeconds() - LastSimTime;
-		TRealTime DeltaInt = InterpolationTime - LastInterpolationTime;
-		TSimTime DeltaSimMSec = NetworkSim->TickInfo.TotalProcessedSimulationTime - LastSimTimeMSec;
-		
+					// VLog the actual motion states
+					const TSyncState* DebugTail = NetworkSim->Buffers.Sync.GetElementFromTail(0);
+					const TSyncState* DebugHead = NetworkSim->Buffers.Sync.GetElementFromHead(0);
 
-		TRealTime Delta = NetworkSim->TickInfo.TotalProcessedSimulationTime.ToRealTimeSeconds() - InterpolationTime;
-		UE_LOG(LogTemp, Warning, TEXT("%s (%s)... %s (%s) - %s (%s). Delta: %s"), *NetworkSim->TickInfo.TotalProcessedSimulationTime.ToString(), *DeltaSimMSec.ToString(),
-			*LexToString(NetworkSim->TickInfo.TotalProcessedSimulationTime.ToRealTimeSeconds()), *LexToString(DeltaSim), *LexToString(InterpolationTime), *LexToString(DeltaInt), *LexToString(Delta));
+					DebugTail->VisualLog( FVisualLoggingParameters(EVisualLoggingContext::InterpolationBufferTail, NetworkSim->Buffers.Sync.GetTailKeyframe(), EVisualLoggingLifetime::Transient), Driver, Driver );
+					DebugHead->VisualLog( FVisualLoggingParameters(EVisualLoggingContext::InterpolationBufferHead, NetworkSim->Buffers.Sync.GetHeadKeyframe(), EVisualLoggingLifetime::Transient), Driver, Driver );
 
-		if (InterpolationTime > NetworkSim->TickInfo.TotalProcessedSimulationTime.ToRealTimeSeconds())
-		{
-			InterpolationTime = NetworkSim->TickInfo.TotalProcessedSimulationTime.ToRealTimeSeconds();
-			UE_LOG(LogTemp, Warning, TEXT("CLAMP! %s"), *LexToString(InterpolationTime));
-		}
+					InterpolationState.VisualLog( FVisualLoggingParameters(EVisualLoggingContext::InterpolationFrom, InterpolationKeyframe, EVisualLoggingLifetime::Transient), Driver, Driver );
+					ToState->VisualLog( FVisualLoggingParameters(EVisualLoggingContext::InterpolationTo, InterpolationKeyframe, EVisualLoggingLifetime::Transient), Driver, Driver );
 
-		LastInterpolationTime = InterpolationTime;
-		LastSimTime = NetworkSim->TickInfo.TotalProcessedSimulationTime.ToRealTimeSeconds();
-		LastSimTimeMSec = NetworkSim->TickInfo.TotalProcessedSimulationTime;
+					NewInterpolatedState.VisualLog( FVisualLoggingParameters(LoggingContext, InterpolationKeyframe, EVisualLoggingLifetime::Transient), Driver, Driver );
+				}
 
-		if (true)
-			return;
 
-		*/
+				InterpolationState = NewInterpolatedState;
+				InterpolationTime = NewInterpolationTime;
+			}
+		}
 	}
 
 private:
@@ -175,6 +252,17 @@ private:
 	TDriver* Driver;
 	
 	TRealTime InterpolationTime = 0.f; // SimTime we are currently interpolating at
+	int32 InterpolationKeyframe = INDEX_NONE; // Keyframe we are currently/last interpolated at
+
+	TSyncState InterpolationState;
+
+	TRealTime WaitUntilTime = 0.f;	
+
+	TRealTime CatchUpUntilTime = 0.f;
+
+
+	TRealTime DynamicBufferedTime = 1/60.f; // SimTime we are currently interpolating at
+	TRealTime DynamicBufferedTimeStep = 1/60.f;
 
 	TRealTime	MinBufferedTime = 1/120.f;
 	TRealTime	MaxBufferedTime = 1.f;
diff --git a/Engine/Plugins/Runtime/NetworkPrediction/Source/NetworkPrediction/Public/NetworkPredictionTypes.h b/Engine/Plugins/Runtime/NetworkPrediction/Source/NetworkPrediction/Public/NetworkPredictionTypes.h
--- a/Engine/Plugins/Runtime/NetworkPrediction/Source/NetworkPrediction/Public/NetworkPredictionTypes.h
+++ b/Engine/Plugins/Runtime/NetworkPrediction/Source/NetworkPrediction/Public/NetworkPredictionTypes.h
@@ -91,6 +91,9 @@ struct FStandardLoggingParameters
 UENUM()
 enum class EVisualLoggingContext: uint8
 {
+	// NOTE: Keep FVisualLoggingParameters::DebugColors in sync with this
+
+	// (Contexts used in core network sim)
 	LastPredicted,		// The last predicted state. Where the character "is now"
 	OtherPredicted,		// "Middle" states between LastComfirmed and LastPredicted. Recommend to draw small crumbs (not full model/cylinder/etc) to avoid washing out the scene.
 	LastConfirmed,		// The last confirmed state from the server
@@ -98,6 +101,17 @@ enum class EVisualLoggingContext: uint8
 	OtherMispredicted,	// subsequent mispredicted states on the client. Recommend to draw small crumbs (not full model/cylinder/etc) to avoid washing out the scene.
 	LastMispredicted,	// The latest predictive state when a misprediction was detected.
 	CurrentServerPIE,	// The current server position *right now*. Only available in PIE
+
+	// (Contexts used by interpolation)
+	InterpolationBufferHead,// Head end of sync buffer while interpolating
+	InterpolationBufferTail,// Tail end of sync buffer while interpolating
+	InterpolationFrom,		// Immediate "from" interpolation keyframe
+	InterpolationTo,		// Immediate "to" interpolation keyframe
+	
+	InterpolationLatest,	// Latest interpolation position in normal circumstances
+	InterpolationWaiting,	// Latest interpolation while waiting (overrun)
+	InterpolationSpeedUp,	// Latest interpolation while speeding up (underrun)
+
 	MAX
 };
 
diff --git a/Engine/Plugins/Runtime/NetworkPrediction/Source/NetworkPrediction/Public/NetworkSimulationModelCVars.h b/Engine/Plugins/Runtime/NetworkPrediction/Source/NetworkPrediction/Public/NetworkSimulationModelCVars.h
--- a/Engine/Plugins/Runtime/NetworkPrediction/Source/NetworkPrediction/Public/NetworkSimulationModelCVars.h
+++ b/Engine/Plugins/Runtime/NetworkPrediction/Source/NetworkPrediction/Public/NetworkSimulationModelCVars.h
@@ -1,6 +1,7 @@
 // Copyright 1998-2019 Epic Games, Inc. All Rights Reserved.
 
 #pragma once
+#include "HAL/IConsoleManager.h"
 
 // ------------------------------------------------------------------------------------------------------------
 //	"Shipping const" cvars: cvars that should compile out to const functions in shipping/test builds
