non unity compile fixes
#rnx

diff --git a/Engine/Plugins/Runtime/NetworkPrediction/Source/NetworkPrediction/Public/NetworkInterpolation.h b/Engine/Plugins/Runtime/NetworkPrediction/Source/NetworkPrediction/Public/NetworkInterpolation.h
--- a/Engine/Plugins/Runtime/NetworkPrediction/Source/NetworkPrediction/Public/NetworkInterpolation.h
+++ b/Engine/Plugins/Runtime/NetworkPrediction/Source/NetworkPrediction/Public/NetworkInterpolation.h
@@ -3,6 +3,8 @@
 #pragma once
 #include "NetworkSimulationModelCVars.h"
 #include "VisualLogger/VisualLogger.h"
+#include "GameFramework/Actor.h"
+#include "NetworkPredictionTypes.h"
 
 DEFINE_LOG_CATEGORY_STATIC(LogNetInterpolation, Log, All);
 
