One more missed file for non unity
#rnx

diff --git a/Engine/Plugins/Runtime/NetworkPrediction/Source/NetworkPrediction/Public/NetworkInterpolation.h b/Engine/Plugins/Runtime/NetworkPrediction/Source/NetworkPrediction/Public/NetworkInterpolation.h
--- a/Engine/Plugins/Runtime/NetworkPrediction/Source/NetworkPrediction/Public/NetworkInterpolation.h
+++ b/Engine/Plugins/Runtime/NetworkPrediction/Source/NetworkPrediction/Public/NetworkInterpolation.h
@@ -5,6 +5,7 @@
 #include "VisualLogger/VisualLogger.h"
 #include "GameFramework/Actor.h"
 #include "NetworkPredictionTypes.h"
+#include "Engine/World.h"
 
 DEFINE_LOG_CATEGORY_STATIC(LogNetInterpolation, Log, All);
 
