Fix installed build compilation failure - NetworkSimulationModel.h, optional argument was on the wrong function in the (NETSIM_MODEL_DEBUG == 0) case.
#rnx

diff --git a/Engine/Plugins/Runtime/NetworkPrediction/Source/NetworkPrediction/Public/NetworkSimulationModel.h b/Engine/Plugins/Runtime/NetworkPrediction/Source/NetworkPrediction/Public/NetworkSimulationModel.h
--- a/Engine/Plugins/Runtime/NetworkPrediction/Source/NetworkPrediction/Public/NetworkSimulationModel.h
+++ b/Engine/Plugins/Runtime/NetworkPrediction/Source/NetworkPrediction/Public/NetworkSimulationModel.h
@@ -389,9 +389,9 @@ public:
 
 	TReplicationBuffer<TDebugState>* GetRemoteDebugBuffer() {	return &RepProxy_Debug.ReceivedBuffer; }
 #else
-	TReplicationBuffer<TDebugState>* GetLocalDebugBuffer(bool bCreate=false) {	return nullptr; }
+	TReplicationBuffer<TDebugState>* GetLocalDebugBuffer() {	return nullptr; }
 	TDebugState* GetNextLocalDebugStateWrite() { return nullptr; }
-	TNetworkSimBufferContainer<TBufferTypes>* GetHistoricBuffers() { return nullptr; }
+	TNetworkSimBufferContainer<TBufferTypes>* GetHistoricBuffers(bool bCreate=false) { return nullptr; }
 	TReplicationBuffer<TDebugState>* GetRemoteDebugBuffer() {	return nullptr; }
 #endif
 
