NetworkPrediction plugin
-More refactors on general system architecture. UNetworkSimulationGlobalManager is solidifying and most boiler plate is out of the actor components (Not an exciting update but this was overdue!).
-PostSimTick added and Interpolator refactored to be part of this step, dependent on network role (Eg., simulated proxy only when enabled).
-Continuing to investigate dependent actor simulations. A few early prototypes did not work out but am going to try a slightly different approach.

#rnx

diff --git a/Engine/Plugins/Runtime/NetworkPrediction/Source/NetworkPrediction/Private/Movement/BaseMovementComponent.cpp b/Engine/Plugins/Runtime/NetworkPrediction/Source/NetworkPrediction/Private/Movement/BaseMovementComponent.cpp
--- a/Engine/Plugins/Runtime/NetworkPrediction/Source/NetworkPrediction/Private/Movement/BaseMovementComponent.cpp
+++ b/Engine/Plugins/Runtime/NetworkPrediction/Source/NetworkPrediction/Private/Movement/BaseMovementComponent.cpp
@@ -341,7 +341,7 @@ FTransform UBaseMovementComponent::GetUpdateComponentTransform() const
 	return FTransform::Identity;
 }
 
-UObject* UBaseMovementComponent::GetVLogOwner() const
+UObject* UBaseMovementComponent::GetBaseMovementVLogOwner() const
 {
 	return GetOwner();
 }
@@ -426,8 +426,8 @@ void UBaseMovementComponent::DrawDebug(const IBaseMovementDriver::FDrawDebugPara
 
 IBaseMovementDriver::FDrawDebugParams::FDrawDebugParams(const FVisualLoggingParameters& Parameters, IBaseMovementDriver* LogDriver)
 {
-	DebugWorld = LogDriver->GetDriverWorld();
-	DebugLogOwner = LogDriver->GetVLogOwner();
+	DebugWorld = LogDriver->GetBaseMovementDriverWorld();
+	DebugLogOwner = LogDriver->GetBaseMovementVLogOwner();
 	DrawType = (Parameters.Context == EVisualLoggingContext::OtherMispredicted || Parameters.Context == EVisualLoggingContext::OtherPredicted) ? EVisualLoggingDrawType::Crumb : EVisualLoggingDrawType::Full;
 	Lifetime = Parameters.Lifetime;
 	DrawColor = Parameters.GetDebugColor();	
diff --git a/Engine/Plugins/Runtime/NetworkPrediction/Source/NetworkPrediction/Private/Movement/FlyingMovementComponent.cpp b/Engine/Plugins/Runtime/NetworkPrediction/Source/NetworkPrediction/Private/Movement/FlyingMovementComponent.cpp
--- a/Engine/Plugins/Runtime/NetworkPrediction/Source/NetworkPrediction/Private/Movement/FlyingMovementComponent.cpp
+++ b/Engine/Plugins/Runtime/NetworkPrediction/Source/NetworkPrediction/Private/Movement/FlyingMovementComponent.cpp
@@ -62,25 +62,16 @@ UFlyingMovementComponent::UFlyingMovementComponent()
 //	Core Network Prediction functions
 // ----------------------------------------------------------------------------------------------------------
 
-IReplicationProxy* UFlyingMovementComponent::InstantiateNetworkSimulation()
+INetworkSimulationModel* UFlyingMovementComponent::InstantiateNetworkSimulation()
 {
-	NetworkSim.Reset(new FlyingMovement::FMovementSystem());
-
-#if NETSIM_MODEL_DEBUG
-	FNetworkSimulationModelDebuggerManager::Get().RegisterNetworkSimulationModel(NetworkSim.Get(), (IFlyingMovementDriver*)this, GetOwner(), TEXT("FlyingPawnMovement"));
-#endif
-	return NetworkSim.Get();
-}
-
-void UFlyingMovementComponent::InitializeForNetworkRole(ENetRole Role)
-{	
-	check(NetworkSim);
-	NetworkSim->InitializeForNetworkRole(Role, GetSimulationInitParameters(Role));
+	auto NewSim = new FlyingMovement::FMovementSystem<0>(this);
+	DO_NETSIM_MODEL_DEBUG(FNetworkSimulationModelDebuggerManager::Get().RegisterNetworkSimulationModel(NewSim, GetOwner()));
+	return NewSim;
 }
 
 void UFlyingMovementComponent::TickComponent(float DeltaTime, enum ELevelTick TickType, FActorComponentTickFunction* ThisTickFunction)
 {
-	// TEMP! Disable existing CMC if it is activate. Just makes A/B testing eaiser for now.
+	// TEMP! Disable existing CMC if it is activate. Just makes A/B testing easier for now.
 	if (AActor* Owner = GetOwner())
 	{
 		if (UCharacterMovementComponent* OldComp = Owner->FindComponentByClass<UCharacterMovementComponent>())
@@ -94,40 +85,14 @@ void UFlyingMovementComponent::TickComponent(float DeltaTime, enum ELevelTick Ti
 		Owner->bReplicateMovement = false;
 	}
 
-
-
 	Super::TickComponent(DeltaTime, TickType, ThisTickFunction);
 	const ENetRole OwnerRole = GetOwnerRole();
 
-	// -------------------------------------
-	// Tick the network sim
-	// -------------------------------------
-	if (NetworkSim)
+	// Check if we should trip a mispredict. (Note how its not possible to do this inside the Update function!)
+	if (OwnerRole == ROLE_Authority && FlyingMovementCVars::RequestMispredict)
 	{
-		PreTickSimulation(DeltaTime); // Fixme
-
-		// Check if we should trip a mispredict. (Note how its not possible to do this inside the Update function!)
-		if (OwnerRole == ROLE_Authority && FlyingMovementCVars::RequestMispredict)
-		{
-			FlyingMovement::FMovementSystem::ForceMispredict = true;
-			FlyingMovementCVars::RequestMispredict = 0;
-		}
-
-		FlyingMovement::FMovementSystem::FTickParameters Parameters(DeltaTime, GetOwner());;
-
-		// Tick the core network sim, this will consume input and generate new sync state
-		NetworkSim->Tick((IFlyingMovementDriver*)this, Parameters);
-
-		// Client->Server communication
-		if (NetworkSim->ShouldSendServerRPC(OwnerRole, DeltaTime))
-		{
-			// Temp hack to make sure the ServerRPC doesn't get suppressed from bandwidth limiting
-			// (system hasn't been optimized and not mature enough yet to handle gaps in input stream)
-			FScopedBandwidthLimitBypass BandwidthBypass(GetOwner());
-
-			FServerReplicationRPCParameter ProxyParameter(ReplicationProxy_ServerRPC);
-			ServerRecieveClientInput(ProxyParameter);
-		}
+		FlyingMovement::FMovementSimulation::ForceMispredict = true;
+		FlyingMovementCVars::RequestMispredict = 0;
 	}
 }
 
@@ -160,10 +125,10 @@ void UFlyingMovementComponent::PreSimSync(const FlyingMovement::FMoveState& Sync
 	}
 }
 
-void UFlyingMovementComponent::ProduceInput(const FlyingMovement::TSimTime& SimFrameTime, FlyingMovement::FInputCmd& Cmd)
+void UFlyingMovementComponent::ProduceInput(const float TimeSeconds, FlyingMovement::FInputCmd& Cmd)
 {
 	// This isn't ideal. It probably makes sense for the component to do all the input binding rather.
-	ProduceInputDelegate.ExecuteIfBound(SimFrameTime, Cmd);
+	ProduceInputDelegate.ExecuteIfBound(TimeSeconds, Cmd);
 }
 
 void UFlyingMovementComponent::FinalizeFrame(const FlyingMovement::FMoveState& SyncState)
@@ -174,4 +139,9 @@ void UFlyingMovementComponent::FinalizeFrame(const FlyingMovement::FMoveState& S
 FString UFlyingMovementComponent::GetDebugName() const
 {
 	return FString::Printf(TEXT("FlyingMovement. %s. %s"), *UEnum::GetValueAsString(TEXT("Engine.ENetRole"), GetOwnerRole()), *GetName());
+}
+
+const UObject* UFlyingMovementComponent::GetVLogOwner() const
+{
+	return GetOwner();
 }
\ No newline at end of file
diff --git a/Engine/Plugins/Runtime/NetworkPrediction/Source/NetworkPrediction/Private/Movement/FlyingMovementSimulation.cpp b/Engine/Plugins/Runtime/NetworkPrediction/Source/NetworkPrediction/Private/Movement/FlyingMovementSimulation.cpp
--- a/Engine/Plugins/Runtime/NetworkPrediction/Source/NetworkPrediction/Private/Movement/FlyingMovementSimulation.cpp
+++ b/Engine/Plugins/Runtime/NetworkPrediction/Source/NetworkPrediction/Private/Movement/FlyingMovementSimulation.cpp
@@ -46,7 +46,9 @@ static FAutoConsoleVariableRef CVarDrawDebugDefaultLifeTime(TEXT("fp.Debug.UseDr
 
 namespace FlyingMovement
 {
-	bool FMovementSystem::ForceMispredict = false;
+	const FName FMovementSimulation::GroupName(TEXT("FlyingMovement"));
+
+	bool FMovementSimulation::ForceMispredict = false;
 	static FVector ForceMispredictVelocityMagnitude = FVector(2000.f, 0.f, 0.f);
 
 	bool IsExceedingMaxSpeed(const FVector& Velocity, float InMaxSpeed)
@@ -99,7 +101,7 @@ namespace FlyingMovement
 		OutDelta = Delta;
 	}
 
-	float SlideAlongSurface(FMovementSystem::IMovementDriver* Driver, const FVector& Delta, float Time, const FQuat Rotation, const FVector& Normal, FHitResult& Hit, bool bHandleImpact)
+	float SlideAlongSurface(IMovementDriver* Driver, const FVector& Delta, float Time, const FQuat Rotation, const FVector& Normal, FHitResult& Hit, bool bHandleImpact)
 	{
 		if (!Hit.bBlockingHit)
 		{
@@ -152,7 +154,8 @@ namespace FlyingMovement
 		return 0.f;
 	}
 
-	void FMovementSystem::Update(IMovementDriver* Driver, const TSimTime& DeltaTimeMS, const FInputCmd& InputCmd, const FMoveState& InputState, FMoveState& OutputState, const FAuxState& AuxState)
+
+	void FMovementSimulation::Update(IMovementDriver* Driver, const float DeltaSeconds, const FInputCmd& InputCmd, const FMoveState& InputState, FMoveState& OutputState, const FAuxState& AuxState)
 	{
 		// ----------------------------------
 		// AuxState todo
@@ -162,8 +165,7 @@ namespace FlyingMovement
 		float Acceleration = 4000.f;
 
 		// ----------------------------------------------------------------
-
-		const float DeltaSeconds = DeltaTimeMS.ToRealTimeSeconds();
+		
 		OutputState = InputState;
 
 		IBaseMovementDriver& BaseMovementDriver = Driver->GetBaseMovementDriver();
@@ -232,7 +234,7 @@ namespace FlyingMovement
 			// Finally, output velocity that we calculated
 			OutputState.Velocity = Velocity;
 		
-			if (ForceMispredict)
+			if (FMovementSimulation::ForceMispredict)
 			{
 				OutputState.Velocity += ForceMispredictVelocityMagnitude;
 				ForceMispredict = false;
@@ -275,7 +277,7 @@ namespace FlyingMovement
 		// here in the simulation layer. This may not be the best choice for all movement simulations, but is ok for this one.
 	}
 
-	void FMoveState::VisualLog(const FVisualLoggingParameters& Parameters, IFlyingMovementDriver* Driver, IFlyingMovementDriver* LogDriver)
+	void FMoveState::VisualLog(const FVisualLoggingParameters& Parameters, IMovementDriver* Driver, IMovementDriver* LogDriver) const
 	{
 		IBaseMovementDriver& BaseMovementDriver = Driver->GetBaseMovementDriver();		
 		FTransform Transform(Rotation, Location);		
diff --git a/Engine/Plugins/Runtime/NetworkPrediction/Source/NetworkPrediction/Private/Movement/NetworkInterpolation.cpp b/Engine/Plugins/Runtime/NetworkPrediction/Source/NetworkPrediction/Private/Movement/NetworkInterpolation.cpp
deleted file mode 100644
diff --git a/Engine/Plugins/Runtime/NetworkPrediction/Source/NetworkPrediction/Private/Movement/ParametricMovement.cpp b/Engine/Plugins/Runtime/NetworkPrediction/Source/NetworkPrediction/Private/Movement/ParametricMovement.cpp
--- a/Engine/Plugins/Runtime/NetworkPrediction/Source/NetworkPrediction/Private/Movement/ParametricMovement.cpp
+++ b/Engine/Plugins/Runtime/NetworkPrediction/Source/NetworkPrediction/Private/Movement/ParametricMovement.cpp
@@ -3,6 +3,7 @@
 #include "Movement/ParametricMovement.h"
 #include "NetworkSimulationModelDebugger.h"
 #include "HAL/IConsoleManager.h"
+#include "NetworkSimulationGlobalManager.h"
 
 DEFINE_LOG_CATEGORY_STATIC(LogParametricMovement, Log, All);
 
@@ -25,6 +26,10 @@ static int32 SimulatedProxyBufferSize = 4;
 static FAutoConsoleVariableRef CVarSimulatedProxyBufferSize(TEXT("parametricmover.SimulatedProxyBufferSize"),
 	DrawDebugDefaultLifeTime, TEXT(""), ECVF_Default);
 
+static int32 FixStep = 0;
+static FAutoConsoleVariableRef CVarFixStepMS(TEXT("parametricmover.FixStep"),
+	DrawDebugDefaultLifeTime, TEXT("If > 0, will use fix step"), ECVF_Default);
+
 }
 
 // -------------------------------------------------------------------------------------------------------------------------------
@@ -33,10 +38,12 @@ static FAutoConsoleVariableRef CVarSimulatedProxyBufferSize(TEXT("parametricmove
 
 namespace ParametricMovement
 {
-	void FMovementSystem::Update(IMovementDriver* Driver, const TSimTime& SimeTimeDeltaMS, const FInputCmd& InputCmd, const FMoveState& InputState, FMoveState& OutputState, const FAuxState& AuxState)
+	const FName FMovementSimulation::GroupName(TEXT("Parametric"));
+
+	void FMovementSimulation::Update(IMovementDriver* Driver, const float DeltaSeconds, const FInputCmd& InputCmd, const FMoveState& InputState, FMoveState& OutputState, const FAuxState& AuxState)
 	{
 		IBaseMovementDriver& BaseMovementDriver = Driver->GetBaseMovementDriver();
-		const float DeltaSeconds = SimeTimeDeltaMS.ToRealTimeSeconds();
+		//const float DeltaSeconds = SimTimeDeltaMS.ToRealTimeSeconds();
 
 		// Advance parametric time. This won't always be linear: we could loop, rewind/bounce, etc
 		const float InputPlayRate = InputCmd.PlayRate.Get(InputState.PlayRate); // Returns InputCmds playrate if set, else returns previous state's playrate		
@@ -45,13 +52,13 @@ namespace ParametricMovement
 		// We have our time that we should be at. We just need to move primitive component to that position.
 		// Again, note that we expect this cannot fail. We move like this so that it can push things, but we don't expect failure.
 		// (I think it would need to be a different simulation that supports this as a failure case. The tricky thing would be working out
-		// the ouput Position in the case where a Move is blocked (E.g, you move 50% towards the desired location)
+		// the output Position in the case where a Move is blocked (E.g, you move 50% towards the desired location)
 
 		FTransform StartingTransform;
-		Driver->SetTransformForPosition(InputState.Position, StartingTransform);
+		Driver->MapTimeToTransform(InputState.Position, StartingTransform);
 
 		FTransform NewTransform;
-		Driver->SetTransformForPosition(OutputState.Position, NewTransform);
+		Driver->MapTimeToTransform(OutputState.Position, NewTransform);
 
 		FHitResult Hit(1.f);
 
@@ -66,21 +73,16 @@ namespace ParametricMovement
 		}
 	}
 
-	void FMoveState::VisualLog(const FVisualLoggingParameters& Parameters, IParametricMovementDriver* Driver, IParametricMovementDriver* LogDriver) const
+	void FMoveState::VisualLog(const FVisualLoggingParameters& Parameters, IMovementDriver* Driver, IMovementDriver* LogDriver) const
 	{
 		IBaseMovementDriver& BaseMovementDriver = Driver->GetBaseMovementDriver();
 		IBaseMovementDriver& LogMovementDriver = LogDriver->GetBaseMovementDriver();
 
 		FTransform Transform;
-		Driver->SetTransformForPosition(Position, Transform);
+		Driver->MapTimeToTransform(Position, Transform);
 
-		IBaseMovementDriver::FDrawDebugParams DrawParams;
-		DrawParams.DebugWorld = LogMovementDriver.GetDriverWorld();
-		DrawParams.DebugLogOwner = LogMovementDriver.GetVLogOwner();
+		IBaseMovementDriver::FDrawDebugParams DrawParams(Parameters, &LogMovementDriver);
 		DrawParams.Transform = Transform;
-		DrawParams.DrawType = (Parameters.Context == EVisualLoggingContext::OtherMispredicted || Parameters.Context == EVisualLoggingContext::OtherPredicted) ? EVisualLoggingDrawType::Crumb : EVisualLoggingDrawType::Full;
-		DrawParams.Lifetime = Parameters.Lifetime;
-		DrawParams.DrawColor = Parameters.GetDebugColor();
 		DrawParams.InWorldText = LexToString(Parameters.Keyframe);
 		DrawParams.LogText = FString::Printf(TEXT("[%d] %s. Position: %.4f. Location: %s. Rotation: %s"), Parameters.Keyframe, *LexToString(Parameters.Context), Position, *Transform.GetLocation().ToString(), *Transform.GetRotation().Rotator().ToString());
 
@@ -98,27 +100,25 @@ UParametricMovementComponent::UParametricMovementComponent()
 
 }
 
-IReplicationProxy* UParametricMovementComponent::InstantiateNetworkSimulation()
+INetworkSimulationModel* UParametricMovementComponent::InstantiateNetworkSimulation()
 {
-	NetworkSim.Reset(new ParametricMovement::FMovementSystem());
-
-#if NETSIM_MODEL_DEBUG
-	FNetworkSimulationModelDebuggerManager::Get().RegisterNetworkSimulationModel(NetworkSim.Get(), (IParametricMovementDriver*)this, GetOwner(), TEXT("ParametricMovement"));
-#endif
-
-	if (bEnableInterpolation)
+	if (bDisableParametricMovementSimulation)
 	{
-		NetworkInterpolator.Reset(new TNetworkSimulationModelInterpolator<ParametricMovement::FMovementSystem, IParametricMovementDriver>(NetworkSim.Get(), (IParametricMovementDriver*)this));
-		NetworkSim->RepProxy_Simulated.bAllowSimulatedExtrapolation = false;
+		return nullptr;
 	}
 
-	return NetworkSim.Get();
-}
-
-void UParametricMovementComponent::InitializeForNetworkRole(ENetRole Role)
-{	
-	check(NetworkSim);
-	NetworkSim->InitializeForNetworkRole(Role, GetSimulationInitParameters(Role));
+	if (ParametricMoverCVars::FixStep > 0)
+	{
+		auto *NewSim = new ParametricMovement::FMovementSystem<16>(this);
+		NewSim->RepProxy_Simulated.bAllowSimulatedExtrapolation = !bEnableInterpolation;
+		DO_NETSIM_MODEL_DEBUG(FNetworkSimulationModelDebuggerManager::Get().RegisterNetworkSimulationModel(NewSim, GetOwner()));
+		return NewSim;
+	}
+	
+	auto *NewSim = new ParametricMovement::FMovementSystem<>(this);
+	NewSim->RepProxy_Simulated.bAllowSimulatedExtrapolation = !bEnableInterpolation;
+	DO_NETSIM_MODEL_DEBUG(FNetworkSimulationModelDebuggerManager::Get().RegisterNetworkSimulationModel(NewSim, GetOwner()));
+	return NewSim;
 }
 
 FNetworkSimulationModelInitParameters UParametricMovementComponent::GetSimulationInitParameters(ENetRole Role)
@@ -139,14 +139,6 @@ void UParametricMovementComponent::BeginPlay()
 	CachedStartingTransform = UpdatedComponent->GetComponentToWorld();
 }
 
-float UParametricMovementComponent::GetTransformForTime(const float InTime, FTransform& OutTransform) const
-{
-	const FVector Delta = ParametricDelta * InTime;
-	OutTransform = CachedStartingTransform;
-	OutTransform.AddToTranslation(Delta);
-	return InTime;
-}
-
 void UParametricMovementComponent::InitSyncState(ParametricMovement::FMoveState& OutSyncState) const
 {
 	// In this case, we just default to the 0 position. Maybe this could be a starting variable set on the component.
@@ -157,7 +149,7 @@ void UParametricMovementComponent::InitSyncState(ParametricMovement::FMoveState&
 void UParametricMovementComponent::FinalizeFrame(const ParametricMovement::FMoveState& SyncState)
 {
 	FTransform NewTransform;
-	GetTransformForTime(SyncState.Position, NewTransform);
+	MapTimeToTransform(SyncState.Position, NewTransform);
 
 	check(UpdatedComponent);
 	UpdatedComponent->SetWorldTransform(NewTransform, false, nullptr, ETeleportType::TeleportPhysics);
@@ -167,6 +159,10 @@ FString UParametricMovementComponent::GetDebugName() const
 {
 	return FString::Printf(TEXT("ParametricMovement. %s. %s"), *UEnum::GetValueAsString(TEXT("Engine.ENetRole"), GetOwnerRole()), *GetName());
 }
+const UObject* UParametricMovementComponent::GetVLogOwner() const
+{
+	return GetOwner();
+}
 
 void UParametricMovementComponent::AdvanceParametricTime(const float InPosition, const float InPlayRate, float &OutPosition, float& OutPlayRate, const float DeltaTimeSeconds) const
 {
@@ -191,7 +187,7 @@ void UParametricMovementComponent::AdvanceParametricTime(const float InPosition,
 	}
 }
 
-void UParametricMovementComponent::SetTransformForPosition(const float InPosition, FTransform& OutTransform) const
+void UParametricMovementComponent::MapTimeToTransform(const float InPosition, FTransform& OutTransform) const
 {
 	const FVector Delta = ParametricDelta * InPosition;
 
@@ -203,36 +199,6 @@ void UParametricMovementComponent::TickComponent(float DeltaTime, enum ELevelTic
 {
 	Super::TickComponent(DeltaTime, TickType, ThisTickFunction);
 
-	if (!bDisableParametricMovementSimulation)
-	{
-		TickSimulation(DeltaTime); // fixme
-
-		if (NetworkInterpolator.IsValid() && GetOwnerRole() == ROLE_SimulatedProxy)
-		{
-			NetworkInterpolator->Tick(DeltaTime, GetOwner());
-		}
-	}
-}
-
-void UParametricMovementComponent::Reconcile()
-{
-
-}
-
-void UParametricMovementComponent::TickSimulation(float DeltaTimeSeconds)
-{
-	const ENetRole OwnerRole = GetOwnerRole();
-
-	PreTickSimulation(DeltaTimeSeconds); // Fixme
-
-	if (NetworkSim && OwnerRole != ROLE_None)
-	{
-		ParametricMovement::FMovementSystem::FTickParameters Parameters(DeltaTimeSeconds, GetOwner());
-
-		// Tick the core network sim, this will consume input and generate new sync state
-		NetworkSim->Tick((IParametricMovementDriver*)this, Parameters);
-	}
-
 	if (ParentNetUpdateFrequency > 0.f)
 	{
 		GetOwner()->NetUpdateFrequency = ParentNetUpdateFrequency;
@@ -244,7 +210,7 @@ void UParametricMovementComponent::TickSimulation(float DeltaTimeSeconds)
 	}
 }
 
-void UParametricMovementComponent::ProduceInput(const ParametricMovement::TSimTime& SimFrameTime, ParametricMovement::FInputCmd& Cmd)
+void UParametricMovementComponent::ProduceInput(const float DeltaTimeSeconds, ParametricMovement::FInputCmd& Cmd)
 {
 	Cmd.PlayRate = PendingPlayRate;
 	PendingPlayRate.Reset();
diff --git a/Engine/Plugins/Runtime/NetworkPrediction/Source/NetworkPrediction/Private/NetworkPredictionComponent.cpp b/Engine/Plugins/Runtime/NetworkPrediction/Source/NetworkPrediction/Private/NetworkPredictionComponent.cpp
--- a/Engine/Plugins/Runtime/NetworkPrediction/Source/NetworkPrediction/Private/NetworkPredictionComponent.cpp
+++ b/Engine/Plugins/Runtime/NetworkPrediction/Source/NetworkPrediction/Private/NetworkPredictionComponent.cpp
@@ -4,6 +4,7 @@
 #include "Net/UnrealNetwork.h"
 #include "GameFramework/Actor.h"
 #include "GameFramework/Pawn.h"
+#include "NetworkSimulationGlobalManager.h"
 
 UNetworkPredictionComponent::UNetworkPredictionComponent()
 {
@@ -14,29 +15,68 @@ void UNetworkPredictionComponent::InitializeComponent()
 {
 	Super::InitializeComponent();
 
-	IReplicationProxy* NetworkSim = InstantiateNetworkSimulation();
-	check(NetworkSim);
+	// Child class will instantiate
+	INetworkSimulationModel* NewNetworkSim = InstantiateNetworkSimulation();
+	NetworkSim.Reset(NewNetworkSim);
 
-	ReplicationProxy_ServerRPC.Init(NetworkSim, EReplicationProxyTarget::ServerRPC);
-	ReplicationProxy_Autonomous.Init(NetworkSim, EReplicationProxyTarget::AutonomousProxy);
-	ReplicationProxy_Simulated.Init(NetworkSim, EReplicationProxyTarget::SimulatedProxy);
-	ReplicationProxy_Replay.Init(NetworkSim, EReplicationProxyTarget::Replay);
+	if (NewNetworkSim == nullptr)
+	{
+		// Its ok to not instantiate a sim
+		return;
+	}
+
+	// Init RepProxies
+	ReplicationProxy_ServerRPC.Init(NewNetworkSim, EReplicationProxyTarget::ServerRPC);
+	ReplicationProxy_Autonomous.Init(NewNetworkSim, EReplicationProxyTarget::AutonomousProxy);
+	ReplicationProxy_Simulated.Init(NewNetworkSim, EReplicationProxyTarget::SimulatedProxy);
+	ReplicationProxy_Replay.Init(NewNetworkSim, EReplicationProxyTarget::Replay);
 #if NETSIM_MODEL_DEBUG
-	ReplicationProxy_Debug.Init(NetworkSim, EReplicationProxyTarget::Debug);
+	ReplicationProxy_Debug.Init(NewNetworkSim, EReplicationProxyTarget::Debug);
 #endif
+
+	// Register with GlobalManager
+	UNetworkSimulationGlobalManager* NetworkSimGlobalManager = GetWorld()->GetSubsystem<UNetworkSimulationGlobalManager>();
+	check(NetworkSimGlobalManager);
+
+	NetworkSimGlobalManager->RegisterModel(NewNetworkSim, GetOwner());
+	
+	CheckOwnerRoleChange();
+}
+
+void UNetworkPredictionComponent::UninitializeComponent()
+{
+	Super::UninitializeComponent();
+
+	if (NetworkSim.IsValid())
+	{
+		UNetworkSimulationGlobalManager* NetworkSimGlobalManager = GetWorld()->GetSubsystem<UNetworkSimulationGlobalManager>();
+		check(NetworkSimGlobalManager);
+
+		NetworkSimGlobalManager->UnregisterModel(NetworkSim.Get(), GetOwner());
+
+		if (ServerRPCHandle.IsValid())
+		{
+			NetworkSimGlobalManager->TickServerRPCDelegate.Remove(ServerRPCHandle);
+		}
+
+		NetworkSim.Reset(nullptr);
+	}
 }
 
 void UNetworkPredictionComponent::PreReplication(IRepChangedPropertyTracker & ChangedPropertyTracker)
 {
 	Super::PreReplication(ChangedPropertyTracker);
 
-	// We have to update our replication proxies so they can be accurately compared against client shadowstate during property replication. ServerRPC proxy does not need to do this.
-	ReplicationProxy_Autonomous.OnPreReplication();
-	ReplicationProxy_Simulated.OnPreReplication();
-	ReplicationProxy_Replay.OnPreReplication();
+	if (NetworkSim.IsValid())
+	{
+		// We have to update our replication proxies so they can be accurately compared against client shadowstate during property replication. ServerRPC proxy does not need to do this.
+		ReplicationProxy_Autonomous.OnPreReplication();
+		ReplicationProxy_Simulated.OnPreReplication();
+		ReplicationProxy_Replay.OnPreReplication();
 #if NETSIM_MODEL_DEBUG
-	ReplicationProxy_Debug.OnPreReplication();
+		ReplicationProxy_Debug.OnPreReplication();
 #endif
+	}
 }
 
 void UNetworkPredictionComponent::PreNetReceive()
@@ -58,6 +98,14 @@ void UNetworkPredictionComponent::GetLifetimeReplicatedProps(TArray< FLifetimePr
 #endif
 }
 
+void UNetworkPredictionComponent::InitializeForNetworkRole(ENetRole Role)
+{
+	if (NetworkSim.IsValid())
+	{
+		NetworkSim->InitializeForNetworkRole(Role, GetSimulationInitParameters(Role));
+	}
+}
+
 FNetworkSimulationModelInitParameters UNetworkPredictionComponent::GetSimulationInitParameters(ENetRole Role)
 {
 	// These are reasonable defaults but may not be right for everyone
@@ -92,19 +140,27 @@ bool UNetworkPredictionComponent::IsLocallyControlled()
 	return bIsLocallyControlled;
 }
 
-void UNetworkPredictionComponent::PreTickSimulation(float DeltaTime)
-{
-	CheckOwnerRoleChange();
-}
-
-void UNetworkPredictionComponent::CheckOwnerRoleChange()
+bool UNetworkPredictionComponent::CheckOwnerRoleChange()
 {
 	ENetRole CurrentRole = GetOwner()->Role;
 	if (CurrentRole != OwnerCachedNetRole)
 	{
+		if (OwnerCachedNetRole == ROLE_AutonomousProxy)
+		{
+			UnregisterServerRPCDelegate();
+		}
+
 		OwnerCachedNetRole = CurrentRole;
-		InitializeForNetworkRole(CurrentRole);		
+		InitializeForNetworkRole(CurrentRole);
+
+		if (OwnerCachedNetRole == ROLE_AutonomousProxy)
+		{
+			RegisterServerRPCDelegate();
+		}
+		return true;
 	}
+
+	return false;
 }
 
 bool UNetworkPredictionComponent::ServerRecieveClientInput_Validate(const FServerReplicationRPCParameter& ProxyParameter)
@@ -116,4 +172,35 @@ void UNetworkPredictionComponent::ServerRecieveClientInput_Implementation(const
 {
 	// The const_cast is unavoidable here because the replication system only allows by value (forces copy, bad) or by const reference. This use case is unique because we are using the RPC parameter as a temp buffer.
 	const_cast<FServerReplicationRPCParameter&>(ProxyParameter).NetSerializeToProxy(ReplicationProxy_ServerRPC);
-}
\ No newline at end of file
+}
+
+void UNetworkPredictionComponent::TickServerRPC(float DeltaSeconds)
+{
+	check(NetworkSim.IsValid());
+	if (NetworkSim->ShouldSendServerRPC(DeltaSeconds))
+	{
+		// Temp hack to make sure the ServerRPC doesn't get suppressed from bandwidth limiting
+		// (system hasn't been optimized and not mature enough yet to handle gaps in input stream)
+		FScopedBandwidthLimitBypass BandwidthBypass(GetOwner());
+
+		FServerReplicationRPCParameter ProxyParameter(ReplicationProxy_ServerRPC);
+		ServerRecieveClientInput(ProxyParameter);
+	}
+}
+
+void UNetworkPredictionComponent::RegisterServerRPCDelegate()
+{
+	UNetworkSimulationGlobalManager* NetworkSimGlobalManager = GetWorld()->GetSubsystem<UNetworkSimulationGlobalManager>();
+	check(NetworkSimGlobalManager);
+	ServerRPCHandle = NetworkSimGlobalManager->TickServerRPCDelegate.AddUObject(this, &UNetworkPredictionComponent::TickServerRPC);
+}
+void UNetworkPredictionComponent::UnregisterServerRPCDelegate()
+{
+	if (ServerRPCHandle.IsValid())
+	{
+		UNetworkSimulationGlobalManager* NetworkSimGlobalManager = GetWorld()->GetSubsystem<UNetworkSimulationGlobalManager>();
+		check(NetworkSimGlobalManager);
+		NetworkSimGlobalManager->TickServerRPCDelegate.Remove(ServerRPCHandle);
+		ServerRPCHandle.Reset();
+	}
+}
diff --git a/Engine/Plugins/Runtime/NetworkPrediction/Source/NetworkPrediction/Private/NetworkPredictionTypes.cpp b/Engine/Plugins/Runtime/NetworkPrediction/Source/NetworkPrediction/Private/NetworkPredictionTypes.cpp
--- a/Engine/Plugins/Runtime/NetworkPrediction/Source/NetworkPrediction/Private/NetworkPredictionTypes.cpp
+++ b/Engine/Plugins/Runtime/NetworkPrediction/Source/NetworkPrediction/Private/NetworkPredictionTypes.cpp
@@ -122,11 +122,17 @@ FScopedBandwidthLimitBypass::~FScopedBandwidthLimitBypass()
 	}
 }
 
-FNetSimTickParametersBase::FNetSimTickParametersBase(AActor* OwningActor)
+FNetSimTickParameters::FNetSimTickParameters(float InLocalDeltaTimeSeconds, AActor* OwningActor)
+{
+	LocalDeltaTimeSeconds = InLocalDeltaTimeSeconds;
+	InitFromActor(OwningActor);
+}
+
+void FNetSimTickParameters::InitFromActor(AActor* Actor)
 {
 	// Default implementation that will probably be good for most actors
-	Role = OwningActor->GetLocalRole();
-	const bool bHasNetConnection = OwningActor->GetNetConnection() != nullptr;
+	Role = Actor->GetLocalRole();
+	const bool bHasNetConnection = Actor->GetNetConnection() != nullptr;
 
 	// The only default case we really don;t want generating local input is on the server when there is a net owning player (auto proxy)
 	// This implies that we will (sim) extrapolate by default. To opt out of extrapolation, set to false or you can just not tick the net sim
diff --git a/Engine/Plugins/Runtime/NetworkPrediction/Source/NetworkPrediction/Private/NetworkSimulationGlobalManager.cpp b/Engine/Plugins/Runtime/NetworkPrediction/Source/NetworkPrediction/Private/NetworkSimulationGlobalManager.cpp
--- a/Engine/Plugins/Runtime/NetworkPrediction/Source/NetworkPrediction/Private/NetworkSimulationGlobalManager.cpp
+++ b/Engine/Plugins/Runtime/NetworkPrediction/Source/NetworkPrediction/Private/NetworkSimulationGlobalManager.cpp
@@ -1,18 +1,81 @@
 // Copyright 1998-2019 Epic Games, Inc. All Rights Reserved.
 
 #include "NetworkSimulationGlobalManager.h"
+#include "Engine/World.h"
+
+DEFINE_LOG_CATEGORY_STATIC(NetworkSimulationGlobalManager, Log, All);
 
 UNetworkSimulationGlobalManager::UNetworkSimulationGlobalManager()
 {
 
 }
 
+void UNetworkSimulationGlobalManager::Initialize(FSubsystemCollectionBase& Collection)
+{
+	UWorld* World = GetWorld();
+	check(World);
+
+	PostTickDispatchHandle = World->OnPostTickDispatch().AddUObject(this, &UNetworkSimulationGlobalManager::ReconcileSimulationsPostNetworkUpdate);
+	PreWorldActorTickHandle = FWorldDelegates::OnWorldPreActorTick.AddUObject(this, &UNetworkSimulationGlobalManager::BeginNewSimulationFrame);
+}
+
+void UNetworkSimulationGlobalManager::Deinitialize()
+{
+	if (UWorld* World = GetWorld())
+	{
+		World->OnPostTickDispatch().Remove(PostTickDispatchHandle);
+	}
+
+	FWorldDelegates::OnWorldPreActorTick.Remove(PreWorldActorTickHandle);
+}
+
 void UNetworkSimulationGlobalManager::ReconcileSimulationsPostNetworkUpdate()
 {
+	for (auto It = SimulationGroupMap.CreateIterator(); It; ++It)
+	{
+		for (auto& SimItem : It.Value().Simulations)
+		{
+			SimItem.Simulation->Reconcile(SimItem.OwningActor->Role);
+		}
+	}
+}
+
+void UNetworkSimulationGlobalManager::BeginNewSimulationFrame(UWorld* InWorld, ELevelTick InLevelTick, float InDeltaSeconds)
+{
+	if (InWorld != GetWorld() || !InWorld->HasBegunPlay())
+	{
+		return;
+	}
 
+	FNetSimTickParameters TickParameters(InDeltaSeconds);
+
+	for (auto It = SimulationGroupMap.CreateIterator(); It; ++It)
+	{
+		for (auto& SimItem : It.Value().Simulations)
+		{
+			TickParameters.InitFromActor(SimItem.OwningActor);
+			SimItem.Simulation->Tick(TickParameters);
+		}
+	}
+
+	TickServerRPCDelegate.Broadcast(InDeltaSeconds);
 }
 
-void UNetworkSimulationGlobalManager::TickNewSimulationFrame(float DeltaTimeSeconds)
+void UNetworkSimulationGlobalManager::RegisterModel(INetworkSimulationModel* Model, AActor* OwningActor)
 {
+	SimulationGroupMap.FindOrAdd(Model->GetSimulationGroupName()).Simulations.Emplace(Model, OwningActor);
+}
 
+void UNetworkSimulationGlobalManager::UnregisterModel(INetworkSimulationModel* Model, AActor* OwningActor)
+{
+	auto& SimulationList = SimulationGroupMap.FindOrAdd(Model->GetSimulationGroupName()).Simulations;
+	int32 idx = SimulationList.IndexOfByKey(Model);
+	if (idx != INDEX_NONE)
+	{
+		SimulationList.RemoveAtSwap(idx, 1, false);
+	}
+	else
+	{
+		UE_LOG(NetworkSimulationGlobalManager, Warning, TEXT("::Unregister model called on model %s but could not be found in manager."), *Model->GetSimulationGroupName().ToString());
+	}
 }
\ No newline at end of file
diff --git a/Engine/Plugins/Runtime/NetworkPrediction/Source/NetworkPrediction/Public/Movement/BaseMovementComponent.h b/Engine/Plugins/Runtime/NetworkPrediction/Source/NetworkPrediction/Public/Movement/BaseMovementComponent.h
--- a/Engine/Plugins/Runtime/NetworkPrediction/Source/NetworkPrediction/Public/Movement/BaseMovementComponent.h
+++ b/Engine/Plugins/Runtime/NetworkPrediction/Source/NetworkPrediction/Public/Movement/BaseMovementComponent.h
@@ -27,8 +27,8 @@ public:
 	virtual bool MoveUpdatedComponent(const FVector& Delta, const FQuat& NewRotation, bool bSweep, FHitResult* OutHit, ETeleportType Teleport) const = 0;
 	virtual FTransform GetUpdateComponentTransform() const = 0;
 
-	virtual UWorld* GetDriverWorld() const = 0;
-	virtual UObject* GetVLogOwner() const = 0;
+	virtual UWorld* GetBaseMovementDriverWorld() const = 0;
+	virtual UObject* GetBaseMovementVLogOwner() const = 0;
 
 	struct FDrawDebugParams
 	{
@@ -69,8 +69,8 @@ public:
 	bool MoveUpdatedComponent(const FVector& Delta, const FQuat& NewRotation, bool bSweep, FHitResult* OutHit, ETeleportType Teleport) const override final;
 	FTransform GetUpdateComponentTransform() const override final;
 	
-	UWorld* GetDriverWorld() const override final { return GetWorld(); }
-	UObject* GetVLogOwner() const override final;	
+	UWorld* GetBaseMovementDriverWorld() const override final { return GetWorld(); }
+	UObject* GetBaseMovementVLogOwner() const override final;	
 
 	virtual void DrawDebug(const IBaseMovementDriver::FDrawDebugParams& Params) const override;
 
diff --git a/Engine/Plugins/Runtime/NetworkPrediction/Source/NetworkPrediction/Public/Movement/FlyingMovement.h b/Engine/Plugins/Runtime/NetworkPrediction/Source/NetworkPrediction/Public/Movement/FlyingMovement.h
--- a/Engine/Plugins/Runtime/NetworkPrediction/Source/NetworkPrediction/Public/Movement/FlyingMovement.h
+++ b/Engine/Plugins/Runtime/NetworkPrediction/Source/NetworkPrediction/Public/Movement/FlyingMovement.h
@@ -27,6 +27,8 @@ class USceneComponent;
 
 namespace FlyingMovement
 {
+	class IMovementDriver;
+
 	// State the client generates
 	struct FInputCmd
 	{
@@ -90,7 +92,13 @@ namespace FlyingMovement
 			}
 		}
 
-		void VisualLog(const FVisualLoggingParameters& Parameters, IFlyingMovementDriver* Driver, IFlyingMovementDriver* LogDriver);
+		void VisualLog(const FVisualLoggingParameters& Parameters, IMovementDriver* Driver, IMovementDriver* LogDriver) const;
+
+		static void Interpolate(const FMoveState& From, const FMoveState& To, const float PCT, FMoveState& OutDest)
+		{
+			OutDest.Location = From.Location + ((To.Location - From.Location) * PCT);
+			OutDest.Rotation = From.Rotation + ((To.Rotation - From.Rotation) * PCT);
+		}
 	};
 
 	// Auxiliary state that is input into the simulation. Doesn't change during the simulation tick.
@@ -106,33 +114,41 @@ namespace FlyingMovement
 
 	using TMovementBufferTypes = TNetworkSimBufferTypes<FInputCmd, FMoveState, FAuxState>;
 
-	// Actual definition of our network simulation.
-	class FMovementSystem : public TNetworkedSimulationModel<FMovementSystem, TMovementBufferTypes, TNetworkSimTickSettings<0>>
+	static FName SimulationGroupName("FlyingMovement");
+
+	// Interface between the simulation and owning component driving it. Functions added here are available in ::Update and must be implemented by UMockNetworkSimulationComponent.
+	class IMovementDriver : public TNetworkSimDriverInterfaceBase<TMovementBufferTypes>
 	{
 	public:
 
-		// Interface between the simulation and owning component driving it. Functions added here are available in ::Update and must be implemented by UMockNetworkSimulationComponent.
-		class IMovementDriver : public IDriver
-		{
-		public:
-			// Interface for moving the collision component around
-			virtual IBaseMovementDriver& GetBaseMovementDriver() = 0;
+		// Interface for moving the collision component around
+		virtual IBaseMovementDriver& GetBaseMovementDriver() = 0;
 
-			// Called prior to running the sim to make sure to make sure the collision component is in the right place. 
-			// This is unfortunate and not good, but is needed to ensure our collision and world position have not been moved out from under us.
-			// Refactoring primitive component movement to allow the sim to do all collision queries outside of the component code would be ideal.
-			virtual void PreSimSync(const FMoveState& SyncState) = 0;
-		};
+		// Called prior to running the sim to make sure to make sure the collision component is in the right place. 
+		// This is unfortunate and not good, but is needed to ensure our collision and world position have not been moved out from under us.
+		// Refactoring primitive component movement to allow the sim to do all collision queries outside of the component code would be ideal.
+		virtual void PreSimSync(const FMoveState& SyncState) = 0;
+	};
 
+	class FMovementSimulation
+	{
+	public:
 		/** Main update function */
-		static void Update(IMovementDriver* Driver, const TSimTime& DeltaTimeMS, const FInputCmd& InputCmd, const FMoveState& InputState, FMoveState& OutputState, const FAuxState& AuxState);
+		static void Update(IMovementDriver* Driver, const float DeltaSeconds, const FInputCmd& InputCmd, const FMoveState& InputState, FMoveState& OutputState, const FAuxState& AuxState);
+
+		/** Tick group the simulation maps to */
+		static const FName GroupName;
 
 		/** Dev tool to force simple mispredict */
 		static bool ForceMispredict;
 	};
 
+	// Actual definition of our network simulation.
+	template<int32 InFixedStepMS=0>
+	using FMovementSystem = TNetworkedSimulationModel<FMovementSimulation, IMovementDriver, TMovementBufferTypes, TNetworkSimTickSettings<InFixedStepMS> >;
+
 	// Simulation time used by the movement system
-	using TSimTime = FMovementSystem::TSimTime;
+	//using TSimTime = FMovementSystem::TSimTime;
 
 	// general tolerance value for rotation checks
 	static const float ROTATOR_TOLERANCE = (1e-3);
@@ -141,7 +157,7 @@ namespace FlyingMovement
 
 // Needed to trick UHT into letting UMockNetworkSimulationComponent implement. UHT cannot parse the ::
 // Also needed for forward declaring. Can't just be a typedef/using =
-class IFlyingMovementDriver : virtual public FlyingMovement::FMovementSystem::IMovementDriver { };
+class IFlyingMovementDriver : public FlyingMovement::IMovementDriver { };
 
 // -------------------------------------------------------------------------------------------------------------------------------
 // ActorComponent for running FlyingMovement 
@@ -161,20 +177,17 @@ public:
 	IBaseMovementDriver& GetBaseMovementDriver() override final { return *static_cast<IBaseMovementDriver*>(this); }
 
 	FString GetDebugName() const override;
+	const UObject* GetVLogOwner() const override;
 	void InitSyncState(FlyingMovement::FMoveState& OutSyncState) const override;
 	void PreSimSync(const FlyingMovement::FMoveState& SyncState) override;
-	void ProduceInput(const FlyingMovement::TSimTime& SimFrameTime, FlyingMovement::FInputCmd& Cmd) override;
+	void ProduceInput(const float DeltaTimeSeconds, FlyingMovement::FInputCmd& Cmd) override;
 	void FinalizeFrame(const FlyingMovement::FMoveState& SyncState) override;
 
-	DECLARE_DELEGATE_TwoParams(FProduceFlyingInput, const FlyingMovement::TSimTime& /*SimFrameTime*/, FlyingMovement::FInputCmd& /*Cmd*/)
+	DECLARE_DELEGATE_TwoParams(FProduceFlyingInput, const float /*DeltaTimeSeconds*/, FlyingMovement::FInputCmd& /*Cmd*/)
 	FProduceFlyingInput ProduceInputDelegate;
 
 protected:
 
 	// Network Prediction
-	virtual IReplicationProxy* InstantiateNetworkSimulation() override;
-	virtual void InitializeForNetworkRole(ENetRole Role) override;
-
-	TUniquePtr<FlyingMovement::FMovementSystem> NetworkSim; // The Network sim that this component is managing. This is what is doing all the work.
-	
+	virtual INetworkSimulationModel* InstantiateNetworkSimulation() override;
 };
\ No newline at end of file
diff --git a/Engine/Plugins/Runtime/NetworkPrediction/Source/NetworkPrediction/Public/Movement/ParametricMovement.h b/Engine/Plugins/Runtime/NetworkPrediction/Source/NetworkPrediction/Public/Movement/ParametricMovement.h
--- a/Engine/Plugins/Runtime/NetworkPrediction/Source/NetworkPrediction/Public/Movement/ParametricMovement.h
+++ b/Engine/Plugins/Runtime/NetworkPrediction/Source/NetworkPrediction/Public/Movement/ParametricMovement.h
@@ -4,7 +4,6 @@
 
 #include "NetworkSimulationModel.h"
 #include "BaseMovementComponent.h"
-#include "NetworkInterpolation.h"
 
 #include "ParametricMovement.generated.h"
 
@@ -12,6 +11,8 @@ class IParametricMovementDriver;
 
 namespace ParametricMovement
 {
+	class IMovementDriver;
+
 	// State the client generates
 	struct FInputCmd
 	{
@@ -72,7 +73,7 @@ namespace ParametricMovement
 			}
 		}
 
-		void VisualLog(const FVisualLoggingParameters& Parameters, IParametricMovementDriver* Driver, IParametricMovementDriver* LogDriver) const;
+		void VisualLog(const FVisualLoggingParameters& Parameters, IMovementDriver* Driver, IMovementDriver* LogDriver) const;
 
 		static void Interpolate(const FMoveState& From, const FMoveState& To, const float PCT, FMoveState& OutDest)
 		{
@@ -94,41 +95,41 @@ namespace ParametricMovement
 
 	using TMovementBufferTypes = TNetworkSimBufferTypes<FInputCmd, FMoveState, FAuxState>;
 
-	// Actual definition of our network simulation.
-	class FMovementSystem : public TNetworkedSimulationModel<FMovementSystem, TMovementBufferTypes, TNetworkSimTickSettings<0>>
+	// Interface between the simulation and owning component driving it. Functions added here are available in ::Update
+	class IMovementDriver : public TNetworkSimDriverInterfaceBase<TMovementBufferTypes>
 	{
 	public:
 
-		// Interface between the simulation and owning component driving it. Functions added here are available in ::Update and must be implemented by UMockNetworkSimulationComponent.
-		class IMovementDriver : public IDriver
-		{
-		public:
-
-			// BaseMovement driver API (functions for moving around a primitive component)
-			virtual IBaseMovementDriver& GetBaseMovementDriver() = 0;
+		// BaseMovement driver API (functions for moving around a primitive component)
+		virtual IBaseMovementDriver& GetBaseMovementDriver() = 0;
 
-			// Advance parametric time. This is meant to do simple things like looping/reversing etc.
-			// Note how this should be STATIC and not rely on state outside of what is passed in (such thing would need to be done inside the simulation, not through the driver!)
-			virtual void AdvanceParametricTime(const float InPosition, const float InPlayRate, float &OutPosition, float& OutPlayRate, const float DeltaTimeSeconds) const = 0;
+		// Advance parametric time. This is meant to do simple things like looping/reversing etc.
+		// Note how this should be STATIC and not rely on state outside of what is passed in (such thing would need to be done inside the simulation, not through the driver!)
+		virtual void AdvanceParametricTime(const float InPosition, const float InPlayRate, float &OutPosition, float& OutPlayRate, const float DeltaTimeSeconds) const = 0;
 
-			// Actually turn the given position into a transform. Again, should be static and not conditional on changing state outside of the network sim
-			virtual void SetTransformForPosition(const float InPosition, FTransform& OutTransform) const = 0;
-		};
+		// Actually turn the given position into a transform. Again, should be static and not conditional on changing state outside of the network sim
+		virtual void MapTimeToTransform(const float InPosition, FTransform& OutTransform) const = 0;
+	};
+	
+	class FMovementSimulation
+	{
+	public:
+		static void Update(IMovementDriver* Driver, const float DeltaSeconds, const FInputCmd& InputCmd, const FMoveState& InputState, FMoveState& OutputState, const FAuxState& AuxState);
+		static const FName GroupName;
+	};
 
-		/** Main update function */
-		static void Update(IMovementDriver* Driver, const TSimTime& SimeTimeDeltaMS, const FInputCmd& InputCmd, const FMoveState& InputState, FMoveState& OutputState, const FAuxState& AuxState);
 
-		/** Dev tool to force simple mispredict */
-		static bool ForceMispredict;
-	};
+	// Actual definition of our network simulation.
+	template<int32 FixedStepMS=0>
+	using FMovementSystem = TNetworkedSimulationModel<FMovementSimulation, IMovementDriver, TMovementBufferTypes, TNetworkSimTickSettings<FixedStepMS>>;
 
-	using TSimTime = FMovementSystem::TSimTime;
+	//using TSimTime = FMovementSystem::TSimTime;
 
 } // End namespace
 
 // Needed to trick UHT into letting UMockNetworkSimulationComponent implement. UHT cannot parse the ::
 // Also needed for forward declaring. Can't just be a typedef/using =
-class IParametricMovementDriver : public ParametricMovement::FMovementSystem::IMovementDriver { };
+class IParametricMovementDriver : public ParametricMovement::IMovementDriver { };
 
 // -------------------------------------------------------------------------------------------------------------------------------
 //	ActorComponent for running basic Parametric movement. 
@@ -149,35 +150,25 @@ class NETWORKPREDICTION_API UParametricMovementComponent : public UBaseMovementC
 	virtual void BeginPlay() override;
 	virtual void TickComponent(float DeltaTime, enum ELevelTick TickType, FActorComponentTickFunction* ThisTickFunction) override;
 
-	virtual void Reconcile() override;
-	virtual void TickSimulation(float DeltaTimeSeconds) override;
-
 	// Base TNetworkModelSimulation driver
 	FString GetDebugName() const override;
+	const UObject* GetVLogOwner() const override;
 	void InitSyncState(ParametricMovement::FMoveState& OutSyncState) const override;
 	void FinalizeFrame(const ParametricMovement::FMoveState& SyncState) override;
-	void ProduceInput(const ParametricMovement::TSimTime& SimFrameTime, ParametricMovement::FInputCmd& Cmd);
+	void ProduceInput(const float DeltaTimeSeconds, ParametricMovement::FInputCmd& Cmd);
 
 	// Base Movement Driver
 	IBaseMovementDriver& GetBaseMovementDriver() override final { return *static_cast<IBaseMovementDriver*>(this); }
 
 	// Parametric Movement Driver
-	float GetTransformForTime(const float InTime, FTransform& OutTransform) const;
-
 	virtual void AdvanceParametricTime(const float InPosition, const float InPlayRate, float &OutPosition, float& OutPlayRate, const float DeltaTimeSeconds) const override;
-	virtual void SetTransformForPosition(const float InPosition, FTransform& OutTransform) const override;
+	virtual void MapTimeToTransform(const float InPosition, FTransform& OutTransform) const override;
 
 protected:
 
-	virtual IReplicationProxy* InstantiateNetworkSimulation() override;	
-	void InitializeForNetworkRole(ENetRole Role) override;
+	virtual INetworkSimulationModel* InstantiateNetworkSimulation() override;
 	FNetworkSimulationModelInitParameters GetSimulationInitParameters(ENetRole Role) override;
 
-	TUniquePtr<ParametricMovement::FMovementSystem> NetworkSim;	
-
-	// Temp, will be moved
-	TUniquePtr< TNetworkSimulationModelInterpolator<ParametricMovement::FMovementSystem, IParametricMovementDriver> > NetworkInterpolator;
-
 	// ------------------------------------------------------------------------
 	// Temp Parametric movement example
 	//	The essence of this movement simulation is to map some Time value to a transform. That is it.
diff --git a/Engine/Plugins/Runtime/NetworkPrediction/Source/NetworkPrediction/Public/Movement/ParametricMovementSimulation.h b/Engine/Plugins/Runtime/NetworkPrediction/Source/NetworkPrediction/Public/Movement/ParametricMovementSimulation.h
new file mode 100644
--- /dev/null
+++ b/Engine/Plugins/Runtime/NetworkPrediction/Source/NetworkPrediction/Public/Movement/ParametricMovementSimulation.h
@@ -0,0 +1,8 @@
+// Copyright 1998-2019 Epic Games, Inc. All Rights Reserved
+
+#pragma once
+
+namespace ParametricMovement
+{
+
+}
\ No newline at end of file
diff --git a/Engine/Plugins/Runtime/NetworkPrediction/Source/NetworkPrediction/Public/NetworkInterpolation.h b/Engine/Plugins/Runtime/NetworkPrediction/Source/NetworkPrediction/Public/NetworkInterpolation.h
deleted file mode 100644
diff --git a/Engine/Plugins/Runtime/NetworkPrediction/Source/NetworkPrediction/Public/NetworkPredictionComponent.h b/Engine/Plugins/Runtime/NetworkPrediction/Source/NetworkPrediction/Public/NetworkPredictionComponent.h
--- a/Engine/Plugins/Runtime/NetworkPrediction/Source/NetworkPrediction/Public/NetworkPredictionComponent.h
+++ b/Engine/Plugins/Runtime/NetworkPrediction/Source/NetworkPrediction/Public/NetworkPredictionComponent.h
@@ -6,12 +6,6 @@
 
 #include "NetworkPredictionComponent.generated.h"
 
-class INetworkSimulationOwner
-{
-	virtual void Reconcile() = 0;
-	virtual void TickSimulation(float DeltaTimeSeconds) = 0;
-};
-
 // ----------------------------------------------------------------------------------------------------------------------------------------------------------------------
 //	UNetworkPredictionComponent
 //	This is the base component for running a TNetworkedSimulationModel through an actor component. This contains the boiler plate hooks into getting the system
@@ -33,22 +27,20 @@ public:
 	UNetworkPredictionComponent();
 
 	virtual void InitializeComponent() override;
+	virtual void UninitializeComponent() override;
 	virtual void PreReplication(IRepChangedPropertyTracker & ChangedPropertyTracker) override;
 	virtual void PreNetReceive() override;
 	
 	virtual void Reconcile() { }
 	virtual void TickSimulation(float DeltaTimeSeconds) { }
 
-	void PreTickSimulation(float DeltaTime);
-
 protected:
+	
+	// Classes must instantiate their own NetworkSim here. The UNetworkPredictionComponent will manage its lifetime
+	virtual INetworkSimulationModel* InstantiateNetworkSimulation() { check(false); return nullptr; }
 
-	// Child classes must allocate and manage lifetime of their own NetworkSim (E.g, TNetworkedSimulationModel). Recommend to just store in a TUniquePtr.
-	// Child classes must also register with Network Sim debugger here (FNetworkSimulationModelDebuggerManager RegisterNetworkSimulationModel)
-	virtual IReplicationProxy* InstantiateNetworkSimulation() { check(false); return nullptr; }
-
-	// Child must override this an initialize their NetworkSim here (call 
-	virtual void InitializeForNetworkRole(ENetRole Role) { check(false); }
+	// Finalizes initialization when NetworkRole changes. Does not need to be overridden.
+	virtual void InitializeForNetworkRole(ENetRole Role);
 
 	// Doesn't need to be overridden. Expected to be used in ::InitializeForNetworkRole implementation
 	virtual FNetworkSimulationModelInitParameters GetSimulationInitParameters(ENetRole Role);
@@ -58,7 +50,7 @@ protected:
 
 	// Helper: Checks if the owner's role has changed and calls InitializeForNetworkRole if necessary.
 	// NOTE: You may need to call this in your TickComponent function
-	void CheckOwnerRoleChange();
+	bool CheckOwnerRoleChange();
 	ENetRole OwnerCachedNetRole = ROLE_None;
 
 	UFUNCTION(Server, unreliable, WithValidation)
@@ -84,6 +76,15 @@ private:
 
 protected:
 
+	void TickServerRPC(float DeltaSeconds);
+
+	void RegisterServerRPCDelegate();
+	void UnregisterServerRPCDelegate();
+	FDelegateHandle ServerRPCHandle;
+
 	UFUNCTION()
 	virtual void OnRep_SimulatedProxy() { }
+
+	// The Network sim that this component is managing. This is what is doing all the work.
+	TUniquePtr<INetworkSimulationModel> NetworkSim;
 };
\ No newline at end of file
diff --git a/Engine/Plugins/Runtime/NetworkPrediction/Source/NetworkPrediction/Public/NetworkPredictionTypes.h b/Engine/Plugins/Runtime/NetworkPrediction/Source/NetworkPrediction/Public/NetworkPredictionTypes.h
--- a/Engine/Plugins/Runtime/NetworkPrediction/Source/NetworkPrediction/Public/NetworkPredictionTypes.h
+++ b/Engine/Plugins/Runtime/NetworkPrediction/Source/NetworkPrediction/Public/NetworkPredictionTypes.h
@@ -11,6 +11,11 @@
 // -------------------------------------------------------------------------------------------------------------------------------
 
 #define NETSIM_MODEL_DEBUG !(UE_BUILD_SHIPPING || UE_BUILD_TEST)
+#if NETSIM_MODEL_DEBUG
+#define DO_NETSIM_MODEL_DEBUG(X) X
+#else
+#define DO_NETSIM_MODEL_DEBUG(X)
+#endif
 
 NETWORKPREDICTION_API DECLARE_LOG_CATEGORY_EXTERN(LogNetworkSim, Log, All);
 
@@ -162,11 +167,25 @@ inline FString LexToString(ESimulatedUpdateMode A)
 // Ticking parameters use to drive the simulation
 // -------------------------------------------------------------------------------------------------------------------------------
 
-class FNetworkSimulationTickingParameters
+struct NETWORKPREDICTION_API FNetSimTickParameters
 {
-	float LocalFrameDeltaTimeSeconds;
-};
+	FNetSimTickParameters(float InLocalDeltaTimeSeconds) : LocalDeltaTimeSeconds(InLocalDeltaTimeSeconds) { }
+	FNetSimTickParameters(float InLocalDeltaTimeSeconds, AActor* Actor);
+
+	/** Initializes Role and bGenerateLocalInputCmds from an Actor's state. */
+	void InitFromActor(AActor* Actor);
 
+	// Owner's role. Necessary to know which proxy we should be forwarding functions in tick to
+	ENetRole Role = ROLE_None;
+
+	// Are we creating input cmds locally. Note this is distinct from Role/Authority:
+	//		-[On Server] Autonomous Proxy client = false
+	//		-[On Server] Non player controlled actor = true
+	//		-[On Client] Simulated proxies (everyone but client) = true, if you want extrapolation. Note clients can just not tick the netsim to disable extrapolation as well.
+	bool bGenerateLocalInputCmds = false;
+
+	float LocalDeltaTimeSeconds = 0.f;
+};
 
 // -------------------------------------------------------------------------------------------------------------------------------
 // Interface that the proxy talks to. This is what will implement the replication.
@@ -181,6 +200,20 @@ public:
 	virtual int32 GetProxyDirtyCount(EReplicationProxyTarget Target) = 0;
 };
 
+class INetworkSimulationModel : public IReplicationProxy
+{
+public:
+
+	virtual FName GetSimulationGroupName() const = 0;
+
+	virtual void Reconcile(const ENetRole Role) = 0;
+	virtual void Tick(const FNetSimTickParameters&) = 0;
+	virtual void InitializeForNetworkRole(const ENetRole Role, const FNetworkSimulationModelInitParameters& Parameters) = 0;
+
+	virtual bool ShouldSendServerRPC(float DeltaSeconds) = 0;
+	virtual void SetDesiredServerRPCSendFrequency(float DesiredHz) = 0;
+};
+
 // -------------------------------------------------------------------------------------------------------------------------------
 //	FReplicationProxy
 //	Replicated USTRUCT that point to the IReplicationProxy to do the replication.
@@ -263,17 +296,3 @@ private:
 	class UNetConnection* CachedNetConnection = nullptr;
 };
 
-struct NETWORKPREDICTION_API FNetSimTickParametersBase
-{
-	FNetSimTickParametersBase();
-	FNetSimTickParametersBase(AActor* Actor);
-
-	// Owner's role. Necessary to know which proxy we should be forwarding functions in tick to
-	ENetRole Role = ROLE_None;
-
-	// Are we creating input cmds locally. Note this is distinct from Role/Authority:
-	//		-[On Server] Autonomous Proxy client = false
-	//		-[On Server] Non player controlled actor = true
-	//		-[On Client] Simulated proxies (everyone but client) = true, if you want extrapolation. Note clients can just not tick the netsim to disable extrapolation as well.
-	bool bGenerateLocalInputCmds = false;
-};
\ No newline at end of file
diff --git a/Engine/Plugins/Runtime/NetworkPrediction/Source/NetworkPrediction/Public/NetworkSimulationGlobalManager.h b/Engine/Plugins/Runtime/NetworkPrediction/Source/NetworkPrediction/Public/NetworkSimulationGlobalManager.h
--- a/Engine/Plugins/Runtime/NetworkPrediction/Source/NetworkPrediction/Public/NetworkSimulationGlobalManager.h
+++ b/Engine/Plugins/Runtime/NetworkPrediction/Source/NetworkPrediction/Public/NetworkSimulationGlobalManager.h
@@ -6,6 +6,9 @@
 #include "UObject/ObjectMacros.h"
 #include "UObject/UObjectGlobals.h"
 #include "UObject/Object.h"
+#include "Subsystems/WorldSubsystem.h"
+#include "Engine/EngineBaseTypes.h"
+#include "NetworkSimulationModelDebugger.h"
 
 #include "NetworkSimulationGlobalManager.generated.h"
 
@@ -33,16 +36,51 @@
 //
 // ----------------------------------------------------------------------------------------------------------------------------
 
+class IReplicationProxy;
+
 UCLASS()
-class UNetworkSimulationGlobalManager : public UObject
+class UNetworkSimulationGlobalManager : public UWorldSubsystem
 {
 	GENERATED_BODY()
-
 public:
 
 	UNetworkSimulationGlobalManager();
+	
+	// Subsystem Init/Deinit
+	virtual void Initialize(FSubsystemCollectionBase& Collection) override;
+	virtual void Deinitialize() override;
+
+	// Netsim Register/Unregister
+	void RegisterModel(INetworkSimulationModel* Model, AActor* OwningActor);
+	void UnregisterModel(INetworkSimulationModel* Model, AActor* OwningActor);
+
+	// Delegate for auto proxies to register with to send their server RPC
+	DECLARE_MULTICAST_DELEGATE_OneParam(FTickServerRPC, float /*RealTime DeltaSeconds*/)
+	FTickServerRPC TickServerRPCDelegate;
+
+private:
 
 	void ReconcileSimulationsPostNetworkUpdate();
-	void TickNewSimulationFrame(float DeltaTimeSeconds);
+	void BeginNewSimulationFrame(UWorld* InWorld, ELevelTick InLevelTick, float InDeltaSeconds);
+	
+	// Structure for organizing our simulations into groups to give some ticking order.
+	// This will need to evolve over time a bit and be optimized for larger number of simulations.
+	struct FSimulationGroup
+	{
+		struct FItem
+		{
+			FItem(INetworkSimulationModel* InSim, AActor* InActor) : Simulation(InSim), OwningActor(InActor) { }
+			bool operator == (const INetworkSimulationModel* InSim) const { return Simulation == InSim; }
+
+			INetworkSimulationModel* Simulation;
+			AActor* OwningActor;
+		};
+
+		TArray<FItem> Simulations;
+	};
+
+	TSortedMap<FName, FSimulationGroup, FDefaultAllocator, FNameFastLess> SimulationGroupMap;
 
+	FDelegateHandle PostTickDispatchHandle;
+	FDelegateHandle PreWorldActorTickHandle;
 };
\ No newline at end of file
diff --git a/Engine/Plugins/Runtime/NetworkPrediction/Source/NetworkPrediction/Public/NetworkSimulationModel.h b/Engine/Plugins/Runtime/NetworkPrediction/Source/NetworkPrediction/Public/NetworkSimulationModel.h
--- a/Engine/Plugins/Runtime/NetworkPrediction/Source/NetworkPrediction/Public/NetworkSimulationModel.h
+++ b/Engine/Plugins/Runtime/NetworkPrediction/Source/NetworkPrediction/Public/NetworkSimulationModel.h
@@ -21,7 +21,8 @@
 // ----------------------------------------------------------------------------------------------------------------------------------------------------------------------
 
 template <
-	typename T,											// Final Driver class
+	typename TSimulation,								// Final Simulation class. Used to call static T::Update function.
+	typename TInDriver,
 	typename TUserBufferTypes,							// The user types (input, sync, aux, debug). Note this gets wrapped in TInternalBufferTypes internally.
 	typename InTTickSettings=TNetworkSimTickSettings<>, // Defines global rules about time keeping and ticking
 
@@ -34,10 +35,12 @@ template <
 	typename TRepProxyReplay =		TReplicator_Sequence	<TInternalBufferTypes<TUserBufferTypes, InTTickSettings>,	InTTickSettings, ENetworkSimBufferTypeId::Sync,  3>,
 	typename TRepProxyDebug =		TReplicator_Debug		<TInternalBufferTypes<TUserBufferTypes, InTTickSettings>,	InTTickSettings>
 >
-class TNetworkedSimulationModel : public IReplicationProxy
+class TNetworkedSimulationModel : public INetworkSimulationModel
 {
 public:
 
+	using TDriver = TInDriver;
+
 	using TBufferTypes = TInternalBufferTypes<TUserBufferTypes, InTTickSettings>;
 	using TTickSettings = InTTickSettings;
 
@@ -55,21 +58,16 @@ public:
 		virtual FString GetDebugName() const = 0; // Used for debugging. Recommended to emit the simulation name and the actor name/role.
 
 		virtual void InitSyncState(TSyncState& OutSyncState) const = 0;	// Called to create initial value of the sync state.
-		virtual void ProduceInput(const TSimTime&, typename TUserBufferTypes::TInputCmd&) = 0; // Called when the sim is ready to process new local input
+		virtual void ProduceInput(const TRealTime, typename TUserBufferTypes::TInputCmd&) = 0; // Called when the sim is ready to process new local input
 		virtual void FinalizeFrame(const TSyncState& SyncState) = 0; // Called from the Network Sim at the end of the sim frame when there is new sync data.
 	};
-
-	struct FTickParameters : public FNetSimTickParametersBase
+	
+	TNetworkedSimulationModel(TDriver* InDriver)
 	{
-		FTickParameters() { }
-		FTickParameters(typename TTickSettings::TRealTime InRealTime, class AActor* Actor=nullptr) : FNetSimTickParametersBase(Actor), LocalDeltaTimeSeconds(InRealTime) {}
-
-		// Local fraem delta time seconds. Just passed in from ::Tick
-		typename TTickSettings::TRealTime LocalDeltaTimeSeconds = 0.f;
-	};
+		Driver = InDriver;
+	}
 
-	template<typename TDriver>
-	void Tick(TDriver* Driver, const FTickParameters& Parameters)
+	void Tick(const FNetSimTickParameters& Parameters) final override
 	{
 		// Update previous DebugState based on what we (might) have sent *after* our last Tick 
 		// (property replication and ServerRPC get sent after the tick, rather than forcing awkward callback into the NetSim post replication, we can check it here)
@@ -107,26 +105,6 @@ public:
 			}
 		}
 
-		// --------------------------------------------------------------------------------------------------------------------------
-		//	Reconcile
-		//	This will eventually be called outside the Tick loop, only after processing a network bunch
-		//	Reconcile is about "making things right" after a network update. We are not processing "more" simulation yet.
-		// --------------------------------------------------------------------------------------------------------------------------
-		switch (Parameters.Role)
-		{
-			case ROLE_Authority:
-				RepProxy_ServerRPC.template Reconcile<T, TDriver>(Driver, Buffers, TickInfo);
-			break;
-
-			case ROLE_AutonomousProxy:
-				RepProxy_Autonomous.template Reconcile<T, TDriver>(Driver, Buffers, TickInfo);
-			break;
-
-			case ROLE_SimulatedProxy:
-				RepProxy_Simulated.template Reconcile<T, TDriver>(Driver, Buffers, TickInfo);
-			break;
-		}
-
 		// ----------------------------------------------------------------------------------------------------------------
 		//	PreSimTick
 		//	This is the beginning of a new frame. PreSimTick will decide if we should take Parameters.LocalDeltaTimeSeconds
@@ -135,15 +113,15 @@ public:
 		switch (Parameters.Role)
 		{
 			case ROLE_Authority:
-				RepProxy_ServerRPC.template PreSimTick<T, TDriver, FTickParameters>(Driver, Buffers, TickInfo, Parameters);
+				RepProxy_ServerRPC.template PreSimTick<TSimulation, TDriver>(Driver, Buffers, TickInfo, Parameters);
 			break;
 
 			case ROLE_AutonomousProxy:
-				RepProxy_Autonomous.template PreSimTick<T, TDriver, FTickParameters>(Driver, Buffers, TickInfo, Parameters);
+				RepProxy_Autonomous.template PreSimTick<TSimulation, TDriver>(Driver, Buffers, TickInfo, Parameters);
 			break;
 
 			case ROLE_SimulatedProxy:
-				RepProxy_Simulated.template PreSimTick<T, TDriver, FTickParameters>(Driver, Buffers, TickInfo, Parameters);
+				RepProxy_Simulated.template PreSimTick<TSimulation, TDriver>(Driver, Buffers, TickInfo, Parameters);
 			break;
 		}
 
@@ -214,8 +192,7 @@ public:
 						}
 					
 						TAuxState AuxState; // Temp: aux buffer not implemented yet
-
-						T::Update(Driver, NextCmd->GetFrameDeltaTime(), *NextCmd, *PrevSyncState, *NextSyncState, AuxState);
+						TSimulation::Update(Driver, NextCmd->GetFrameDeltaTime().ToRealTimeSeconds(), *NextCmd, *PrevSyncState, *NextSyncState, AuxState);
 					
 						TickInfo.IncrementTotalProcessedSimulationTime(NextCmd->GetFrameDeltaTime(), Keyframe);
 						TickInfo.LastProcessedInputKeyframe = Keyframe;
@@ -230,18 +207,25 @@ public:
 					break;
 				}
 			}
+		}
 
-			// FIXME: this needs to be sorted out. We really want to check if there is new sync state and then call this here.
-			// Call into the driver to sync to the latest state if we processed any input
-			//if (NumProcessed > 0)	
-			{
-				//check(SyncBuffer.GetNumValidElements() > 0);
-				if (Buffers.Sync.GetNumValidElements() > 0)
-				{
-					Driver->FinalizeFrame(*Buffers.Sync.GetElementFromHead(0));
-				}
-			}
+		// -------------------------------------------------------------------------------------------------------------------------------------------------
+		//												Post Sim Tick: finalize the frame
+		// -------------------------------------------------------------------------------------------------------------------------------------------------
 
+		switch (Parameters.Role)
+		{
+			case ROLE_Authority:
+				RepProxy_ServerRPC.template PostSimTick<TDriver>(Driver, Buffers, TickInfo, Parameters);
+			break;
+
+			case ROLE_AutonomousProxy:
+				RepProxy_Autonomous.template PostSimTick<TDriver>(Driver, Buffers, TickInfo, Parameters);
+			break;
+
+			case ROLE_SimulatedProxy:
+				RepProxy_Simulated.template PostSimTick<TDriver>(Driver, Buffers, TickInfo, Parameters);
+			break;
 		}
 
 		// -------------------------------------------------------------------------------------------------------------------------------------------------
@@ -263,9 +247,32 @@ public:
 			HistoricData->Sync.CopyAndMerge(Buffers.Sync);
 			HistoricData->Aux.CopyAndMerge(Buffers.Aux);
 		}
-	}	
+	}
+
+	virtual void Reconcile(const ENetRole Role) final override
+	{
+		// --------------------------------------------------------------------------------------------------------------------------
+		//	Reconcile
+		//	This will eventually be called outside the Tick loop, only after processing a network bunch
+		//	Reconcile is about "making things right" after a network update. We are not processing "more" simulation yet.
+		// --------------------------------------------------------------------------------------------------------------------------
+		switch (Role)
+		{
+			case ROLE_Authority:
+				RepProxy_ServerRPC.template Reconcile<TSimulation, TDriver>(Driver, Buffers, TickInfo);
+			break;
+
+			case ROLE_AutonomousProxy:
+				RepProxy_Autonomous.template Reconcile<TSimulation, TDriver>(Driver, Buffers, TickInfo);
+			break;
+
+			case ROLE_SimulatedProxy:
+				RepProxy_Simulated.template Reconcile<TSimulation, TDriver>(Driver, Buffers, TickInfo);
+			break;
+		}
+	}
 	
-	void InitializeForNetworkRole(const ENetRole Role, const FNetworkSimulationModelInitParameters& Parameters)
+	void InitializeForNetworkRole(const ENetRole Role, const FNetworkSimulationModelInitParameters& Parameters) final override
 	{
 		Buffers.Input.SetBufferSize(Parameters.InputBufferSize);
 		Buffers.Sync.SetBufferSize(Parameters.SyncedBufferSize);
@@ -289,7 +296,7 @@ public:
 		*Buffers.Input.GetWriteNext() = TInputCmd();
 	}
 
-	void NetSerializeProxy(EReplicationProxyTarget Target, const FNetSerializeParams& Params)
+	void NetSerializeProxy(EReplicationProxyTarget Target, const FNetSerializeParams& Params) final override
 	{
 		switch(Target)
 		{
@@ -315,7 +322,7 @@ public:
 		};
 	}
 
-	int32 GetProxyDirtyCount(EReplicationProxyTarget Target)
+	int32 GetProxyDirtyCount(EReplicationProxyTarget Target) final override
 	{
 		switch(Target)
 		{
@@ -336,7 +343,8 @@ public:
 			return 0;
 		};
 	}
-	
+
+	TDriver* Driver = nullptr;	
 	TSimulationTickState<TTickSettings> TickInfo;	// Manages simulation time and what inputs we are processed
 
 	TNetworkSimBufferContainer<TBufferTypes> Buffers;
@@ -353,22 +361,24 @@ public:
 	// completely be tracked at the driver level, but that will also push more boilerplate to that layer for users.
 	// ------------------------------------------------------------------
 
-	void SetDesiredServerRPCSendFrequency(float DesiredHz) { ServerRPCThresholdTimeSeconds = 1.f / DesiredHz; }
-	bool ShouldSendServerRPC(ENetRole OwnerRole, float DeltaTimeSeconds)
+	void SetDesiredServerRPCSendFrequency(float DesiredHz) final override { ServerRPCThresholdTimeSeconds = 1.f / DesiredHz; }
+	bool ShouldSendServerRPC(float DeltaTimeSeconds) final override
 	{
 		// Don't allow a large delta time to pollute the accumulator
 		const float CappedDeltaTimeSeconds = FMath::Min<float>(DeltaTimeSeconds, ServerRPCThresholdTimeSeconds);
-		if (OwnerRole == ROLE_AutonomousProxy)
+		ServerRPCAccumulatedTimeSeconds += DeltaTimeSeconds;
+		if (ServerRPCAccumulatedTimeSeconds >= ServerRPCThresholdTimeSeconds)
 		{
-			ServerRPCAccumulatedTimeSeconds += DeltaTimeSeconds;
-			if (ServerRPCAccumulatedTimeSeconds >= ServerRPCThresholdTimeSeconds)
-			{
-				ServerRPCAccumulatedTimeSeconds -= ServerRPCThresholdTimeSeconds;
-				return true;
-			}
+			ServerRPCAccumulatedTimeSeconds -= ServerRPCThresholdTimeSeconds;
+			return true;
 		}
+
 		return false;
 	}
+
+	// Simulation class should have a static const FName GroupName member
+	FName GetSimulationGroupName() const final override { return TSimulation::GroupName; }
+
 private:
 	float ServerRPCAccumulatedTimeSeconds = 0.f;
 	float ServerRPCThresholdTimeSeconds = 1.f / 999.f; // Default is to send at a max of 999hz. This part of the system needs to be build out more (better handling of super high FPS clients and fixed rate servers)
diff --git a/Engine/Plugins/Runtime/NetworkPrediction/Source/NetworkPrediction/Public/NetworkSimulationModelDebugger.h b/Engine/Plugins/Runtime/NetworkPrediction/Source/NetworkPrediction/Public/NetworkSimulationModelDebugger.h
--- a/Engine/Plugins/Runtime/NetworkPrediction/Source/NetworkPrediction/Public/NetworkSimulationModelDebugger.h
+++ b/Engine/Plugins/Runtime/NetworkPrediction/Source/NetworkPrediction/Public/NetworkSimulationModelDebugger.h
@@ -70,8 +70,8 @@ struct NETWORKPREDICTION_API FNetworkSimulationModelDebuggerManager: public FTic
 	//	Outside API (registration, console commands, draw services, etc)
 	// ---------------------------------------------------------------------------------------------------------------------------------------
 
-	template <typename T, typename TDriver>
-	void RegisterNetworkSimulationModel(T* NetworkSim, TDriver* Driver, AActor* OwningActor, FString DebugName);
+	template <typename T>
+	void RegisterNetworkSimulationModel(T* NetworkSim, AActor* OwningActor);
 
 	void SetDebuggerActive(AActor* OwningActor, bool InActive)
 	{
@@ -333,18 +333,16 @@ private:
 	TWeakObjectPtr<UCanvas> LastCanvas;
 };
 
-template <typename TNetSimModel, typename TDriver>
+template <typename TNetSimModel>
 struct TNetworkSimulationModelDebugger : public INetworkSimulationModelDebugger
 {
 	using TSimTime = typename TNetSimModel::TSimTime;
 	using TDebugState = typename TNetSimModel::TDebugState;
 
-	TNetworkSimulationModelDebugger(TNetSimModel* InNetSim, TDriver* InDriver, AActor* OwningActor, FString InDebugName)
+	TNetworkSimulationModelDebugger(TNetSimModel* InNetSim, AActor* OwningActor)
 	{
 		NetworkSim = InNetSim;
-		Driver = InDriver;
 		WeakOwningActor = OwningActor;
-		DebugName = InDebugName;
 	}
 
 	~TNetworkSimulationModelDebugger()
@@ -662,15 +660,15 @@ struct TNetworkSimulationModelDebugger : public INetworkSimulationModelDebugger
 
 		if (auto* LatestSync = NetworkSim->Buffers.Sync.GetElementFromHead(0))
 		{
-			LatestSync->VisualLog( FVisualLoggingParameters(EVisualLoggingContext::LastPredicted, NetworkSim->Buffers.Sync.GetHeadKeyframe(), EVisualLoggingLifetime::Transient), Driver, Driver);
+			LatestSync->VisualLog( FVisualLoggingParameters(EVisualLoggingContext::LastPredicted, NetworkSim->Buffers.Sync.GetHeadKeyframe(), EVisualLoggingLifetime::Transient), NetworkSim->Driver, NetworkSim->Driver);
 		}
 
 		FStuff ServerPIEStuff = GetServerPIEStuff();
-		if (ServerPIEStuff.Driver && ServerPIEStuff.NetworkSim)
+		if (ServerPIEStuff.NetworkSim)
 		{
 			if (auto* ServerLatestSync = ServerPIEStuff.NetworkSim->Buffers.Sync.GetElementFromHead(0))
 			{
-				ServerLatestSync->VisualLog( FVisualLoggingParameters(EVisualLoggingContext::CurrentServerPIE, ServerPIEStuff.NetworkSim->Buffers.Sync.GetHeadKeyframe(), EVisualLoggingLifetime::Transient), ServerPIEStuff.Driver, Driver);
+				ServerLatestSync->VisualLog( FVisualLoggingParameters(EVisualLoggingContext::CurrentServerPIE, ServerPIEStuff.NetworkSim->Buffers.Sync.GetHeadKeyframe(), EVisualLoggingLifetime::Transient), ServerPIEStuff.NetworkSim->Driver, NetworkSim->Driver);
 			}
 		}
 
@@ -679,7 +677,7 @@ struct TNetworkSimulationModelDebugger : public INetworkSimulationModelDebugger
 			if (auto* SyncState = NetworkSim->Buffers.Sync.FindElementByKeyframe(Keyframe))
 			{
 				const EVisualLoggingContext Context = (Keyframe == NetworkSim->RepProxy_Autonomous.GetLastSerializedKeyframe()) ? EVisualLoggingContext::LastConfirmed : EVisualLoggingContext::OtherPredicted;
-				SyncState->VisualLog( FVisualLoggingParameters(Context, NetworkSim->Buffers.Sync.GetHeadKeyframe(), EVisualLoggingLifetime::Transient), Driver, Driver);
+				SyncState->VisualLog( FVisualLoggingParameters(Context, NetworkSim->Buffers.Sync.GetHeadKeyframe(), EVisualLoggingLifetime::Transient), NetworkSim->Driver, NetworkSim->Driver);
 			}
 		}
 	}
@@ -688,30 +686,25 @@ struct TNetworkSimulationModelDebugger : public INetworkSimulationModelDebugger
 	struct FStuff
 	{
 		TNetSimModel* NetworkSim = nullptr;
-		TDriver* Driver = nullptr;
 	};
 
 	FStuff GetStuff()
 	{
-		return {NetworkSim, Driver};
+		return {NetworkSim};
 	}
 
 	TFunction< FStuff() > GetServerPIEStuff;
 
 private:
-
 	
 	TWeakObjectPtr<AActor>	WeakOwningActor;
-	FString DebugName;
-
 	TNetSimModel* NetworkSim = nullptr;
-	TDriver* Driver = nullptr;
 };
 
-template <typename T, typename TDriver>
-void FNetworkSimulationModelDebuggerManager::RegisterNetworkSimulationModel(T* NetworkSim, TDriver* Driver, AActor* OwningActor, FString DebugName)
+template <typename T>
+void FNetworkSimulationModelDebuggerManager::RegisterNetworkSimulationModel(T* NetworkSim, AActor* OwningActor)
 {
-	TNetworkSimulationModelDebugger<T, TDriver>* Debugger = new TNetworkSimulationModelDebugger<T, TDriver>(NetworkSim, Driver, OwningActor, DebugName);
+	TNetworkSimulationModelDebugger<T>* Debugger = new TNetworkSimulationModelDebugger<T>(NetworkSim, OwningActor);
 	DebuggerMap.Add( TWeakObjectPtr<AActor>(OwningActor), Debugger );
 
 	// Gross stuff so that the debugger can find the ServerPIE equiv
@@ -720,8 +713,8 @@ void FNetworkSimulationModelDebuggerManager::RegisterNetworkSimulationModel(T* N
 	{
 		if (AActor* ServerOwner = Cast<AActor>(FindReplicatedObjectOnPIEServer(WeakOwner.Get())))
 		{
-			return ((TNetworkSimulationModelDebugger<T, TDriver>*)DebuggerMap.FindRef(TWeakObjectPtr<AActor>(ServerOwner)))->GetStuff();
+			return ((TNetworkSimulationModelDebugger<T>*)DebuggerMap.FindRef(TWeakObjectPtr<AActor>(ServerOwner)))->GetStuff();
 		}
-		return typename TNetworkSimulationModelDebugger<T, TDriver>::FStuff();
+		return typename TNetworkSimulationModelDebugger<T>::FStuff();
 	};
 }
\ No newline at end of file
diff --git a/Engine/Plugins/Runtime/NetworkPrediction/Source/NetworkPrediction/Public/NetworkSimulationModelInterpolator.h b/Engine/Plugins/Runtime/NetworkPrediction/Source/NetworkPrediction/Public/NetworkSimulationModelInterpolator.h
new file mode 100644
--- /dev/null
+++ b/Engine/Plugins/Runtime/NetworkPrediction/Source/NetworkPrediction/Public/NetworkSimulationModelInterpolator.h
@@ -0,0 +1,229 @@
+
+// Copyright 1998-2019 Epic Games, Inc. All Rights Reserved.
+
+#pragma once
+
+#include "NetworkSimulationModelCVars.h"
+#include "VisualLogger/VisualLogger.h"
+#include "GameFramework/Actor.h"
+#include "NetworkPredictionTypes.h"
+#include "Engine/World.h"
+
+DEFINE_LOG_CATEGORY_STATIC(LogNetInterpolation, Log, All);
+
+namespace NetworkInterpolationDebugCVars
+{
+	NETSIM_DEVCVAR_SHIPCONST_INT(Disable, 0, "ni.Disable", "Disables Network Interpolation");
+	NETSIM_DEVCVAR_SHIPCONST_INT(VLog, 0, "ni.VLog", "Enables Network Interpolation VLog ");
+
+	NETSIM_DEVCVAR_SHIPCONST_FLOAT(WaitSlack, 0.05, "ni.WaitSlack", "How much slack to wait for when waiting");
+
+	NETSIM_DEVCVAR_SHIPCONST_FLOAT(CatchUpThreshold, 0.200, "ni.CatchUpThreshold", "When we start catching up (seconds from head)");
+	NETSIM_DEVCVAR_SHIPCONST_FLOAT(CatchUpGoal, 0.100, "ni.CatchUpGoal", "When we stop cathcing up (seconds from head)");
+	NETSIM_DEVCVAR_SHIPCONST_FLOAT(CatchUpFactor, 1.25, "ni.CatchUpFactor", "Factor we use to catch up");
+}
+
+template<typename TBufferTypes, typename TTickSettings>
+struct TInterpolator
+{
+	using TSimTime = TNetworkSimTime<TTickSettings>;
+	using TRealTime = typename TTickSettings::TRealTime;
+	using TSyncState = typename TBufferTypes::TSyncState;
+
+	bool bEnableVisualLog = true;
+
+	template<typename TDriver>
+	void PostSimTick(TDriver* Driver, const TNetworkSimBufferContainer<TBufferTypes>& Buffers, const TSimulationTickState<TTickSettings>& TickInfo, const FNetSimTickParameters& TickParameters)
+	{
+		const bool bDoVLog = NetworkInterpolationDebugCVars::VLog() && bEnableVisualLog;
+		const float DeltaSeconds = TickParameters.LocalDeltaTimeSeconds;
+
+		const UObject* LogOwner = Driver->GetVLogOwner();
+
+		//
+
+		if (NetworkInterpolationDebugCVars::Disable() > 0)
+		{			
+			if (const TSyncState* HeadState = Buffers.Sync.GetElementFromHead(0))
+			{
+				Driver->FinalizeFrame(*HeadState);
+			}
+			return;
+		}
+		
+		if (TickInfo.SimulationTimeBuffer.GetNumValidElements() <= 1)
+		{
+			// Cant interpolate yet	
+			return;
+		}
+
+
+		auto& SimulationTimeBuffer = TickInfo.SimulationTimeBuffer;
+
+		// Starting off: start at the tail end
+		if (InterpolationTime <= 0.f)
+		{
+			InterpolationTime = SimulationTimeBuffer.GetElementFromTail(0)->ToRealTimeSeconds();
+			InterpolationKeyframe = SimulationTimeBuffer.GetTailKeyframe();
+			InterpolationState = *Buffers.Sync.GetElementFromTail(0);
+		}
+
+		// Wait if we were too far ahead
+		if (WaitUntilTime > 0.f)
+		{
+			if (WaitUntilTime <= SimulationTimeBuffer.GetElementFromHead(0)->ToRealTimeSeconds())
+			{
+				// done waiting, we can continue
+				WaitUntilTime = 0.f;
+				UE_VLOG(LogOwner, LogNetInterpolation, Log, TEXT("Done Waiting! Head: %s"), *LexToString(SimulationTimeBuffer.GetElementFromHead(0)->ToRealTimeSeconds()));
+			}
+			else
+			{
+				if (bDoVLog)
+				{
+					// Still waiting, return
+					UE_VLOG(LogOwner, LogNetInterpolation, Log, TEXT("Still Waiting! %s < %s"), *LexToString(WaitUntilTime), *LexToString(SimulationTimeBuffer.GetElementFromHead(0)->ToRealTimeSeconds()));
+
+					const FName LocalInterpolationTimeName("Local Interpolation Time");
+					FVector2D LocalTimeVsInterpolationTime(LogOwner->GetWorld()->GetTimeSeconds(), (int32)(InterpolationTime * 1000.f));
+					UE_VLOG_HISTOGRAM(LogOwner, LogNetInterpolation, Log, "ServerSimulationTimeGraph", LocalInterpolationTimeName, LocalTimeVsInterpolationTime);
+				}
+				return;
+			}
+		}
+
+		EVisualLoggingContext LoggingContext = EVisualLoggingContext::InterpolationLatest;
+
+		// Calc new interpolation time
+		TRealTime NewInterpolationTime = InterpolationTime;
+		{
+			TRealTime Step = DeltaSeconds;
+			
+			// Speed up if we are too far behind
+			TRealTime CatchUpThreshold = SimulationTimeBuffer.GetElementFromHead(0)->ToRealTimeSeconds() - NetworkInterpolationDebugCVars::CatchUpThreshold();
+			if (CatchUpUntilTime <= 0.f && InterpolationTime < CatchUpThreshold)
+			{
+				CatchUpUntilTime = SimulationTimeBuffer.GetElementFromHead(0)->ToRealTimeSeconds() - NetworkInterpolationDebugCVars::CatchUpGoal();
+			}
+
+			if (CatchUpUntilTime > 0.f)
+			{
+				if (InterpolationTime  < CatchUpUntilTime)
+				{
+					Step *= NetworkInterpolationDebugCVars::CatchUpFactor();
+					LoggingContext = EVisualLoggingContext::InterpolationSpeedUp;
+
+					UE_VLOG(LogOwner, LogNetInterpolation, Log, TEXT("Catching up! %s < %s"), *LexToString(InterpolationTime), *LexToString(CatchUpUntilTime));
+					//!UE_LOG(LogNetInterpolation, Warning, TEXT("Catching up! %s < %s"), *LexToString(InterpolationTime), *LexToString(CatchUpUntilTime));
+				}
+				else
+				{
+					CatchUpUntilTime = 0.f;
+				}
+			}
+
+			NewInterpolationTime += Step;
+
+			// Did this put us too far ahead, and now we need to start waiting?
+			if (NewInterpolationTime > SimulationTimeBuffer.GetElementFromHead(0)->ToRealTimeSeconds())
+			{
+				UE_VLOG(LogOwner, LogNetInterpolation, Log, TEXT("Too far ahead! Starting to wait! Head: %s"), *LexToString(SimulationTimeBuffer.GetElementFromHead(0)->ToRealTimeSeconds()));
+				WaitUntilTime = SimulationTimeBuffer.GetElementFromHead(0)->ToRealTimeSeconds() + NetworkInterpolationDebugCVars::WaitSlack();
+				NewInterpolationTime = SimulationTimeBuffer.GetElementFromHead(0)->ToRealTimeSeconds();
+			}
+		}
+
+		// Find "To" keyframe
+		const TSyncState* ToState = nullptr;
+		TRealTime ToTime = 0.f;
+
+		for (int32 Keyframe = SimulationTimeBuffer.GetTailKeyframe(); Keyframe <= SimulationTimeBuffer.GetHeadKeyframe(); ++Keyframe)
+		{
+			TSimTime ElementSimTime = *SimulationTimeBuffer.FindElementByKeyframe(Keyframe);
+			if (NewInterpolationTime <= ElementSimTime.ToRealTimeSeconds())
+			{
+				InterpolationKeyframe = Keyframe;
+				ToTime = ElementSimTime.ToRealTimeSeconds();
+				ToState = Buffers.Sync.FindElementByKeyframe(Keyframe);
+				break;
+			}
+		}
+
+		if (ensure(ToState))
+		{
+			const TRealTime FromRealTime = InterpolationTime;
+			const TRealTime ToRealTime = ToTime;
+			const TRealTime InterpolationInterval = ToRealTime - FromRealTime;
+		
+			if (ensure(FMath::Abs(InterpolationInterval) > 0.f))
+			{
+				const float InterpolationPCT = (NewInterpolationTime - FromRealTime) / InterpolationInterval;
+				ensureMsgf(InterpolationPCT >= 0.f && InterpolationPCT <= 1.f, TEXT("Calculated InterpolationPCT not in expected range. NewInterpolationTime: %s. From: %s. To: %s"), *LexToString(NewInterpolationTime), *LexToString(FromRealTime), *LexToString(ToRealTime));
+
+				TSyncState NewInterpolatedState;
+				TSyncState::Interpolate(InterpolationState, *ToState, InterpolationPCT, NewInterpolatedState);
+
+				Driver->FinalizeFrame(NewInterpolatedState);
+				
+				if (bDoVLog)
+				{
+					UE_VLOG(LogOwner, LogNetInterpolation, Log, TEXT("%s - %s - %s.  InterpolationPCT: %f"), *LexToString(FromRealTime), *LexToString(NewInterpolationTime), *LexToString(ToRealTime), InterpolationPCT);
+
+					// Graph Interpolation Time vs Buffer Head/Tail times
+					const FName ServerSimulationGraphName("ServerSimulationTimeGraph");
+					const FName ServerSimTimeName("Server Simulation Time");
+					FVector2D LocalTimeVsServerSimTime(LogOwner->GetWorld()->GetTimeSeconds(), (int32)TickInfo.SimulationTimeBuffer.GetElementFromHead(0)->ToRealTimeMS());
+					UE_VLOG_HISTOGRAM(LogOwner, LogNetInterpolation, Log, ServerSimulationGraphName, ServerSimTimeName, LocalTimeVsServerSimTime);
+
+					const FName BufferTailSimTimeName("Buffer Tail Simulation Time");
+					FVector2D LocalTimeVsBufferTailSim(LogOwner->GetWorld()->GetTimeSeconds(), (int32)TickInfo.SimulationTimeBuffer.GetElementFromTail(0)->ToRealTimeMS());
+					UE_VLOG_HISTOGRAM(LogOwner, LogNetInterpolation, Log, ServerSimulationGraphName, BufferTailSimTimeName, LocalTimeVsBufferTailSim);
+
+					const FName LocalInterpolationTimeName("Local Interpolation Time");
+					FVector2D LocalTimeVsInterpolationTime(LogOwner->GetWorld()->GetTimeSeconds(), (int32)(NewInterpolationTime * 1000.f));
+					UE_VLOG_HISTOGRAM(LogOwner, LogNetInterpolation, Log, ServerSimulationGraphName, LocalInterpolationTimeName, LocalTimeVsInterpolationTime);
+					
+					FVector2D LocalTimeVsCatchUpThreshold(LogOwner->GetWorld()->GetTimeSeconds(), (SimulationTimeBuffer.GetElementFromHead(0)->ToRealTimeSeconds() - NetworkInterpolationDebugCVars::CatchUpThreshold()) * 1000.f);
+					UE_VLOG_HISTOGRAM(LogOwner, LogNetInterpolation, Log, ServerSimulationGraphName, "Catch Up Threshold", LocalTimeVsCatchUpThreshold);
+
+					FVector2D LocalTimeVsCatchUpGoal(LogOwner->GetWorld()->GetTimeSeconds(), (SimulationTimeBuffer.GetElementFromHead(0)->ToRealTimeSeconds() - NetworkInterpolationDebugCVars::CatchUpGoal()) * 1000.f);
+					UE_VLOG_HISTOGRAM(LogOwner, LogNetInterpolation, Log, ServerSimulationGraphName, "Catch Up Goal", LocalTimeVsCatchUpGoal);
+
+					// VLog the actual motion states
+					const TSyncState* DebugTail = Buffers.Sync.GetElementFromTail(0);
+					const TSyncState* DebugHead = Buffers.Sync.GetElementFromHead(0);
+
+					DebugTail->VisualLog( FVisualLoggingParameters(EVisualLoggingContext::InterpolationBufferTail, Buffers.Sync.GetTailKeyframe(), EVisualLoggingLifetime::Transient), Driver, Driver );
+					DebugHead->VisualLog( FVisualLoggingParameters(EVisualLoggingContext::InterpolationBufferHead, Buffers.Sync.GetHeadKeyframe(), EVisualLoggingLifetime::Transient), Driver, Driver );
+
+					InterpolationState.VisualLog( FVisualLoggingParameters(EVisualLoggingContext::InterpolationFrom, InterpolationKeyframe, EVisualLoggingLifetime::Transient), Driver, Driver );
+					ToState->VisualLog( FVisualLoggingParameters(EVisualLoggingContext::InterpolationTo, InterpolationKeyframe, EVisualLoggingLifetime::Transient), Driver, Driver );
+
+					NewInterpolatedState.VisualLog( FVisualLoggingParameters(LoggingContext, InterpolationKeyframe, EVisualLoggingLifetime::Transient), Driver, Driver );
+				}
+
+
+				InterpolationState = NewInterpolatedState;
+				InterpolationTime = NewInterpolationTime;
+			}
+		}
+	}
+
+private:
+
+	TRealTime InterpolationTime = 0.f; // SimTime we are currently interpolating at
+	int32 InterpolationKeyframe = INDEX_NONE; // Keyframe we are currently/last interpolated at
+
+	TSyncState InterpolationState;
+
+	TRealTime WaitUntilTime = 0.f;	
+
+	TRealTime CatchUpUntilTime = 0.f;
+
+
+	TRealTime DynamicBufferedTime = 1/60.f; // SimTime we are currently interpolating at
+	TRealTime DynamicBufferedTimeStep = 1/60.f;
+
+	TRealTime	MinBufferedTime = 1/120.f;
+	TRealTime	MaxBufferedTime = 1.f;
+};
\ No newline at end of file
diff --git a/Engine/Plugins/Runtime/NetworkPrediction/Source/NetworkPrediction/Public/NetworkSimulationModelReplicators.h b/Engine/Plugins/Runtime/NetworkPrediction/Source/NetworkPrediction/Public/NetworkSimulationModelReplicators.h
--- a/Engine/Plugins/Runtime/NetworkPrediction/Source/NetworkPrediction/Public/NetworkSimulationModelReplicators.h
+++ b/Engine/Plugins/Runtime/NetworkPrediction/Source/NetworkPrediction/Public/NetworkSimulationModelReplicators.h
@@ -4,6 +4,7 @@
 
 #include "Net/UnrealNetwork.h" // For MakeRelative
 #include "NetworkSimulationModelCVars.h"
+#include "NetworkSimulationModelInterpolator.h"
 
 // ----------------------------------------------------------------------------------------------------------------------------------------------------------------------
 //	CVars and compile time constants
@@ -67,8 +68,13 @@ struct TReplicatorEmpty
 	template<typename T, typename TDriver>
 	void Reconcile(TDriver* Driver, TNetworkSimBufferContainer<TBufferTypes>& Buffers, TSimulationTickState<TTickSettings>& TickInfo) { }
 
-	template<typename T, typename TDriver, typename TTickParameters>
-	void PreSimTick(TDriver* Driver, TNetworkSimBufferContainer<TBufferTypes>& Buffers, TSimulationTickState<TTickSettings>& TickInfo, const TTickParameters& TickParameters) { }
+	// Called prior to input processing. This function must updated TickInfo to allow simulation time (from TickParameters) and to possibly get new input.
+	template<typename T, typename TDriver>
+	void PreSimTick(TDriver* Driver, TNetworkSimBufferContainer<TBufferTypes>& Buffers, TSimulationTickState<TTickSettings>& TickInfo, const FNetSimTickParameters& TickParameters) { }
+
+	// Called after input processing. Should finalize the frame and do any smoothing/interpolation. This function is not allowed to modify the buffers or tick state, or even call the simulation/Update function.
+	template<typename TDriver>
+	void PostSimTick(TDriver* Driver, const TNetworkSimBufferContainer<TBufferTypes>& Buffers, const TSimulationTickState<TTickSettings>& TickInfo, const FNetSimTickParameters& TickParameters) { }
 };
 
 // This is the "templated base class" for replicators but is not required (i.e., this not an official interface used by TNetworkedSimulation model. Just a base implementation you can start with)
@@ -88,12 +94,23 @@ struct TReplicatorBase
 	void Reconcile(TDriver* Driver, TNetworkSimBufferContainer<TBufferTypes>& Buffers, TSimulationTickState<TTickSettings>& TickInfo) { }
 
 	// Called prior to input processing. This function must updated TickInfo to allow simulation time (from TickParameters) and to possibly get new input.
-	template<typename T, typename TDriver, typename TTickParameters>
-	void PreSimTick(TDriver* Driver, TNetworkSimBufferContainer<TBufferTypes>& Buffers, TSimulationTickState<TTickSettings>& TickInfo, const TTickParameters& TickParameters)
+	template<typename T, typename TDriver>
+	void PreSimTick(TDriver* Driver, TNetworkSimBufferContainer<TBufferTypes>& Buffers, TSimulationTickState<TTickSettings>& TickInfo, const FNetSimTickParameters& TickParameters)
 	{
 		// Accumulate local delta time into TickInfo
 		TickInfo.GiveSimulationTime(TickParameters.LocalDeltaTimeSeconds);
 	}
+
+	// Called after input processing. Should finalize the frame and do any smoothing/interpolation. This function is not allowed to modify the buffers or tick state, or even call the simulation/Update function.
+	template<typename TDriver>
+	void PostSimTick(TDriver* Driver, const TNetworkSimBufferContainer<TBufferTypes>& Buffers, const TSimulationTickState<TTickSettings>& TickInfo, const FNetSimTickParameters& TickParameters)
+	{
+		// Sync to latest frame if there is any
+		if (Buffers.Sync.GetNumValidElements() > 0)
+		{
+			Driver->FinalizeFrame(*Buffers.Sync.GetElementFromHead(0));
+		}
+	}
 };
 
 // -------------------------------------------------------------------------------------------------------
@@ -237,8 +254,8 @@ struct TReplicator_Server : public TBase
 		}
 	}
 
-	template<typename T, typename TDriver, typename TTickParameters>
-	void PreSimTick(TDriver* Driver, TNetworkSimBufferContainer<TBufferTypes>& Buffers, TSimulationTickState<TTickSettings>& TickInfo, const TTickParameters& TickParameters)
+	template<typename T, typename TDriver>
+	void PreSimTick(TDriver* Driver, TNetworkSimBufferContainer<TBufferTypes>& Buffers, TSimulationTickState<TTickSettings>& TickInfo, const FNetSimTickParameters& TickParameters)
 	{
 		TickInfo.GiveSimulationTime(TickParameters.LocalDeltaTimeSeconds);
 
@@ -251,7 +268,7 @@ struct TReplicator_Server : public TBase
 				{
 					*InputCmd = TInputCmd();
 					InputCmd->SetFrameDeltaTime(DeltaSimTime);
-					Driver->ProduceInput(DeltaSimTime, *InputCmd);
+					Driver->ProduceInput(DeltaSimTime.ToRealTimeSeconds(), *InputCmd);
 					TickInfo.MaxAllowedInputKeyframe = Buffers.Input.GetHeadKeyframe();
 				}
 			}
@@ -270,6 +287,8 @@ struct TReplicator_Simulated : public TBase
 	// Instance flag for enabling simulated extrapolation
 	bool bAllowSimulatedExtrapolation = true;
 
+	TInterpolator<TBufferTypes, TTickSettings> Interpolator;
+
 	TNetworkSimTime<TTickSettings> GetLastSerializedSimulationTime() const { return LastSerializedSimulationTime; }
 	
 	int32 GetProxyDirtyCount(TNetworkSimBufferContainer<TBufferTypes>& Buffers) const
@@ -409,7 +428,7 @@ struct TReplicator_Simulated : public TBase
 				//UE_LOG(LogNetworkSim, Warning, TEXT("   >>> Reconcile Update: %s"), *DeltaSimTime.ToString());
 
 				// Do the actual update
-				T::Update(Driver, NewCmd->GetFrameDeltaTime(), *NewCmd, *PrevSyncState, *NextSyncState, Junk);
+				T::Update(Driver, NewCmd->GetFrameDeltaTime().ToRealTimeSeconds(), *NewCmd, *PrevSyncState, *NextSyncState, Junk);
 				TickInfo.IncrementTotalProcessedSimulationTime(NewCmd->GetFrameDeltaTime(), Buffers.Sync.GetHeadKeyframe());
 
 				// Set our LastProcessedInputKeyframe to fake that we handled it
@@ -422,8 +441,8 @@ struct TReplicator_Simulated : public TBase
 		ReconcileSimulationTime.Reset();
 	}
 
-	template<typename T, typename TDriver, typename TTickParameters>
-	void PreSimTick(TDriver* Driver, TNetworkSimBufferContainer<TBufferTypes>& Buffers, TSimulationTickState<TTickSettings>& TickInfo, const TTickParameters& TickParameters)
+	template<typename T, typename TDriver>
+	void PreSimTick(TDriver* Driver, TNetworkSimBufferContainer<TBufferTypes>& Buffers, TSimulationTickState<TTickSettings>& TickInfo, const FNetSimTickParameters& TickParameters)
 	{
 		if (bAllowSimulatedExtrapolation && NetworkSimulationModelCVars::EnableSimulatedExtrapolation())
 		{
@@ -444,7 +463,7 @@ struct TReplicator_Simulated : public TBase
 
 						*InputCmd = TInputCmd();
 						InputCmd->SetFrameDeltaTime(DeltaSimTime);
-						Driver->ProduceInput(DeltaSimTime, *InputCmd);
+						Driver->ProduceInput(DeltaSimTime.ToRealTimeSeconds(), *InputCmd);
 						TickInfo.MaxAllowedInputKeyframe = Buffers.Input.GetHeadKeyframe();
 					}
 				}
@@ -452,7 +471,23 @@ struct TReplicator_Simulated : public TBase
 		}
 	}
 
-
+	template<typename TDriver>
+	void PostSimTick(TDriver* Driver, const TNetworkSimBufferContainer<TBufferTypes>& Buffers, const TSimulationTickState<TTickSettings>& TickInfo, const FNetSimTickParameters& TickParameters)
+	{
+		if (bAllowSimulatedExtrapolation)
+		{
+			// Sync to latest frame if there is any
+			if (Buffers.Sync.GetNumValidElements() > 0)
+			{
+				Driver->FinalizeFrame(*Buffers.Sync.GetElementFromHead(0));
+			}
+			
+		}
+		else
+		{
+			Interpolator.template PostSimTick<TDriver>(Driver, Buffers, TickInfo, TickParameters);
+		}
+	}
 
 private:
 	
@@ -645,7 +680,7 @@ struct TReplicator_Autonomous : public TBase
 			NextMotionState->VisualLog( FVisualLoggingParameters(Keyframe == LastKeyframeToProcess ? EVisualLoggingContext::LastMispredicted : EVisualLoggingContext::OtherMispredicted, Keyframe, EVisualLoggingLifetime::Persistent), Driver, Driver);
 
 			// Do the actual update
-			T::Update(Driver, ResimulateCmd->GetFrameDeltaTime(), *ResimulateCmd, *PrevMotionState, *NextMotionState, *AuxState);
+			T::Update(Driver, ResimulateCmd->GetFrameDeltaTime().ToRealTimeSeconds(), *ResimulateCmd, *PrevMotionState, *NextMotionState, *AuxState);
 			
 			// Update TickInfo
 			TickInfo.IncrementTotalProcessedSimulationTime(ResimulateCmd->GetFrameDeltaTime(), Keyframe);
@@ -659,8 +694,8 @@ struct TReplicator_Autonomous : public TBase
 	// --------------------------------------------------------------------
 	//	PreSimTick
 	// --------------------------------------------------------------------
-	template<typename T, typename TDriver, typename TTickParameters>
-	void PreSimTick(TDriver* Driver, TNetworkSimBufferContainer<TBufferTypes>& Buffers, TSimulationTickState<TTickSettings>& TickInfo, const TTickParameters& TickParameters)
+	template<typename T, typename TDriver>
+	void PreSimTick(TDriver* Driver, TNetworkSimBufferContainer<TBufferTypes>& Buffers, TSimulationTickState<TTickSettings>& TickInfo, const FNetSimTickParameters& TickParameters)
 	{
 		// If we have a reconcile fault, we cannot continue on with the simulation until it clears itself out. This effectively drops the input time and does not sample new inputs
 		if (bReconcileFaultDetected)
@@ -681,7 +716,7 @@ struct TReplicator_Autonomous : public TBase
 					{
 						*InputCmd = TInputCmd();
 						InputCmd->SetFrameDeltaTime(DeltaSimTime);
-						Driver->ProduceInput(DeltaSimTime, *InputCmd);
+						Driver->ProduceInput(DeltaSimTime.ToRealTimeSeconds(), *InputCmd);
 						TickInfo.MaxAllowedInputKeyframe = Buffers.Input.GetHeadKeyframe(); // Allow the new command to be processed by the local simulation
 					}
 				}
@@ -700,7 +735,7 @@ struct TReplicator_Autonomous : public TBase
 					{
 						*InputCmd = TInputCmd();
 						InputCmd->SetFrameDeltaTime(NonPredictedInputTime);
-						Driver->ProduceInput(NonPredictedInputTime, *InputCmd);
+						Driver->ProduceInput(NonPredictedInputTime.ToRealTimeSeconds(), *InputCmd);
 					}
 				}
 			}
diff --git a/Engine/Plugins/Runtime/NetworkPrediction/Source/NetworkPrediction/Public/NetworkSimulationModelTypes.h b/Engine/Plugins/Runtime/NetworkPrediction/Source/NetworkPrediction/Public/NetworkSimulationModelTypes.h
--- a/Engine/Plugins/Runtime/NetworkPrediction/Source/NetworkPrediction/Public/NetworkSimulationModelTypes.h
+++ b/Engine/Plugins/Runtime/NetworkPrediction/Source/NetworkPrediction/Public/NetworkSimulationModelTypes.h
@@ -164,8 +164,21 @@ struct TNetworkSimBufferContainer
 	}
 };
 
+/** Interface base for the simulation driver. This is what is used by the Network Sim Model internally. Basically, the non simulation specific functions that need to be defined. */
+template<typename TBufferTypes>
+class TNetworkSimDriverInterfaceBase
+{
+public:
+	virtual FString GetDebugName() const = 0; // Used for debugging. Recommended to emit the simulation name and the actor name/role.
+	virtual const UObject* GetVLogOwner() const = 0; // Owning object for Visual Logs
+
+	virtual void InitSyncState(typename TBufferTypes::TSyncState& OutSyncState) const = 0;	// Called to create initial value of the sync state.
+	virtual void ProduceInput(const float DeltaTimeSeconds, typename TBufferTypes::TInputCmd&) = 0; // Called when the sim is ready to process new local input
+	virtual void FinalizeFrame(const typename TBufferTypes::TSyncState& SyncState) = 0; // Called from the Network Sim at the end of the sim frame when there is new sync data.
+};
+
 // ----------------------------------------------------------------------------------------------------------------------------------------------
-//
+//		Tick and time keeping related structures
 // ----------------------------------------------------------------------------------------------------------------------------------------------
 
 // Holds all settings for a network sim related to ticking
@@ -317,8 +330,8 @@ struct TSimulationTickState
 {
 	using TSettings = TickSettings;
 
-	int32 LastProcessedInputKeyframe;	// The last input keyframe that we processed
-	int32 MaxAllowedInputKeyframe;		// The max input keyframe that we are allowed to process (e.g, don't process input past this keyframe yet)
+	int32 LastProcessedInputKeyframe = 0;	// The last input keyframe that we processed
+	int32 MaxAllowedInputKeyframe = 0;		// The max input keyframe that we are allowed to process (e.g, don't process input past this keyframe yet)
 
 	void SetTotalAllowedSimulationTime(const TNetworkSimTime<TSettings>& SimTime)
 	{
diff --git a/Engine/Plugins/Runtime/NetworkPrediction/readme.txt b/Engine/Plugins/Runtime/NetworkPrediction/readme.txt
--- a/Engine/Plugins/Runtime/NetworkPrediction/readme.txt
+++ b/Engine/Plugins/Runtime/NetworkPrediction/readme.txt
@@ -108,6 +108,10 @@ Smoothing: Taking the output of the simulation and applying an additional layer
 // Release notes
 // ----------------------------------------------------------------------------------------------------------
 
+Update (9-26-19)
+-More refactors on general system architecture. UNetworkSimulationGlobalManager is solidifying and most boiler plate is out of the actor components (Not an exciting update but this was overdue!).
+-PostSimTick added and Interpolator refactored to be part of this step, dependent on network role (Eg., simulated proxy only when enabled).
+-Continuing to investigate dependent actor simulations. A few early prototypes did not work out but am going to try a slightly different approach.
 
 Update (8-14-19)
 -Some more cleanup. Cvar fixes and other consolidation. Now time for vacation (2 weeks). 
diff --git a/Engine/Plugins/Runtime/NetworkPredictionExtras/Source/NetworkPredictionExtras/Private/MockNetworkSimulation.cpp b/Engine/Plugins/Runtime/NetworkPredictionExtras/Source/NetworkPredictionExtras/Private/MockNetworkSimulation.cpp
--- a/Engine/Plugins/Runtime/NetworkPredictionExtras/Source/NetworkPredictionExtras/Private/MockNetworkSimulation.cpp
+++ b/Engine/Plugins/Runtime/NetworkPredictionExtras/Source/NetworkPredictionExtras/Private/MockNetworkSimulation.cpp
@@ -59,9 +59,12 @@ static FAutoConsoleVariableRef CVarBindAutomatically(TEXT("mns.BindAutomatically
 // ============================================================================================
 
 static bool ForceMispredict = false;
-void FMockNetworkSimulation::Update(IMockNetworkSimulationDriver* Driver, const TSimTime& DeltaTimeMS, const FMockInputCmd& InputCmd, const FMockSyncState& InputState, FMockSyncState& OutputState, const FMockAuxState& AuxState)
+
+const FName FMockNetworkSimulation::GroupName(TEXT("Mock"));
+
+void FMockNetworkSimulation::Update(IMockDriver* Driver, const float DeltaTimeSeconds, const FMockInputCmd& InputCmd, const FMockSyncState& InputState, FMockSyncState& OutputState, const FMockAuxState& AuxState)
 {
-	OutputState.Total = InputState.Total + (InputCmd.InputValue * AuxState.Multiplier * DeltaTimeMS.ToRealTimeSeconds());
+	OutputState.Total = InputState.Total + (InputCmd.InputValue * AuxState.Multiplier * DeltaTimeSeconds);
 
 	// Dev hack to force mispredict
 	if (ForceMispredict)
@@ -74,7 +77,7 @@ void FMockNetworkSimulation::Update(IMockNetworkSimulationDriver* Driver, const
 	}
 }
 
-void FMockSyncState::VisualLog(const FVisualLoggingParameters& Parameters, IMockNetworkSimulationDriver* Driver, IMockNetworkSimulationDriver* LogDriver)
+void FMockSyncState::VisualLog(const FVisualLoggingParameters& Parameters, IMockDriver* Driver, IMockDriver* LogDriver) const
 {
 	// This function is awkward because the mock example doesn't include any positional data, so we added an explicit ::GetDebugWorldTransform to the driver interface.
 	// still, this doesn't give us historical location info. So we will just use Z offset based on the EvisualLoggingContext to space things out.
@@ -129,60 +132,29 @@ UMockNetworkSimulationComponent::UMockNetworkSimulationComponent()
 	}
 }
 
-IReplicationProxy* UMockNetworkSimulationComponent::InstantiateNetworkSimulation()
+INetworkSimulationModel* UMockNetworkSimulationComponent::InstantiateNetworkSimulation()
 {
-	NetworkSim.Reset(new FMockNetworkSimulation());
-
-#if NETSIM_MODEL_DEBUG
-	FNetworkSimulationModelDebuggerManager::Get().RegisterNetworkSimulationModel(NetworkSim.Get(), (IMockNetworkSimulationDriver*)this, GetOwner(), TEXT("MockNetworkSimulation"));
-#endif
-	return NetworkSim.Get();
-}
-
-// Child classes should override this an initialize their NetworkSim here
-void UMockNetworkSimulationComponent::InitializeForNetworkRole(ENetRole Role)
-{	
-	NetworkSim->InitializeForNetworkRole(Role, GetSimulationInitParameters(Role));
+	auto* NewSim = new FMockNetworkModel(this);
+	DO_NETSIM_MODEL_DEBUG(FNetworkSimulationModelDebuggerManager::Get().RegisterNetworkSimulationModel(NewSim, GetOwner()));
+	return NewSim;
 }
 
 void UMockNetworkSimulationComponent::TickComponent(float DeltaTime, enum ELevelTick TickType, FActorComponentTickFunction *ThisTickFunction)
 {
-	CheckOwnerRoleChange();
-
 	const ENetRole OwnerRole = GetOwnerRole();
-	
-	// -------------------------------------
-	// Tick the network sim
-	// -------------------------------------
-	if (NetworkSim)
-	{
-		PreTickSimulation(DeltaTime); // Fixme
 
-		// Check if we should trip a mispredict. (Note how its not possible to do this inside the Update function!)
-		if (OwnerRole == ROLE_Authority && MockNetworkSimCVars::RequestMispredict)
-		{
-			ForceMispredict = true;
-			MockNetworkSimCVars::RequestMispredict = 0;
-		}
-		
-		FMockNetworkSimulation::FTickParameters Parameters(DeltaTime, GetOwner());
-
-		// Tick the core network sim, this will consume input and generate new sync state
-		NetworkSim->Tick((IMockNetworkSimulationDriver*)this, Parameters);
-
-		// Client->Server communication
-		if (NetworkSim->ShouldSendServerRPC(OwnerRole, DeltaTime))
-		{
-			FServerReplicationRPCParameter ProxyParameter(ReplicationProxy_ServerRPC);
-			ServerRecieveClientInput(ProxyParameter);
-		}
+	// Check if we should trip a mispredict. (Note how its not possible to do this inside the Update function!)
+	if (OwnerRole == ROLE_Authority && MockNetworkSimCVars::RequestMispredict)
+	{
+		ForceMispredict = true;
+		MockNetworkSimCVars::RequestMispredict = 0;
 	}
 	
 	// Mock example of displaying
 	DrawDebugString( GetWorld(), GetOwner()->GetActorLocation() + FVector(0.f,0.f,100.f), *LexToString(MockValue), nullptr, FColor::White, 0.00001f );
 }
 
-void UMockNetworkSimulationComponent::ProduceInput(const FMockNetworkSimulation::TSimTime& SimFrameTime, FMockInputCmd& Cmd)
+void UMockNetworkSimulationComponent::ProduceInput(const float DeltaTimeSeconds, FMockInputCmd& Cmd)
 {
 	if (MockNetworkSimCVars::DoLocalInput)
 	{
diff --git a/Engine/Plugins/Runtime/NetworkPredictionExtras/Source/NetworkPredictionExtras/Private/NetworkPredictionExtrasFlyingPawn.cpp b/Engine/Plugins/Runtime/NetworkPredictionExtras/Source/NetworkPredictionExtras/Private/NetworkPredictionExtrasFlyingPawn.cpp
--- a/Engine/Plugins/Runtime/NetworkPredictionExtras/Source/NetworkPredictionExtras/Private/NetworkPredictionExtrasFlyingPawn.cpp
+++ b/Engine/Plugins/Runtime/NetworkPredictionExtras/Source/NetworkPredictionExtras/Private/NetworkPredictionExtrasFlyingPawn.cpp
@@ -135,7 +135,7 @@ void ANetworkPredictionExtrasFlyingPawn::Tick( float DeltaSeconds)
 	// Do whatever you want here. By now we have the latest movement state and latest input processed.
 }
 
-void ANetworkPredictionExtrasFlyingPawn::ProduceInput(const FlyingMovement::TSimTime& SimFrameTime, FlyingMovement::FInputCmd& Cmd)
+void ANetworkPredictionExtrasFlyingPawn::ProduceInput(const float DeltaTimeSeconds, FlyingMovement::FInputCmd& Cmd)
 {
 	// Generate user commands. Called right before the flying movement simulation will tick (for a locally controlled pawn)
 	// This isn't meant to be the best way of doing a camera system. It is just meant to show a couple of ways it may be done
@@ -172,8 +172,6 @@ void ANetworkPredictionExtrasFlyingPawn::ProduceInput(const FlyingMovement::TSim
 	Cmd.MovementInput = FVector::ZeroVector;
 	Cmd.RotationInput = FRotator::ZeroRotator;
 
-	const float DeltaTimeSeconds = SimFrameTime.ToRealTimeSeconds();
-
 	switch (FlyingPawnCVars::CameraStyle)
 	{
 		case 0:
diff --git a/Engine/Plugins/Runtime/NetworkPredictionExtras/Source/NetworkPredictionExtras/Public/MockNetworkSimulation.h b/Engine/Plugins/Runtime/NetworkPredictionExtras/Source/NetworkPredictionExtras/Public/MockNetworkSimulation.h
--- a/Engine/Plugins/Runtime/NetworkPredictionExtras/Source/NetworkPredictionExtras/Public/MockNetworkSimulation.h
+++ b/Engine/Plugins/Runtime/NetworkPredictionExtras/Source/NetworkPredictionExtras/Public/MockNetworkSimulation.h
@@ -20,6 +20,7 @@
 #include "MockNetworkSimulation.generated.h"
 
 class IMockNetworkSimulationDriver;
+class IMockDriver;
 
 // -------------------------------------------------------------------------------------------------------------------------------
 //	Mock Network Simulation
@@ -29,7 +30,7 @@ class IMockNetworkSimulationDriver;
 //	in this example. This is just meant to show the bare minimum hook ups into the system to make it easier to understand.
 //
 //	Highlights:
-//		FMockNetworkSimulation::Update			The "core update" function of the simulation.
+//		FMockNetworkModel::Update			The "core update" function of the simulation.
 //		UMockNetworkSimulationComponent			The UE4 Actor Component that anchors the system to an actor.
 //
 //	Usage:
@@ -98,7 +99,12 @@ struct FMockSyncState
 		}
 	}
 
-	void VisualLog(const FVisualLoggingParameters& Parameters, IMockNetworkSimulationDriver* Driver, IMockNetworkSimulationDriver* LogDriver);
+	void VisualLog(const FVisualLoggingParameters& Parameters, IMockDriver* Driver, IMockDriver* LogDriver) const;
+
+	static void Interpolate(const FMockSyncState& From, const FMockSyncState& To, const float PCT, FMockSyncState& OutDest)
+	{
+		OutDest.Total = From.Total + ((To.Total - From.Total) * PCT);
+	}
 };
 
 // Auxiliary state that is input into the simulation. Doesn't change during the simulation tick.
@@ -115,27 +121,34 @@ struct FMockAuxState
 
 using TMockNetworkSimulationBufferTypes = TNetworkSimBufferTypes<FMockInputCmd, FMockSyncState, FMockAuxState>;
 
-// Actual definition of our network simulation.
-class FMockNetworkSimulation : public TNetworkedSimulationModel<FMockNetworkSimulation, TMockNetworkSimulationBufferTypes>
+static FName MockSimulationGroupName("Mock");
+
+// Interface between the simulation and owning component driving it. Functions added here are available in ::Update and must be implemented by UMockNetworkSimulationComponent.
+class IMockDriver : public TNetworkSimDriverInterfaceBase<TMockNetworkSimulationBufferTypes>
 {
 public:
 
-	// Interface between the simulation and owning component driving it. Functions added here are available in ::Update and must be implemented by UMockNetworkSimulationComponent.
-	class IMockDriver : public IDriver
-	{
-	public:
-		virtual UWorld* GetDriverWorld() const = 0;
-		virtual UObject* GetVLogOwner() const = 0;
-		virtual FTransform GetDebugWorldTransform() const = 0;
-	};
+	virtual UWorld* GetDriverWorld() const = 0;
+	virtual FTransform GetDebugWorldTransform() const = 0;
+};
+
+class FMockNetworkSimulation
+{
+public:
 
 	/** Main update function */
-	static void Update(IMockNetworkSimulationDriver* Driver, const TSimTime& DeltaTimeMS, const FMockInputCmd& InputCmd, const FMockSyncState& InputState, FMockSyncState& OutputState, const FMockAuxState& AuxState);
+	static void Update(IMockDriver* Driver, const float DeltaTimeSeconds, const FMockInputCmd& InputCmd, const FMockSyncState& InputState, FMockSyncState& OutputState, const FMockAuxState& AuxState);
+
+	/** Tick group the simulation maps to */
+	static const FName GroupName;
 };
 
+// Actual definition of our network simulation.
+using FMockNetworkModel = TNetworkedSimulationModel<FMockNetworkSimulation, IMockDriver, TMockNetworkSimulationBufferTypes>;
+
 // Needed to trick UHT into letting UMockNetworkSimulationComponent implement. UHT cannot parse the ::
 // Also needed for forward declaring. Can't just be a typedef/using =
-class IMockNetworkSimulationDriver : public FMockNetworkSimulation::IMockDriver { };
+class IMockNetworkSimulationDriver : public IMockDriver { };
 
 // -------------------------------------------------------------------------------------------------------------------------------
 // ActorComponent for running a MockNetworkSimulation
@@ -152,23 +165,17 @@ public:
 
 	virtual void TickComponent(float DeltaTime, enum ELevelTick TickType, FActorComponentTickFunction *ThisTickFunction) override;
 	
-	virtual IReplicationProxy* InstantiateNetworkSimulation() override;
-	virtual void InitializeForNetworkRole(ENetRole Role) override;
-
-private:
-
-	// The Network sim that this component is managing. This is what is doing all the work.
-	TUniquePtr<FMockNetworkSimulation> NetworkSim;
+	virtual INetworkSimulationModel* InstantiateNetworkSimulation() override;
 	
 public:
 
 	FString GetDebugName() const override;
+	virtual UObject* GetVLogOwner() const override final;
+
 	void InitSyncState(FMockSyncState& OutSyncState) const override;
 	void FinalizeFrame(const FMockSyncState& SyncState) override;
-	void ProduceInput(const FMockNetworkSimulation::TSimTime& SimFrameTime, FMockInputCmd& Cmd);
-
+	void ProduceInput(const float DeltaTimeSeconds, FMockInputCmd& Cmd);
 	virtual UWorld* GetDriverWorld() const override final { return GetWorld(); }
-	virtual UObject* GetVLogOwner() const override final;
 	virtual FTransform GetDebugWorldTransform() const override final;
 
 	// Mock representation of "syncing' to the sync state in the network sim.
diff --git a/Engine/Plugins/Runtime/NetworkPredictionExtras/Source/NetworkPredictionExtras/Public/NetworkPredictionExtrasFlyingPawn.h b/Engine/Plugins/Runtime/NetworkPredictionExtras/Source/NetworkPredictionExtras/Public/NetworkPredictionExtrasFlyingPawn.h
--- a/Engine/Plugins/Runtime/NetworkPredictionExtras/Source/NetworkPredictionExtras/Public/NetworkPredictionExtrasFlyingPawn.h
+++ b/Engine/Plugins/Runtime/NetworkPredictionExtras/Source/NetworkPredictionExtras/Public/NetworkPredictionExtrasFlyingPawn.h
@@ -45,7 +45,7 @@ class NETWORKPREDICTIONEXTRAS_API ANetworkPredictionExtrasFlyingPawn : public AP
 
 private:
 
-	void ProduceInput(const FlyingMovement::TSimTime& SimFrameTime, FlyingMovement::FInputCmd& Cmd);
+	void ProduceInput(const float TimeSeconds, FlyingMovement::FInputCmd& Cmd);
 
 	FVector CachedMoveInput;
 	FVector2D CachedLookInput;
