NetworkPrediction: very minor cleanup before next phase
#rnx

diff --git a/Engine/Plugins/Runtime/NetworkPrediction/Source/NetworkPrediction/Public/NetworkPredictionComponent.h b/Engine/Plugins/Runtime/NetworkPrediction/Source/NetworkPrediction/Public/NetworkPredictionComponent.h
--- a/Engine/Plugins/Runtime/NetworkPrediction/Source/NetworkPrediction/Public/NetworkPredictionComponent.h
+++ b/Engine/Plugins/Runtime/NetworkPrediction/Source/NetworkPrediction/Public/NetworkPredictionComponent.h
@@ -12,7 +12,7 @@
 //	initialized and plugged into the UE4 replication system.
 //
 //	This is an abstract component and cannot function on its own. It must be subclassed and InstantiateNetworkSimulation must be implemented.
-//	The subclass is responsible for running the network simulation. This at a bare minimum would mean ticking the simulation and supplying it with input.
+//	Ticking and RPC sending will be handled automatically, but the subclass will need to supply input to the simulation (via TNetworkSimDriverInterfaceBase::ProduceInput).
 //
 //	Its also worth pointing out that nothing about being a UActorComponent is essential here. All that this component does could be done within an AActor itself.
 //	An actor component makes sense for flexible/reusable code provided by the plugin. But there is nothing stopping you from copying this directly into an actor if you had reason to.
@@ -53,6 +53,7 @@ protected:
 	bool CheckOwnerRoleChange();
 	ENetRole OwnerCachedNetRole = ROLE_None;
 
+	// The actual ServerRPC. This needs to be a UFUNCTION but is generic and not dependent on the simulation instance
 	UFUNCTION(Server, unreliable, WithValidation)
 	void ServerRecieveClientInput(const FServerReplicationRPCParameter& ProxyParameter);
 
@@ -76,15 +77,13 @@ private:
 
 protected:
 
+	// Called via ServerRPCDelegate, ticks serverRPC timing and decides whether to send the RPC or not
 	void TickServerRPC(float DeltaSeconds);
 
 	void RegisterServerRPCDelegate();
 	void UnregisterServerRPCDelegate();
 	FDelegateHandle ServerRPCHandle;
 
-	UFUNCTION()
-	virtual void OnRep_SimulatedProxy() { }
-
 	// The Network sim that this component is managing. This is what is doing all the work.
 	TUniquePtr<INetworkSimulationModel> NetworkSim;
 };
\ No newline at end of file
