Network prediction: untemplatized the time keeping types in the system. Tick settings are still templates. This is to allow easier communication and interfaces to exist within the eco system.
#rnx

diff --git a/Engine/Plugins/Runtime/NetworkPrediction/Source/NetworkPrediction/Private/Movement/FlyingMovementComponent.cpp b/Engine/Plugins/Runtime/NetworkPrediction/Source/NetworkPrediction/Private/Movement/FlyingMovementComponent.cpp
--- a/Engine/Plugins/Runtime/NetworkPrediction/Source/NetworkPrediction/Private/Movement/FlyingMovementComponent.cpp
+++ b/Engine/Plugins/Runtime/NetworkPrediction/Source/NetworkPrediction/Private/Movement/FlyingMovementComponent.cpp
@@ -125,10 +125,10 @@ void UFlyingMovementComponent::PreSimSync(const FlyingMovement::FMoveState& Sync
 	}
 }
 
-void UFlyingMovementComponent::ProduceInput(const float TimeSeconds, FlyingMovement::FInputCmd& Cmd)
+void UFlyingMovementComponent::ProduceInput(const FNetworkSimTime SimTime, FlyingMovement::FInputCmd& Cmd)
 {
 	// This isn't ideal. It probably makes sense for the component to do all the input binding rather.
-	ProduceInputDelegate.ExecuteIfBound(TimeSeconds, Cmd);
+	ProduceInputDelegate.ExecuteIfBound(SimTime, Cmd);
 }
 
 void UFlyingMovementComponent::FinalizeFrame(const FlyingMovement::FMoveState& SyncState)
diff --git a/Engine/Plugins/Runtime/NetworkPrediction/Source/NetworkPrediction/Private/Movement/ParametricMovement.cpp b/Engine/Plugins/Runtime/NetworkPrediction/Source/NetworkPrediction/Private/Movement/ParametricMovement.cpp
--- a/Engine/Plugins/Runtime/NetworkPrediction/Source/NetworkPrediction/Private/Movement/ParametricMovement.cpp
+++ b/Engine/Plugins/Runtime/NetworkPrediction/Source/NetworkPrediction/Private/Movement/ParametricMovement.cpp
@@ -210,7 +210,7 @@ void UParametricMovementComponent::TickComponent(float DeltaTime, enum ELevelTic
 	}
 }
 
-void UParametricMovementComponent::ProduceInput(const float DeltaTimeSeconds, ParametricMovement::FInputCmd& Cmd)
+void UParametricMovementComponent::ProduceInput(const FNetworkSimTime DeltaTimeSeconds, ParametricMovement::FInputCmd& Cmd)
 {
 	Cmd.PlayRate = PendingPlayRate;
 	PendingPlayRate.Reset();
diff --git a/Engine/Plugins/Runtime/NetworkPrediction/Source/NetworkPrediction/Public/Movement/FlyingMovement.h b/Engine/Plugins/Runtime/NetworkPrediction/Source/NetworkPrediction/Public/Movement/FlyingMovement.h
--- a/Engine/Plugins/Runtime/NetworkPrediction/Source/NetworkPrediction/Public/Movement/FlyingMovement.h
+++ b/Engine/Plugins/Runtime/NetworkPrediction/Source/NetworkPrediction/Public/Movement/FlyingMovement.h
@@ -180,10 +180,10 @@ public:
 	const UObject* GetVLogOwner() const override;
 	void InitSyncState(FlyingMovement::FMoveState& OutSyncState) const override;
 	void PreSimSync(const FlyingMovement::FMoveState& SyncState) override;
-	void ProduceInput(const float DeltaTimeSeconds, FlyingMovement::FInputCmd& Cmd) override;
+	void ProduceInput(const FNetworkSimTime SimTime, FlyingMovement::FInputCmd& Cmd) override;
 	void FinalizeFrame(const FlyingMovement::FMoveState& SyncState) override;
 
-	DECLARE_DELEGATE_TwoParams(FProduceFlyingInput, const float /*DeltaTimeSeconds*/, FlyingMovement::FInputCmd& /*Cmd*/)
+	DECLARE_DELEGATE_TwoParams(FProduceFlyingInput, const FNetworkSimTime /*SimTime*/, FlyingMovement::FInputCmd& /*Cmd*/)
 	FProduceFlyingInput ProduceInputDelegate;
 
 protected:
diff --git a/Engine/Plugins/Runtime/NetworkPrediction/Source/NetworkPrediction/Public/Movement/ParametricMovement.h b/Engine/Plugins/Runtime/NetworkPrediction/Source/NetworkPrediction/Public/Movement/ParametricMovement.h
--- a/Engine/Plugins/Runtime/NetworkPrediction/Source/NetworkPrediction/Public/Movement/ParametricMovement.h
+++ b/Engine/Plugins/Runtime/NetworkPrediction/Source/NetworkPrediction/Public/Movement/ParametricMovement.h
@@ -155,7 +155,7 @@ class NETWORKPREDICTION_API UParametricMovementComponent : public UBaseMovementC
 	const UObject* GetVLogOwner() const override;
 	void InitSyncState(ParametricMovement::FMoveState& OutSyncState) const override;
 	void FinalizeFrame(const ParametricMovement::FMoveState& SyncState) override;
-	void ProduceInput(const float DeltaTimeSeconds, ParametricMovement::FInputCmd& Cmd);
+	void ProduceInput(const FNetworkSimTime SimTime, ParametricMovement::FInputCmd& Cmd);
 
 	// Base Movement Driver
 	IBaseMovementDriver& GetBaseMovementDriver() override final { return *static_cast<IBaseMovementDriver*>(this); }
diff --git a/Engine/Plugins/Runtime/NetworkPrediction/Source/NetworkPrediction/Public/NetworkSimulationModel.h b/Engine/Plugins/Runtime/NetworkPrediction/Source/NetworkPrediction/Public/NetworkSimulationModel.h
--- a/Engine/Plugins/Runtime/NetworkPrediction/Source/NetworkPrediction/Public/NetworkSimulationModel.h
+++ b/Engine/Plugins/Runtime/NetworkPrediction/Source/NetworkPrediction/Public/NetworkSimulationModel.h
@@ -49,8 +49,8 @@ public:
 	using TAuxState = typename TBufferTypes::TAuxState;
 	using TDebugState = typename TBufferTypes::TDebugState;
 
-	using TSimTime = TNetworkSimTime<TTickSettings>;
-	using TRealTime = typename TTickSettings::TRealTime;
+	using TSimTime = FNetworkSimTime;
+	using TRealTime = FNetworkSimTime::FRealTime;
 
 	class IDriver
 	{
@@ -58,7 +58,7 @@ public:
 		virtual FString GetDebugName() const = 0; // Used for debugging. Recommended to emit the simulation name and the actor name/role.
 
 		virtual void InitSyncState(TSyncState& OutSyncState) const = 0;	// Called to create initial value of the sync state.
-		virtual void ProduceInput(const TRealTime, typename TUserBufferTypes::TInputCmd&) = 0; // Called when the sim is ready to process new local input
+		virtual void ProduceInput(const FNetworkSimTime, typename TUserBufferTypes::TInputCmd&) = 0; // Called when the sim is ready to process new local input
 		virtual void FinalizeFrame(const TSyncState& SyncState) = 0; // Called from the Network Sim at the end of the sim frame when there is new sync data.
 	};
 	
diff --git a/Engine/Plugins/Runtime/NetworkPrediction/Source/NetworkPrediction/Public/NetworkSimulationModelInterpolator.h b/Engine/Plugins/Runtime/NetworkPrediction/Source/NetworkPrediction/Public/NetworkSimulationModelInterpolator.h
--- a/Engine/Plugins/Runtime/NetworkPrediction/Source/NetworkPrediction/Public/NetworkSimulationModelInterpolator.h
+++ b/Engine/Plugins/Runtime/NetworkPrediction/Source/NetworkPrediction/Public/NetworkSimulationModelInterpolator.h
@@ -25,8 +25,8 @@ namespace NetworkInterpolationDebugCVars
 template<typename TBufferTypes, typename TTickSettings>
 struct TInterpolator
 {
-	using TSimTime = TNetworkSimTime<TTickSettings>;
-	using TRealTime = typename TTickSettings::TRealTime;
+	using FSimTime = FNetworkSimTime;
+	using FRealTime = FNetworkSimTime::FRealTime;
 	using TSyncState = typename TBufferTypes::TSyncState;
 
 	bool bEnableVisualLog = true;
@@ -94,12 +94,12 @@ struct TInterpolator
 		EVisualLoggingContext LoggingContext = EVisualLoggingContext::InterpolationLatest;
 
 		// Calc new interpolation time
-		TRealTime NewInterpolationTime = InterpolationTime;
+		FRealTime NewInterpolationTime = InterpolationTime;
 		{
-			TRealTime Step = DeltaSeconds;
+			FRealTime Step = DeltaSeconds;
 			
 			// Speed up if we are too far behind
-			TRealTime CatchUpThreshold = SimulationTimeBuffer.GetElementFromHead(0)->ToRealTimeSeconds() - NetworkInterpolationDebugCVars::CatchUpThreshold();
+			FRealTime CatchUpThreshold = SimulationTimeBuffer.GetElementFromHead(0)->ToRealTimeSeconds() - NetworkInterpolationDebugCVars::CatchUpThreshold();
 			if (CatchUpUntilTime <= 0.f && InterpolationTime < CatchUpThreshold)
 			{
 				CatchUpUntilTime = SimulationTimeBuffer.GetElementFromHead(0)->ToRealTimeSeconds() - NetworkInterpolationDebugCVars::CatchUpGoal();
@@ -134,11 +134,11 @@ struct TInterpolator
 
 		// Find "To" keyframe
 		const TSyncState* ToState = nullptr;
-		TRealTime ToTime = 0.f;
+		FRealTime ToTime = 0.f;
 
 		for (int32 Keyframe = SimulationTimeBuffer.GetTailKeyframe(); Keyframe <= SimulationTimeBuffer.GetHeadKeyframe(); ++Keyframe)
 		{
-			TSimTime ElementSimTime = *SimulationTimeBuffer.FindElementByKeyframe(Keyframe);
+			FSimTime ElementSimTime = *SimulationTimeBuffer.FindElementByKeyframe(Keyframe);
 			if (NewInterpolationTime <= ElementSimTime.ToRealTimeSeconds())
 			{
 				InterpolationKeyframe = Keyframe;
@@ -150,9 +150,9 @@ struct TInterpolator
 
 		if (ensure(ToState))
 		{
-			const TRealTime FromRealTime = InterpolationTime;
-			const TRealTime ToRealTime = ToTime;
-			const TRealTime InterpolationInterval = ToRealTime - FromRealTime;
+			const FRealTime FromRealTime = InterpolationTime;
+			const FRealTime ToRealTime = ToTime;
+			const FRealTime InterpolationInterval = ToRealTime - FromRealTime;
 		
 			if (ensure(FMath::Abs(InterpolationInterval) > 0.f))
 			{
@@ -210,19 +210,19 @@ struct TInterpolator
 
 private:
 
-	TRealTime InterpolationTime = 0.f; // SimTime we are currently interpolating at
+	FRealTime InterpolationTime = 0.f; // SimTime we are currently interpolating at
 	int32 InterpolationKeyframe = INDEX_NONE; // Keyframe we are currently/last interpolated at
 
 	TSyncState InterpolationState;
 
-	TRealTime WaitUntilTime = 0.f;	
+	FRealTime WaitUntilTime = 0.f;	
 
-	TRealTime CatchUpUntilTime = 0.f;
+	FRealTime CatchUpUntilTime = 0.f;
 
 
-	TRealTime DynamicBufferedTime = 1/60.f; // SimTime we are currently interpolating at
-	TRealTime DynamicBufferedTimeStep = 1/60.f;
+	FRealTime DynamicBufferedTime = 1/60.f; // SimTime we are currently interpolating at
+	FRealTime DynamicBufferedTimeStep = 1/60.f;
 
-	TRealTime	MinBufferedTime = 1/120.f;
-	TRealTime	MaxBufferedTime = 1.f;
+	FRealTime	MinBufferedTime = 1/120.f;
+	FRealTime	MaxBufferedTime = 1.f;
 };
\ No newline at end of file
diff --git a/Engine/Plugins/Runtime/NetworkPrediction/Source/NetworkPrediction/Public/NetworkSimulationModelReplicators.h b/Engine/Plugins/Runtime/NetworkPrediction/Source/NetworkPrediction/Public/NetworkSimulationModelReplicators.h
--- a/Engine/Plugins/Runtime/NetworkPrediction/Source/NetworkPrediction/Public/NetworkSimulationModelReplicators.h
+++ b/Engine/Plugins/Runtime/NetworkPrediction/Source/NetworkPrediction/Public/NetworkSimulationModelReplicators.h
@@ -128,7 +128,7 @@ struct TReplicator_SimTime : public TBase
 		SerializedTime.NetSerialize(P.Ar);
 	}
 
-	TNetworkSimTime<TTickSettings> SerializedTime;
+	FNetworkSimTime SerializedTime;
 };
 
 // Enabled=false specialization: do nothing
@@ -261,14 +261,14 @@ struct TReplicator_Server : public TBase
 
 		if (TickParameters.bGenerateLocalInputCmds)
 		{
-			TNetworkSimTime<TTickSettings> DeltaSimTime = TickInfo.GetRemainingAllowedSimulationTime();
+			FNetworkSimTime DeltaSimTime = TickInfo.GetRemainingAllowedSimulationTime();
 			if (DeltaSimTime.IsPositive())
 			{
 				if (TInputCmd* InputCmd = Buffers.Input.GetWriteNext())
 				{
 					*InputCmd = TInputCmd();
 					InputCmd->SetFrameDeltaTime(DeltaSimTime);
-					Driver->ProduceInput(DeltaSimTime.ToRealTimeSeconds(), *InputCmd);
+					Driver->ProduceInput(DeltaSimTime, *InputCmd);
 					TickInfo.MaxAllowedInputKeyframe = Buffers.Input.GetHeadKeyframe();
 				}
 			}
@@ -289,7 +289,7 @@ struct TReplicator_Simulated : public TBase
 
 	TInterpolator<TBufferTypes, TTickSettings> Interpolator;
 
-	TNetworkSimTime<TTickSettings> GetLastSerializedSimulationTime() const { return LastSerializedSimulationTime; }
+	FNetworkSimTime GetLastSerializedSimulationTime() const { return LastSerializedSimulationTime; }
 	
 	int32 GetProxyDirtyCount(TNetworkSimBufferContainer<TBufferTypes>& Buffers) const
 	{
@@ -300,7 +300,7 @@ struct TReplicator_Simulated : public TBase
 	{
 		FArchive& Ar = P.Ar;
 
-		TNetworkSimTime<TTickSettings> PrevLastSerializedSimulationTime = LastSerializedSimulationTime;
+		FNetworkSimTime PrevLastSerializedSimulationTime = LastSerializedSimulationTime;
 
 		// Serialize latest simulation time
 		LastSerializedSimulationTime = TickInfo.GetTotalProcessedSimulationTime();
@@ -349,8 +349,8 @@ struct TReplicator_Simulated : public TBase
 
 				if (DestinationKeyframe == INDEX_NONE)
 				{
-					TNetworkSimTime<TTickSettings> TotalTimeAhead = *TickInfo.SimulationTimeBuffer.GetElementFromHead(0) - LastSerializedSimulationTime;
-					TNetworkSimTime<TTickSettings> SerializeDelta = LastSerializedSimulationTime - PrevLastSerializedSimulationTime;
+					FNetworkSimTime TotalTimeAhead = *TickInfo.SimulationTimeBuffer.GetElementFromHead(0) - LastSerializedSimulationTime;
+					FNetworkSimTime SerializeDelta = LastSerializedSimulationTime - PrevLastSerializedSimulationTime;
 
 					// We are way far ahead of the server... we will need to clear out sync buffers, take what they gave us, and catch up in reconcile
 					//UE_LOG(LogNetworkSim, Warning, TEXT("!!! TReplicator_Simulated. Large gap detected. SerializedTime: %s. Buffer time: [%s-%s]. %d Elements. DeltaFromHead: %s. DeltaSerialize: %s"), *LastSerializedSimulationTime.ToString(), 
@@ -407,7 +407,7 @@ struct TReplicator_Simulated : public TBase
 			}
 
 			// Do we have time to make up? We may have extrapolated ahead of the server (totally fine - can happen with small amount of latency variance)
-			TNetworkSimTime<TTickSettings> DeltaSimTime = ReconcileSimulationTime - TickInfo.GetTotalProcessedSimulationTime();
+			FNetworkSimTime DeltaSimTime = ReconcileSimulationTime - TickInfo.GetTotalProcessedSimulationTime();
 			if (DeltaSimTime.IsPositive() && NetworkSimulationModelCVars::EnableSimulatedReconcile())
 			{
 				// We have extrapolated ahead of the server. The latest network update is now "in the past" from what we rendered last frame.
@@ -454,7 +454,7 @@ struct TReplicator_Simulated : public TBase
 
 			if (TickParameters.bGenerateLocalInputCmds)
 			{
-				TNetworkSimTime<TTickSettings> DeltaSimTime = TickInfo.GetRemainingAllowedSimulationTime();
+				FNetworkSimTime DeltaSimTime = TickInfo.GetRemainingAllowedSimulationTime();
 				if (DeltaSimTime.IsPositive())
 				{
 					if (TInputCmd* InputCmd = Buffers.Input.GetWriteNext())
@@ -463,7 +463,7 @@ struct TReplicator_Simulated : public TBase
 
 						*InputCmd = TInputCmd();
 						InputCmd->SetFrameDeltaTime(DeltaSimTime);
-						Driver->ProduceInput(DeltaSimTime.ToRealTimeSeconds(), *InputCmd);
+						Driver->ProduceInput(DeltaSimTime, *InputCmd);
 						TickInfo.MaxAllowedInputKeyframe = Buffers.Input.GetHeadKeyframe();
 					}
 				}
@@ -491,8 +491,8 @@ struct TReplicator_Simulated : public TBase
 
 private:
 	
-	TNetworkSimTime<TTickSettings> ReconcileSimulationTime;
-	TNetworkSimTime<TTickSettings> LastSerializedSimulationTime;
+	FNetworkSimTime ReconcileSimulationTime;
+	FNetworkSimTime LastSerializedSimulationTime;
 	
 };
 
@@ -507,7 +507,7 @@ struct TReplicator_Autonomous : public TBase
 
 	int32 GetLastSerializedKeyframe() const { return LastSerializedKeyframe; }
 	bool IsReconcileFaultDetected() const { return bReconcileFaultDetected; }
-	const TNetworkSimTime<TTickSettings>& GetLastSerializedSimTime() const { return SerializedTime; }
+	const FNetworkSimTime& GetLastSerializedSimTime() const { return SerializedTime; }
 
 	int32 GetProxyDirtyCount(TNetworkSimBufferContainer<TBufferTypes>& Buffers) const
 	{
@@ -709,14 +709,14 @@ struct TReplicator_Autonomous : public TBase
 			{
 				// Prediction: add simulation time and generate new commands
 				TickInfo.GiveSimulationTime(TickParameters.LocalDeltaTimeSeconds);
-				const TNetworkSimTime<TTickSettings> DeltaSimTime = TickInfo.GetRemainingAllowedSimulationTime();
+				const FNetworkSimTime DeltaSimTime = TickInfo.GetRemainingAllowedSimulationTime();
 				if (DeltaSimTime.IsPositive())
 				{
 					if (TInputCmd* InputCmd = Buffers.Input.GetWriteNext())
 					{
 						*InputCmd = TInputCmd();
 						InputCmd->SetFrameDeltaTime(DeltaSimTime);
-						Driver->ProduceInput(DeltaSimTime.ToRealTimeSeconds(), *InputCmd);
+						Driver->ProduceInput(DeltaSimTime, *InputCmd);
 						TickInfo.MaxAllowedInputKeyframe = Buffers.Input.GetHeadKeyframe(); // Allow the new command to be processed by the local simulation
 					}
 				}
@@ -727,7 +727,7 @@ struct TReplicator_Autonomous : public TBase
 				// Since we aren't processing the simulation locally, our core simulation time will only advance from network updates.
 				// (still, we need *something* to tell us when to generate a new command and what the delta time should be)
 
-				TNetworkSimTime<TTickSettings> NonPredictedInputTime;
+				FNetworkSimTime NonPredictedInputTime;
 				NonPredictedInputTimeAccumulator.Accumulate(NonPredictedInputTime, TickParameters.LocalDeltaTimeSeconds);
 				if (NonPredictedInputTime.IsPositive())
 				{
@@ -735,7 +735,7 @@ struct TReplicator_Autonomous : public TBase
 					{
 						*InputCmd = TInputCmd();
 						InputCmd->SetFrameDeltaTime(NonPredictedInputTime);
-						Driver->ProduceInput(NonPredictedInputTime.ToRealTimeSeconds(), *InputCmd);
+						Driver->ProduceInput(NonPredictedInputTime, *InputCmd);
 					}
 				}
 			}
@@ -743,11 +743,9 @@ struct TReplicator_Autonomous : public TBase
 	}
 
 private:
-
-
 	
 	TReplicationBuffer<TSyncState> ReconciliationBuffer;
-	TNetworkSimTime<TTickSettings> SerializedTime; // last serialized time keeper
+	FNetworkSimTime SerializedTime; // last serialized time keeper
 	TRealTimeAccumulator<TTickSettings> NonPredictedInputTimeAccumulator; // for tracking input time in the non predictive case
 
 	int32 LastSerializedKeyframe = -1;
diff --git a/Engine/Plugins/Runtime/NetworkPrediction/Source/NetworkPrediction/Public/NetworkSimulationModelTypes.h b/Engine/Plugins/Runtime/NetworkPrediction/Source/NetworkPrediction/Public/NetworkSimulationModelTypes.h
--- a/Engine/Plugins/Runtime/NetworkPrediction/Source/NetworkPrediction/Public/NetworkSimulationModelTypes.h
+++ b/Engine/Plugins/Runtime/NetworkPrediction/Source/NetworkPrediction/Public/NetworkSimulationModelTypes.h
@@ -164,66 +164,37 @@ struct TNetworkSimBufferContainer
 	}
 };
 
-/** Interface base for the simulation driver. This is what is used by the Network Sim Model internally. Basically, the non simulation specific functions that need to be defined. */
-template<typename TBufferTypes>
-class TNetworkSimDriverInterfaceBase
-{
-public:
-	virtual FString GetDebugName() const = 0; // Used for debugging. Recommended to emit the simulation name and the actor name/role.
-	virtual const UObject* GetVLogOwner() const = 0; // Owning object for Visual Logs
-
-	virtual void InitSyncState(typename TBufferTypes::TSyncState& OutSyncState) const = 0;	// Called to create initial value of the sync state.
-	virtual void ProduceInput(const float DeltaTimeSeconds, typename TBufferTypes::TInputCmd&) = 0; // Called when the sim is ready to process new local input
-	virtual void FinalizeFrame(const typename TBufferTypes::TSyncState& SyncState) = 0; // Called from the Network Sim at the end of the sim frame when there is new sync data.
-};
 
 // ----------------------------------------------------------------------------------------------------------------------------------------------
 //		Tick and time keeping related structures
 // ----------------------------------------------------------------------------------------------------------------------------------------------
 
-// Holds all settings for a network sim related to ticking
-// Main point are we want to distinguish between real time coming in from the outside in seconds. Float by default.
-// "SimTime" meaning time type we when advancing the simulation or talking about how long it has been running. This is stored in int32 milliseconds by default.
-// (This can be changed by the InRealToSimFactor parameter. Though code will always refer to is as MSec for clarity/consistency).
-template<int32 InFixedStepMS=0, int32 InMaxStepMS=0, typename TUnderlyingSimTimeType=int32, typename TUnderlyingRealTimeType=float, int32 InRealToSimFactor=1000>
-struct TNetworkSimTickSettings
+// The main Simulation time type. All sims use this to talk about time.
+struct FNetworkSimTime
 {
-	using TUnderlingSimTime = TUnderlyingSimTimeType;	// Underlying type of time. This type shouldn't be used directly (use TNetworkSimTime)
-	using TRealTime = TUnderlyingRealTimeType;	// This is the "final" real time type. We do not wrap it in anything else.
+	using FSimTime = int32; // Underlying type used to store simulation time
+	using FRealTime = float; // Underlying type used when dealing with realtime (usually coming in from the engine tick).
 
-	static_assert(!(InFixedStepMS!=0 && InMaxStepMS != 0), "MaxStepMS is only applicable when using variable step (FixedStepMS == 0)");
-
-	enum 
+	enum
 	{
-		MaxStepMS = InMaxStepMS,				// Max step. Only applicable to variable time step.
-		FixedStepMS = InFixedStepMS,			// Fixed step. If 0, then we are "variable time step"
-		RealToSimFactor = InRealToSimFactor		// Factor to go from RealTime (always seconds) to SimTime (MSec by default with factor of 1000)
+		RealToSimFactor = 1000 // Factor to go from RealTime (always seconds) to SimTime (MSec by default with factor of 1000)
 	};
 
-	// Typed accessors
-	static constexpr TUnderlingSimTime GetMaxStepMS() { return static_cast<TUnderlingSimTime>(InMaxStepMS); }
-	static constexpr TUnderlingSimTime GetFixedStepMS() { return static_cast<TUnderlingSimTime>(InFixedStepMS); }
-	static constexpr TRealTime GetRealToSimFactor() { return static_cast<TRealTime>(InRealToSimFactor); }
-	static constexpr TRealTime GetSimToRealFactor() { return static_cast<TRealTime>(1.f / InRealToSimFactor); }
-};
+	static constexpr FRealTime GetRealToSimFactor() { return static_cast<FRealTime>(RealToSimFactor); }
+	static constexpr FRealTime GetSimToRealFactor() { return static_cast<FRealTime>(1.f / RealToSimFactor); }
 
-// Actual time value. This by default will store time in MSec (ultimately determined by TickSettings::RealToSimFactor)
-template<typename TickSettings>
-struct TNetworkSimTime
-{
-	using TTime = typename TickSettings::TUnderlingSimTime;
-	using TRealTime = typename TickSettings::TRealTime;
+	// ---------------------------------------------------------------
 
-	TNetworkSimTime() { }
+	FNetworkSimTime() { }
 
 	// Things get confusing with templated types and overloaded functions. To avoid that, use these funcs to construct from either msec or real time
-	static inline TNetworkSimTime FromMSec(const TTime& InTime) { return TNetworkSimTime(InTime); } 
-	static inline TNetworkSimTime FromRealTimeSeconds(const TRealTime& InRealTime) { return TNetworkSimTime( static_cast<TTime>(InRealTime * TickSettings::GetRealToSimFactor())); } 
-	TRealTime ToRealTimeSeconds() const { return (Time * TickSettings::GetSimToRealFactor()); }
+	static inline FNetworkSimTime FromMSec(const FSimTime& InTime) { return FNetworkSimTime(InTime); } 
+	static inline FNetworkSimTime FromRealTimeSeconds(const FRealTime& InRealTime) { return FNetworkSimTime( static_cast<FSimTime>(InRealTime * GetRealToSimFactor())); } 
+	FRealTime ToRealTimeSeconds() const { return (Time * GetSimToRealFactor()); }
 
 	// Direct casts to "real time MS" which should be rarely used in practice (TRealTimeAccumulator only current case). All other cases of "real time" imply seconds.
-	static inline TNetworkSimTime FromRealTimeMS(const TRealTime& InRealTime) { return TNetworkSimTime( static_cast<TTime>(InRealTime)); } 
-	TRealTime ToRealTimeMS() const { return static_cast<TRealTime>(Time); }
+	static inline FNetworkSimTime FromRealTimeMS(const FRealTime& InRealTime) { return FNetworkSimTime( static_cast<FSimTime>(InRealTime)); } 
+	FRealTime ToRealTimeMS() const { return static_cast<FRealTime>(Time); }
 
 	FString ToString() const { return LexToString(this->Time); }
 
@@ -234,7 +205,7 @@ struct TNetworkSimTime
 	// FIXME
 	void NetSerialize(FArchive& Ar) { Ar << Time; }
 
-	using T = TNetworkSimTime;
+	using T = FNetworkSimTime;
 	T& operator+= (const T &rhs) { this->Time += rhs.Time; return(*this); }
 	T& operator-= (const T &rhs) { this->Time -= rhs.Time; return(*this); }
 	
@@ -248,29 +219,45 @@ struct TNetworkSimTime
 	bool operator== (const T &rhs) const { return(this->Time == rhs.Time); }
 	bool operator!= (const T &rhs) const { return(this->Time != rhs.Time); }
 
-	TTime Time = 0;
-
 private:
-	TNetworkSimTime(const TTime& InTime) { Time = InTime; }
+
+	FSimTime Time = 0;
+
+	FNetworkSimTime(const FSimTime& InTime) { Time = InTime; }
+};
+
+// Holds per-simulation settings about how ticking is supposed to happen.
+template<int32 InFixedStepMS=0, int32 InMaxStepMS=0>
+struct TNetworkSimTickSettings
+{
+	static_assert(!(InFixedStepMS!=0 && InMaxStepMS != 0), "MaxStepMS is only applicable when using variable step (FixedStepMS == 0)");
+	enum 
+	{
+		MaxStepMS = InMaxStepMS,				// Max step. Only applicable to variable time step.
+		FixedStepMS = InFixedStepMS,			// Fixed step. If 0, then we are "variable time step"
+	};
+
+	// Typed accessors
+	static constexpr FNetworkSimTime::FSimTime GetMaxStepMS() { return static_cast<FNetworkSimTime::FSimTime>(InMaxStepMS); }
+	static constexpr FNetworkSimTime::FSimTime GetFixedStepMS() { return static_cast<FNetworkSimTime::FSimTime>(InFixedStepMS); }
 };
 
 // ----------------------------------------------------------------------------------------------------------------------------------------------
-//
+//	Accumulator: Helper for accumulating real time into sim time based on TickSettings
 // ----------------------------------------------------------------------------------------------------------------------------------------------
 
-// Helper for accumulating real time into sim time based on TickSettings
 template<typename TickSettings, bool IsFixedTick=(TickSettings::FixedStepMS!=0)>
 struct TRealTimeAccumulator
 {
-	using TRealTime = typename TickSettings::TRealTime;
-	void Accumulate(TNetworkSimTime<TickSettings>& NetworkSimTime, const TRealTime RealTimeSeconds)
+	using TRealTime = FNetworkSimTime::FRealTime;
+	void Accumulate(FNetworkSimTime& NetworkSimTime, const TRealTime RealTimeSeconds)
 	{
 		// Even though we are variable tick, we still want to truncate down to an even msec. This keeps sim steps as whole integer values that serialize better, don't have denormals or other floating point weirdness.
 		// (If we ever wanted to just fully let floats pass through and be used as the msec sim time, this could be done through another specialization for float/float time).		
 		// Also note that MaxStepMS enforcement does NOT belong here. Dropping time due to MaxStepMS would just make the sim run slower. MaxStepMS is used at the input processing level.
 		
-		AccumulatedTimeMS += RealTimeSeconds * TickSettings::GetRealToSimFactor(); // convert input seconds -> msec
-		const TNetworkSimTime<TickSettings> AccumulatedSimTimeMS = TNetworkSimTime<TickSettings>::FromRealTimeMS(AccumulatedTimeMS);	// truncate (float) accumulated msec to (int32) sim time msec
+		AccumulatedTimeMS += RealTimeSeconds * FNetworkSimTime::GetRealToSimFactor(); // convert input seconds -> msec
+		const FNetworkSimTime AccumulatedSimTimeMS = FNetworkSimTime::FromRealTimeMS(AccumulatedTimeMS);	// truncate (float) accumulated msec to (int32) sim time msec
 
 		NetworkSimTime += AccumulatedSimTimeMS;
 		AccumulatedTimeMS -= AccumulatedSimTimeMS.ToRealTimeMS(); // subtract out the "whole" msec, we are left with the remainder msec
@@ -290,10 +277,10 @@ private:
 template<typename TickSettings>
 struct TRealTimeAccumulator<TickSettings, true>
 {
-	using TRealTime = typename TickSettings::TRealTime;
-	const TRealTime RealTimeFixedStep = static_cast<TRealTime>(TickSettings::GetFixedStepMS() * TickSettings::GetSimToRealFactor());
+	using TRealTime = FNetworkSimTime::FRealTime;
+	const TRealTime RealTimeFixedStep = static_cast<TRealTime>(TickSettings::GetFixedStepMS() * FNetworkSimTime::GetSimToRealFactor());
 
-	void Accumulate(TNetworkSimTime<TickSettings>& NetworkSimTime, const TRealTime RealTimeSeconds)
+	void Accumulate(FNetworkSimTime& NetworkSimTime, const TRealTime RealTimeSeconds)
 	{
 		AccumulatedTime += RealTimeSeconds;
 		if (AccumulatedTime > RealTimeFixedStep)
@@ -306,7 +293,7 @@ struct TRealTimeAccumulator<TickSettings, true>
 				AccumulatedTime = TRealTime(0.f);
 			}
 
-			NetworkSimTime += TNetworkSimTime<TickSettings>::FromMSec(NumFrames * TickSettings::FixedStepMS);
+			NetworkSimTime += FNetworkSimTime::FromMSec(NumFrames * TickSettings::FixedStepMS);
 		}
 	}
 
@@ -333,25 +320,25 @@ struct TSimulationTickState
 	int32 LastProcessedInputKeyframe = 0;	// The last input keyframe that we processed
 	int32 MaxAllowedInputKeyframe = 0;		// The max input keyframe that we are allowed to process (e.g, don't process input past this keyframe yet)
 
-	void SetTotalAllowedSimulationTime(const TNetworkSimTime<TSettings>& SimTime)
+	void SetTotalAllowedSimulationTime(const FNetworkSimTime& SimTime)
 	{
 		TotalAllowedSimulationTime = SimTime;
 		RealtimeAccumulator.Reset();
 	}
 
-	TNetworkSimTime<TSettings> GetTotalProcessedSimulationTime() const 
+	FNetworkSimTime GetTotalProcessedSimulationTime() const 
 	{ 
 		return TotalProcessedSimulationTime; 
 	}
 
-	void SetTotalProcessedSimulationTime(const TNetworkSimTime<TSettings>& SimTime, int32 Keyframe)
+	void SetTotalProcessedSimulationTime(const FNetworkSimTime& SimTime, int32 Keyframe)
 	{
 		TotalProcessedSimulationTime = SimTime;
 		SimulationTimeBuffer.ResetNextHeadKeyframe(Keyframe);
 		*SimulationTimeBuffer.GetWriteNext() = SimTime;
 	}
 
-	void IncrementTotalProcessedSimulationTime(const TNetworkSimTime<TSettings>& DeltaSimTime, int32 Keyframe)
+	void IncrementTotalProcessedSimulationTime(const FNetworkSimTime& DeltaSimTime, int32 Keyframe)
 	{
 		TotalProcessedSimulationTime += DeltaSimTime;
 		*SimulationTimeBuffer.GetWriteNext() = TotalProcessedSimulationTime;
@@ -361,11 +348,10 @@ struct TSimulationTickState
 	void InitSimulationTimeBuffer(int32 Size)
 	{
 		SimulationTimeBuffer.SetBufferSize(Size);
-		//*SimulationTimeBuffer.GetWriteNext() = TNetworkSimTime<TSettings>::FromMSec(0);
 	}
 		
 	// Historic tracking of simulation time. This allows us to timestamp sync data as its produced
-	TReplicationBuffer<TNetworkSimTime<TSettings>>	SimulationTimeBuffer;
+	TReplicationBuffer<FNetworkSimTime>	SimulationTimeBuffer;
 
 	// "Grants" allowed simulation time to this tick state. That is, we are now allowed to advance the simulation by this amount the next time the sim ticks.
 	// Note the input is RealTime in SECONDS. This is what the rest of the engine uses when dealing with float delta time.
@@ -375,20 +361,20 @@ struct TSimulationTickState
 	}
 
 	// How much granted simulation time is left to process
-	TNetworkSimTime<TSettings> GetRemainingAllowedSimulationTime() const
+	FNetworkSimTime GetRemainingAllowedSimulationTime() const
 	{
 		return TotalAllowedSimulationTime - TotalProcessedSimulationTime;
 	}
 
-	TNetworkSimTime<TSettings> GetTotalAllowedSimulationTime() const
+	FNetworkSimTime GetTotalAllowedSimulationTime() const
 	{
 		return TotalAllowedSimulationTime;
 	}
 
 private:
 
-	TNetworkSimTime<TSettings> TotalAllowedSimulationTime;	// Total time we have been "given" to process. We cannot process more simulation time than this: doing so would be speed hacking.
-	TNetworkSimTime<TSettings> TotalProcessedSimulationTime;	// How much time we've actually processed. The only way to increment this is to process user commands or receive authoritative state from the network.
+	FNetworkSimTime TotalAllowedSimulationTime;	// Total time we have been "given" to process. We cannot process more simulation time than this: doing so would be speed hacking.
+	FNetworkSimTime TotalProcessedSimulationTime;	// How much time we've actually processed. The only way to increment this is to process user commands or receive authoritative state from the network.
 
 	TRealTimeAccumulator<TSettings>	RealtimeAccumulator;
 };
@@ -401,8 +387,8 @@ private:
 template<typename BaseCmdType, typename TickSettings, bool IsFixedTick=(TickSettings::FixedStepMS!=0)>
 struct TFrameCmd : public BaseCmdType
 {
-	TNetworkSimTime<TickSettings> GetFrameDeltaTime() const { return FrameDeltaTime; }
-	void SetFrameDeltaTime(const TNetworkSimTime<TickSettings>& InTime) { FrameDeltaTime = InTime; }
+	FNetworkSimTime GetFrameDeltaTime() const { return FrameDeltaTime; }
+	void SetFrameDeltaTime(const FNetworkSimTime& InTime) { FrameDeltaTime = InTime; }
 	void NetSerialize(const FNetSerializeParams& P)
 	{
 		FrameDeltaTime.NetSerialize(P.Ar);
@@ -410,15 +396,15 @@ struct TFrameCmd : public BaseCmdType
 	}
 
 private:
-	TNetworkSimTime<TickSettings> FrameDeltaTime;
+	FNetworkSimTime FrameDeltaTime;
 };
 
 // Fixed tick specialization
 template<typename BaseCmdType, typename TickSettings>
 struct TFrameCmd<BaseCmdType, TickSettings, true> : public BaseCmdType
 {
-	constexpr TNetworkSimTime<TickSettings> GetFrameDeltaTime() const { return TNetworkSimTime<TickSettings>::FromMSec(TickSettings::GetFixedStepMS()); }
-	void SetFrameDeltaTime(const TNetworkSimTime<TickSettings>& InTime) { }
+	FNetworkSimTime GetFrameDeltaTime() const { return FNetworkSimTime::FromMSec(TickSettings::GetFixedStepMS()); }
+	void SetFrameDeltaTime(const FNetworkSimTime& InTime) { }
 	void NetSerialize(const FNetSerializeParams& P) { BaseCmdType::NetSerialize(P); }
 };
 
@@ -434,4 +420,17 @@ struct TInternalBufferTypes : TNetworkSimBufferTypes<
 	typename TUserBufferTypes::TDebugState	// Debugstate passes through
 >
 {
+};
+
+/** Interface base for the simulation driver. This is what is used by the Network Sim Model internally. Basically, the non simulation specific functions that need to be defined. */
+template<typename TBufferTypes>
+class TNetworkSimDriverInterfaceBase
+{
+public:
+	virtual FString GetDebugName() const = 0; // Used for debugging. Recommended to emit the simulation name and the actor name/role.
+	virtual const UObject* GetVLogOwner() const = 0; // Owning object for Visual Logs
+
+	virtual void InitSyncState(typename TBufferTypes::TSyncState& OutSyncState) const = 0;	// Called to create initial value of the sync state.
+	virtual void ProduceInput(const FNetworkSimTime SimTime, typename TBufferTypes::TInputCmd&) = 0; // Called when the sim is ready to process new local input
+	virtual void FinalizeFrame(const typename TBufferTypes::TSyncState& SyncState) = 0; // Called from the Network Sim at the end of the sim frame when there is new sync data.
 };
\ No newline at end of file
diff --git a/Engine/Plugins/Runtime/NetworkPredictionExtras/Source/NetworkPredictionExtras/Private/MockNetworkSimulation.cpp b/Engine/Plugins/Runtime/NetworkPredictionExtras/Source/NetworkPredictionExtras/Private/MockNetworkSimulation.cpp
--- a/Engine/Plugins/Runtime/NetworkPredictionExtras/Source/NetworkPredictionExtras/Private/MockNetworkSimulation.cpp
+++ b/Engine/Plugins/Runtime/NetworkPredictionExtras/Source/NetworkPredictionExtras/Private/MockNetworkSimulation.cpp
@@ -154,7 +154,7 @@ void UMockNetworkSimulationComponent::TickComponent(float DeltaTime, enum ELevel
 	DrawDebugString( GetWorld(), GetOwner()->GetActorLocation() + FVector(0.f,0.f,100.f), *LexToString(MockValue), nullptr, FColor::White, 0.00001f );
 }
 
-void UMockNetworkSimulationComponent::ProduceInput(const float DeltaTimeSeconds, FMockInputCmd& Cmd)
+void UMockNetworkSimulationComponent::ProduceInput(const FNetworkSimTime SimTime, FMockInputCmd& Cmd)
 {
 	if (MockNetworkSimCVars::DoLocalInput)
 	{
diff --git a/Engine/Plugins/Runtime/NetworkPredictionExtras/Source/NetworkPredictionExtras/Private/NetworkPredictionExtrasFlyingPawn.cpp b/Engine/Plugins/Runtime/NetworkPredictionExtras/Source/NetworkPredictionExtras/Private/NetworkPredictionExtrasFlyingPawn.cpp
--- a/Engine/Plugins/Runtime/NetworkPredictionExtras/Source/NetworkPredictionExtras/Private/NetworkPredictionExtrasFlyingPawn.cpp
+++ b/Engine/Plugins/Runtime/NetworkPredictionExtras/Source/NetworkPredictionExtras/Private/NetworkPredictionExtrasFlyingPawn.cpp
@@ -135,7 +135,7 @@ void ANetworkPredictionExtrasFlyingPawn::Tick( float DeltaSeconds)
 	// Do whatever you want here. By now we have the latest movement state and latest input processed.
 }
 
-void ANetworkPredictionExtrasFlyingPawn::ProduceInput(const float DeltaTimeSeconds, FlyingMovement::FInputCmd& Cmd)
+void ANetworkPredictionExtrasFlyingPawn::ProduceInput(const FNetworkSimTime SimTime, FlyingMovement::FInputCmd& Cmd)
 {
 	// Generate user commands. Called right before the flying movement simulation will tick (for a locally controlled pawn)
 	// This isn't meant to be the best way of doing a camera system. It is just meant to show a couple of ways it may be done
@@ -172,6 +172,8 @@ void ANetworkPredictionExtrasFlyingPawn::ProduceInput(const float DeltaTimeSecon
 	Cmd.MovementInput = FVector::ZeroVector;
 	Cmd.RotationInput = FRotator::ZeroRotator;
 
+	const float DeltaTimeSeconds = SimTime.ToRealTimeSeconds();
+
 	switch (FlyingPawnCVars::CameraStyle)
 	{
 		case 0:
diff --git a/Engine/Plugins/Runtime/NetworkPredictionExtras/Source/NetworkPredictionExtras/Public/MockNetworkSimulation.h b/Engine/Plugins/Runtime/NetworkPredictionExtras/Source/NetworkPredictionExtras/Public/MockNetworkSimulation.h
--- a/Engine/Plugins/Runtime/NetworkPredictionExtras/Source/NetworkPredictionExtras/Public/MockNetworkSimulation.h
+++ b/Engine/Plugins/Runtime/NetworkPredictionExtras/Source/NetworkPredictionExtras/Public/MockNetworkSimulation.h
@@ -174,7 +174,7 @@ public:
 
 	void InitSyncState(FMockSyncState& OutSyncState) const override;
 	void FinalizeFrame(const FMockSyncState& SyncState) override;
-	void ProduceInput(const float DeltaTimeSeconds, FMockInputCmd& Cmd);
+	void ProduceInput(const FNetworkSimTime SimTime, FMockInputCmd& Cmd);
 	virtual UWorld* GetDriverWorld() const override final { return GetWorld(); }
 	virtual FTransform GetDebugWorldTransform() const override final;
 
diff --git a/Engine/Plugins/Runtime/NetworkPredictionExtras/Source/NetworkPredictionExtras/Public/NetworkPredictionExtrasFlyingPawn.h b/Engine/Plugins/Runtime/NetworkPredictionExtras/Source/NetworkPredictionExtras/Public/NetworkPredictionExtrasFlyingPawn.h
--- a/Engine/Plugins/Runtime/NetworkPredictionExtras/Source/NetworkPredictionExtras/Public/NetworkPredictionExtrasFlyingPawn.h
+++ b/Engine/Plugins/Runtime/NetworkPredictionExtras/Source/NetworkPredictionExtras/Public/NetworkPredictionExtrasFlyingPawn.h
@@ -45,7 +45,7 @@ class NETWORKPREDICTIONEXTRAS_API ANetworkPredictionExtrasFlyingPawn : public AP
 
 private:
 
-	void ProduceInput(const float TimeSeconds, FlyingMovement::FInputCmd& Cmd);
+	void ProduceInput(const FNetworkSimTime SimTime, FlyingMovement::FInputCmd& Cmd);
 
 	FVector CachedMoveInput;
 	FVector2D CachedLookInput;
