NetworkPrediction plugin: Forwad Predicted simulated proxy support (aka dependent simulations).

#rnx

diff --git a/Engine/Plugins/Runtime/NetworkPrediction/Source/NetworkPrediction/Private/Movement/ParametricMovement.cpp b/Engine/Plugins/Runtime/NetworkPrediction/Source/NetworkPrediction/Private/Movement/ParametricMovement.cpp
--- a/Engine/Plugins/Runtime/NetworkPrediction/Source/NetworkPrediction/Private/Movement/ParametricMovement.cpp
+++ b/Engine/Plugins/Runtime/NetworkPrediction/Source/NetworkPrediction/Private/Movement/ParametricMovement.cpp
@@ -4,6 +4,8 @@
 #include "NetworkSimulationModelDebugger.h"
 #include "HAL/IConsoleManager.h"
 #include "NetworkSimulationGlobalManager.h"
+#include "Engine/LocalPlayer.h"
+#include "GameFramework/PlayerController.h"
 
 DEFINE_LOG_CATEGORY_STATIC(LogParametricMovement, Log, All);
 
@@ -199,6 +201,28 @@ void UParametricMovementComponent::TickComponent(float DeltaTime, enum ELevelTic
 {
 	Super::TickComponent(DeltaTime, TickType, ThisTickFunction);
 
+	if (bEnableDependentSimulation && NetworkSim.IsValid() && NetworkSim->GetParentSimulation() == nullptr)
+	{
+		// This is a very simply and generic way of finding a locally controlled network sim.
+		// The long term vision for dependent simulations is more game specific and dynamic. 
+		// E.g, stepping on elevators, getting close to moving platforms, etc. This example here is meant to 
+		// just be a simple way of linking them up no matter where they are in relationship to each other
+		UWorld* World = GetWorld();
+		check(World);
+		ULocalPlayer* Player = World->GetFirstLocalPlayerFromController();
+		if (Player && Player->GetPlayerController(World))
+		{
+			APawn* Pawn = Player->GetPlayerController(World)->GetPawn();
+			if (UNetworkPredictionComponent* NetComponent = Pawn->FindComponentByClass<UNetworkPredictionComponent>())
+			{
+				if (INetworkSimulationModel* ParentNetSim = NetComponent->GetNetworkSimulation())
+				{
+					NetworkSim->SetParentSimulation(ParentNetSim);
+				}
+			}
+		}
+	}
+
 	if (ParentNetUpdateFrequency > 0.f)
 	{
 		GetOwner()->NetUpdateFrequency = ParentNetUpdateFrequency;
diff --git a/Engine/Plugins/Runtime/NetworkPrediction/Source/NetworkPrediction/Public/Movement/ParametricMovement.h b/Engine/Plugins/Runtime/NetworkPrediction/Source/NetworkPrediction/Public/Movement/ParametricMovement.h
--- a/Engine/Plugins/Runtime/NetworkPrediction/Source/NetworkPrediction/Public/Movement/ParametricMovement.h
+++ b/Engine/Plugins/Runtime/NetworkPrediction/Source/NetworkPrediction/Public/Movement/ParametricMovement.h
@@ -190,6 +190,9 @@ protected:
 	UPROPERTY(EditAnywhere, BlueprintReadOnly, Category=ParametricMovement)
 	float MaxTime = 1.f;
 
+	UPROPERTY(EditAnywhere, BlueprintReadOnly, Category=ParametricMovementNetworking)
+	bool bEnableDependentSimulation = false;
+
 	UPROPERTY(EditAnywhere, BlueprintReadOnly, Category=ParametricMovementNetworking)
 	bool bEnableInterpolation = true;
 
diff --git a/Engine/Plugins/Runtime/NetworkPrediction/Source/NetworkPrediction/Public/NetworkPredictionComponent.h b/Engine/Plugins/Runtime/NetworkPrediction/Source/NetworkPrediction/Public/NetworkPredictionComponent.h
--- a/Engine/Plugins/Runtime/NetworkPrediction/Source/NetworkPrediction/Public/NetworkPredictionComponent.h
+++ b/Engine/Plugins/Runtime/NetworkPrediction/Source/NetworkPrediction/Public/NetworkPredictionComponent.h
@@ -34,6 +34,8 @@ public:
 	virtual void Reconcile() { }
 	virtual void TickSimulation(float DeltaTimeSeconds) { }
 
+	INetworkSimulationModel* GetNetworkSimulation() const { return NetworkSim.Get(); }
+
 protected:
 	
 	// Classes must instantiate their own NetworkSim here. The UNetworkPredictionComponent will manage its lifetime
diff --git a/Engine/Plugins/Runtime/NetworkPrediction/Source/NetworkPrediction/Public/NetworkPredictionTypes.h b/Engine/Plugins/Runtime/NetworkPrediction/Source/NetworkPrediction/Public/NetworkPredictionTypes.h
--- a/Engine/Plugins/Runtime/NetworkPrediction/Source/NetworkPrediction/Public/NetworkPredictionTypes.h
+++ b/Engine/Plugins/Runtime/NetworkPrediction/Source/NetworkPrediction/Public/NetworkPredictionTypes.h
@@ -212,6 +212,24 @@ public:
 
 	virtual bool ShouldSendServerRPC(float DeltaSeconds) = 0;
 	virtual void SetDesiredServerRPCSendFrequency(float DesiredHz) = 0;
+
+	// ----------------------------------------------------------------------
+	// Functions for depedent simulation (forward predicting a simulated proxy sim along with an auto proxy sim)
+	// ----------------------------------------------------------------------
+	
+	// Main function to call on simulated proxy sim
+	virtual void SetParentSimulation(INetworkSimulationModel* Simulation) = 0;
+	virtual INetworkSimulationModel* GetParentSimulation() const = 0;
+	
+	virtual void AddDepdentSimulation(INetworkSimulationModel* Simulation) = 0;
+	virtual void RemoveDependentSimulation(INetworkSimulationModel* Simulation) = 0;
+	
+	// Tell parent sim that a dependent sim needs to reconcile (parent sim drives this)
+	virtual void NotifyDependentSimNeedsReconcile() = 0;
+	
+	// Called by parent sim on the dependent sim as it reconciles
+	virtual void BeginRollback(const struct FNetworkSimTime& RollbackDeltaTime, const int32 ParentKeyframe) = 0;
+	virtual void StepRollback(const struct FNetworkSimTime& Step, const int32 ParentKeyframe, const bool FinalStep) = 0;
 };
 
 // -------------------------------------------------------------------------------------------------------------------------------
diff --git a/Engine/Plugins/Runtime/NetworkPrediction/Source/NetworkPrediction/Public/NetworkSimulationModel.h b/Engine/Plugins/Runtime/NetworkPrediction/Source/NetworkPrediction/Public/NetworkSimulationModel.h
--- a/Engine/Plugins/Runtime/NetworkPrediction/Source/NetworkPrediction/Public/NetworkSimulationModel.h
+++ b/Engine/Plugins/Runtime/NetworkPrediction/Source/NetworkPrediction/Public/NetworkSimulationModel.h
@@ -66,6 +66,11 @@ public:
 	{
 		Driver = InDriver;
 	}
+	virtual ~TNetworkedSimulationModel()
+	{
+		SetParentSimulation(nullptr);
+		ClearAllDependentSimulations();
+	}
 
 	void Tick(const FNetSimTickParameters& Parameters) final override
 	{
@@ -344,6 +349,70 @@ public:
 		};
 	}
 
+	ESimulatedUpdateMode GetSimulatedUpdateMode() const
+	{
+		return RepProxy_Simulated.GetSimulatedUpdateMode();
+	}
+
+	// ------------------------------------------------------------------------------------------------------
+
+	void SetParentSimulation(INetworkSimulationModel* Simulation) final override
+	{
+		if (RepProxy_Simulated.ParentSimulation)
+		{
+			RepProxy_Simulated.ParentSimulation->RemoveDependentSimulation(this);
+		}
+		
+		RepProxy_Simulated.ParentSimulation = Simulation;
+		if (Simulation)
+		{
+			Simulation->AddDepdentSimulation(this);
+		}
+	}
+
+	INetworkSimulationModel* GetParentSimulation() const final override
+	{
+		return RepProxy_Simulated.ParentSimulation;
+	}
+
+	void AddDepdentSimulation(INetworkSimulationModel* Simulation) final override
+	{
+		check(RepProxy_Autonomous.DependentSimulations.Contains(Simulation) == false);
+		RepProxy_Autonomous.DependentSimulations.Add(Simulation);
+		NotifyDependentSimNeedsReconcile(); // force reconcile on purpose
+	}
+
+	void RemoveDependentSimulation(INetworkSimulationModel* Simulation) final override
+	{
+		RepProxy_Autonomous.DependentSimulations.Remove(Simulation);
+	}
+
+	void NotifyDependentSimNeedsReconcile()
+	{
+		RepProxy_Autonomous.bDependentSimulationNeedsReconcile = true;
+	}
+
+	void BeginRollback(const FNetworkSimTime& RollbackDeltaTime, const int32 ParentKeyframe) final override
+	{
+		RepProxy_Simulated.template DependentRollbackBegin<TSimulation, TDriver>(Driver, Buffers, TickInfo, RollbackDeltaTime, ParentKeyframe);
+	}
+
+	void StepRollback(const FNetworkSimTime& Step, const int32 ParentKeyframe, const bool bFinalStep) final override
+	{
+		RepProxy_Simulated.template DependentRollbackStep<TSimulation, TDriver>(Driver, Buffers, TickInfo, Step, ParentKeyframe, bFinalStep);
+	}
+
+	void ClearAllDependentSimulations()
+	{
+		TArray<INetworkSimulationModel*> LocalList = MoveTemp(RepProxy_Autonomous.DependentSimulations);
+		for (INetworkSimulationModel* DependentSim : LocalList)
+		{
+			DependentSim->SetParentSimulation(nullptr);
+		}
+	}
+
+	// -------------------------------------------------------------------------------------------------------
+
 	TDriver* Driver = nullptr;	
 	TSimulationTickState<TTickSettings> TickInfo;	// Manages simulation time and what inputs we are processed
 
diff --git a/Engine/Plugins/Runtime/NetworkPrediction/Source/NetworkPrediction/Public/NetworkSimulationModelCVars.h b/Engine/Plugins/Runtime/NetworkPrediction/Source/NetworkPrediction/Public/NetworkSimulationModelCVars.h
--- a/Engine/Plugins/Runtime/NetworkPrediction/Source/NetworkPrediction/Public/NetworkSimulationModelCVars.h
+++ b/Engine/Plugins/Runtime/NetworkPrediction/Source/NetworkPrediction/Public/NetworkSimulationModelCVars.h
@@ -38,6 +38,7 @@ struct FConditionalAutoConsoleRegister
 #if NETSIM_CONST_CVARS
 #define NETSIM_DEVCVAR_SHIPCONST_INT(Var,Value,VarName,Help) \
 	inline int32 Var() { return Value; }
+	inline void Set##Var(int32 V) { }
 #else
 #define NETSIM_DEVCVAR_SHIPCONST_INT(Var,Value,VarName,Help) \
 	static FConditionalAutoConsoleRegister Var##Auto(TEXT(VarName),(int32)Value,TEXT(Help)); \
@@ -46,12 +47,19 @@ struct FConditionalAutoConsoleRegister
 		static const auto* Existing = IConsoleManager::Get().FindConsoleVariable(TEXT(VarName), false); \
 		check(Existing) \
 		return Existing->GetInt(); \
+	} \
+	inline void Set##Var(int32 V) \
+	{ \
+		static auto* Existing = IConsoleManager::Get().FindConsoleVariable(TEXT(VarName), false); \
+		check(Existing) \
+		Existing->Set(V, ECVF_SetByConsole); \
 	}
 #endif
 
 #if NETSIM_CONST_CVARS
 #define NETSIM_DEVCVAR_SHIPCONST_FLOAT(Var,Value,VarName,Help) \
 	inline float Var() { return Value; }
+	inline void Set##Vat() { }
 #else
 #define NETSIM_DEVCVAR_SHIPCONST_FLOAT(Var,Value,VarName,Help) \
 	static FConditionalAutoConsoleRegister Var##Auto(TEXT(VarName),(float)Value,TEXT(Help)); \
@@ -60,5 +68,11 @@ struct FConditionalAutoConsoleRegister
 		static const auto* Existing = IConsoleManager::Get().FindConsoleVariable(TEXT(VarName), false); \
 		check(Existing) \
 		return Existing->GetFloat(); \
+	} \
+	inline void Set##Var(float V) \
+	{ \
+		static auto* Existing = IConsoleManager::Get().FindConsoleVariable(TEXT(VarName), false); \
+		check(Existing) \
+		Existing->Set(V, ECVF_SetByConsole); \
 	}
 #endif
\ No newline at end of file
diff --git a/Engine/Plugins/Runtime/NetworkPrediction/Source/NetworkPrediction/Public/NetworkSimulationModelDebugger.h b/Engine/Plugins/Runtime/NetworkPrediction/Source/NetworkPrediction/Public/NetworkSimulationModelDebugger.h
--- a/Engine/Plugins/Runtime/NetworkPrediction/Source/NetworkPrediction/Public/NetworkSimulationModelDebugger.h
+++ b/Engine/Plugins/Runtime/NetworkPrediction/Source/NetworkPrediction/Public/NetworkSimulationModelDebugger.h
@@ -672,12 +672,28 @@ struct TNetworkSimulationModelDebugger : public INetworkSimulationModelDebugger
 			}
 		}
 
-		for (int32 Keyframe = NetworkSim->RepProxy_Autonomous.GetLastSerializedKeyframe(); Keyframe < NetworkSim->Buffers.Sync.GetHeadKeyframe(); ++Keyframe)
+		if (Owner->GetLocalRole() == ROLE_AutonomousProxy)
+		{
+			for (int32 Keyframe = NetworkSim->RepProxy_Autonomous.GetLastSerializedKeyframe(); Keyframe < NetworkSim->Buffers.Sync.GetHeadKeyframe(); ++Keyframe)
+			{
+				if (auto* SyncState = NetworkSim->Buffers.Sync.FindElementByKeyframe(Keyframe))
+				{
+					const EVisualLoggingContext Context = (Keyframe == NetworkSim->RepProxy_Autonomous.GetLastSerializedKeyframe()) ? EVisualLoggingContext::LastConfirmed : EVisualLoggingContext::OtherPredicted;
+					SyncState->VisualLog( FVisualLoggingParameters(Context, NetworkSim->Buffers.Sync.GetHeadKeyframe(), EVisualLoggingLifetime::Transient), NetworkSim->Driver, NetworkSim->Driver);
+				}
+			}
+		}
+		else if (Owner->GetLocalRole() == ROLE_SimulatedProxy)
 		{
-			if (auto* SyncState = NetworkSim->Buffers.Sync.FindElementByKeyframe(Keyframe))
+			NetworkSim->RepProxy_Simulated.GetLastSerializedSyncState().VisualLog( FVisualLoggingParameters(EVisualLoggingContext::LastConfirmed, NetworkSim->Buffers.Sync.GetHeadKeyframe(), EVisualLoggingLifetime::Transient), NetworkSim->Driver, NetworkSim->Driver); 
+			if (NetworkSim->GetSimulatedUpdateMode() != ESimulatedUpdateMode::Interpolate)
 			{
-				const EVisualLoggingContext Context = (Keyframe == NetworkSim->RepProxy_Autonomous.GetLastSerializedKeyframe()) ? EVisualLoggingContext::LastConfirmed : EVisualLoggingContext::OtherPredicted;
-				SyncState->VisualLog( FVisualLoggingParameters(Context, NetworkSim->Buffers.Sync.GetHeadKeyframe(), EVisualLoggingLifetime::Transient), NetworkSim->Driver, NetworkSim->Driver);
+				FVector2D ServerSimulationTimeData(Owner->GetWorld()->GetTimeSeconds(), NetworkSim->RepProxy_Simulated.GetLastSerializedSimulationTime().ToRealTimeMS());
+				UE_VLOG_HISTOGRAM(Owner, LogNetworkSimDebug, Log, "Simulated Time Graph", "Serialized Simulation Time", ServerSimulationTimeData);
+
+				FVector2D LocalSimulationTimeData(Owner->GetWorld()->GetTimeSeconds(), NetworkSim->TickInfo.GetTotalProcessedSimulationTime().ToRealTimeMS());
+				UE_VLOG_HISTOGRAM(Owner, LogNetworkSimDebug, Log, "Simulated Time Graph", "Local Simulation Time", LocalSimulationTimeData);
+			
 			}
 		}
 	}
diff --git a/Engine/Plugins/Runtime/NetworkPrediction/Source/NetworkPrediction/Public/NetworkSimulationModelReplicators.h b/Engine/Plugins/Runtime/NetworkPrediction/Source/NetworkPrediction/Public/NetworkSimulationModelReplicators.h
--- a/Engine/Plugins/Runtime/NetworkPrediction/Source/NetworkPrediction/Public/NetworkSimulationModelReplicators.h
+++ b/Engine/Plugins/Runtime/NetworkPrediction/Source/NetworkPrediction/Public/NetworkSimulationModelReplicators.h
@@ -16,6 +16,7 @@ namespace NetworkSimulationModelCVars
 	NETSIM_DEVCVAR_SHIPCONST_INT(EnableSimulatedReconcile, 1, "ns.EnableSimulatedReconcile", "Toggle simulated proxy reconciliation.");
 	NETSIM_DEVCVAR_SHIPCONST_INT(EnableSimulatedExtrapolation, 1, "ns.EnableSimulatedExtrapolation", "Toggle simulated proxy extrapolation.");
 	NETSIM_DEVCVAR_SHIPCONST_INT(ForceReconcile, 0, "ns.ForceReconcile", "Forces reconcile even if state does not differ. E.g, force resimulation after every netupdate.");
+	NETSIM_DEVCVAR_SHIPCONST_INT(ForceReconcileSingle, 0, "ns.ForceReconcileSingle", "Forces a since reconcile to happen on the next frame");
 }
 
 #ifndef NETSIM_NETCONSTANT_NUM_BITS_KEYFRAME
@@ -284,12 +285,32 @@ struct TReplicator_Simulated : public TBase
 	using TSyncState = typename TBufferTypes::TSyncState;
 	using TAuxState = typename TBufferTypes::TAuxState;
 
+	// Parent Simulation. If this is set, this simulation will forward predict in sync with this parent sim. The parent sim should be an autonomous proxy driven simulation
+	INetworkSimulationModel* ParentSimulation = nullptr;
+
 	// Instance flag for enabling simulated extrapolation
 	bool bAllowSimulatedExtrapolation = true;
 
+	// Interpolated that will be used if bAllowSimulatedExtrapolation == false && ParentSimulation == nullptr
 	TInterpolator<TBufferTypes, TTickSettings> Interpolator;
 
+	// Last Serialized time and state
 	FNetworkSimTime GetLastSerializedSimulationTime() const { return LastSerializedSimulationTime; }
+	const TSyncState& GetLastSerializedSyncState() const { return LastSerializedSyncState; }
+	
+	ESimulatedUpdateMode GetSimulatedUpdateMode() const
+	{
+		if (ParentSimulation)
+		{
+			return ESimulatedUpdateMode::ForwardPredict;
+		}
+		if (bAllowSimulatedExtrapolation && NetworkSimulationModelCVars::EnableSimulatedExtrapolation())
+		{
+			return ESimulatedUpdateMode::Extrapolate;
+		}
+
+		return ESimulatedUpdateMode::Interpolate;
+	}
 	
 	int32 GetProxyDirtyCount(TNetworkSimBufferContainer<TBufferTypes>& Buffers) const
 	{
@@ -331,8 +352,6 @@ struct TReplicator_Simulated : public TBase
 			{
 				// We are getting new state that is ahead of what we have locally, so it can safety go right to head
 				DestinationKeyframe = TickInfo.SimulationTimeBuffer.GetHeadKeyframe()+1;
-
-				//UE_LOG(LogNetworkSim, Warning, TEXT("@@@"));
 			}
 			else
 			{
@@ -341,7 +360,6 @@ struct TReplicator_Simulated : public TBase
 				{
 					if (LastSerializedSimulationTime > *TickInfo.SimulationTimeBuffer.FindElementByKeyframe(Keyframe))
 					{
-						//UE_LOG(LogNetworkSim, Warning, TEXT("..."));
 						DestinationKeyframe = Keyframe+1;
 						break;
 					}
@@ -382,6 +400,11 @@ struct TReplicator_Simulated : public TBase
 
 		check(State);
 		State->NetSerialize(Ar);
+
+		if (Ar.IsLoading())
+		{
+			LastSerializedSyncState = *State;
+		}
 	}
 
 	template<typename T, typename TDriver>
@@ -394,7 +417,7 @@ struct TReplicator_Simulated : public TBase
 
 		check(TickInfo.GetTotalProcessedSimulationTime() <= TickInfo.GetTotalAllowedSimulationTime());
 
-		if (bAllowSimulatedExtrapolation && NetworkSimulationModelCVars::EnableSimulatedExtrapolation() && NetworkSimulationModelCVars::EnableSimulatedReconcile())
+		if (bAllowSimulatedExtrapolation && ParentSimulation == nullptr && NetworkSimulationModelCVars::EnableSimulatedExtrapolation() && NetworkSimulationModelCVars::EnableSimulatedReconcile())
 		{
 			TInputCmd* LastCmd = Buffers.Input.GetElementFromHead(0);
 
@@ -410,30 +433,7 @@ struct TReplicator_Simulated : public TBase
 			FNetworkSimTime DeltaSimTime = ReconcileSimulationTime - TickInfo.GetTotalProcessedSimulationTime();
 			if (DeltaSimTime.IsPositive() && NetworkSimulationModelCVars::EnableSimulatedReconcile())
 			{
-				// We have extrapolated ahead of the server. The latest network update is now "in the past" from what we rendered last frame.
-				// We will insert a new keyframe to make up the difference from the last known state to where we want to be in the now.
-				
-				TInputCmd* NewCmd = Buffers.Input.GetWriteNext();
-				*NewCmd = LastCmd ? *LastCmd : TInputCmd();
-				NewCmd->SetFrameDeltaTime(DeltaSimTime);
-
-				TSyncState* PrevSyncState = Buffers.Sync.GetElementFromHead(0);
-				TSyncState* NextSyncState = Buffers.Sync.GetWriteNext();
-
-				TAuxState Junk;
-
-				// TODO: log this guy
-				//NextMotionState->VisualLog( FVisualLoggingParameters(Keyframe == LastKeyframeToProcess ? EVisualLoggingContext::LastMispredicted : EVisualLoggingContext::OtherMispredicted, Keyframe, EVisualLoggingLifetime::Persistent), Driver, Driver);
-
-				//UE_LOG(LogNetworkSim, Warning, TEXT("   >>> Reconcile Update: %s"), *DeltaSimTime.ToString());
-
-				// Do the actual update
-				T::Update(Driver, NewCmd->GetFrameDeltaTime().ToRealTimeSeconds(), *NewCmd, *PrevSyncState, *NextSyncState, Junk);
-				TickInfo.IncrementTotalProcessedSimulationTime(NewCmd->GetFrameDeltaTime(), Buffers.Sync.GetHeadKeyframe());
-
-				// Set our LastProcessedInputKeyframe to fake that we handled it
-				TickInfo.LastProcessedInputKeyframe = Buffers.Input.GetHeadKeyframe();
-				TickInfo.MaxAllowedInputKeyframe = Buffers.Input.GetHeadKeyframe();
+				SimulationExtrapolation<T, TDriver>(Driver, Buffers, TickInfo, DeltaSimTime);
 			}
 		}
 		
@@ -444,7 +444,8 @@ struct TReplicator_Simulated : public TBase
 	template<typename T, typename TDriver>
 	void PreSimTick(TDriver* Driver, TNetworkSimBufferContainer<TBufferTypes>& Buffers, TSimulationTickState<TTickSettings>& TickInfo, const FNetSimTickParameters& TickParameters)
 	{
-		if (bAllowSimulatedExtrapolation && NetworkSimulationModelCVars::EnableSimulatedExtrapolation())
+		// Tick if we are dependent simulation or extrapolation is enabled
+		if (ParentSimulation || (bAllowSimulatedExtrapolation && NetworkSimulationModelCVars::EnableSimulatedExtrapolation()))
 		{
 			// Don't start this simulation until you've gotten at least one update from the server
 			if (TickInfo.GetTotalProcessedSimulationTime().IsPositive())
@@ -474,7 +475,7 @@ struct TReplicator_Simulated : public TBase
 	template<typename TDriver>
 	void PostSimTick(TDriver* Driver, const TNetworkSimBufferContainer<TBufferTypes>& Buffers, const TSimulationTickState<TTickSettings>& TickInfo, const FNetSimTickParameters& TickParameters)
 	{
-		if (bAllowSimulatedExtrapolation)
+		if (bAllowSimulatedExtrapolation || ParentSimulation)
 		{
 			// Sync to latest frame if there is any
 			if (Buffers.Sync.GetNumValidElements() > 0)
@@ -489,11 +490,71 @@ struct TReplicator_Simulated : public TBase
 		}
 	}
 
+	template<typename T, typename TDriver>
+	void DependentRollbackBegin(TDriver* Driver, TNetworkSimBufferContainer<TBufferTypes>& Buffers, TSimulationTickState<TTickSettings>& TickInfo, const FNetworkSimTime& RollbackDeltaTime, const int32 ParentKeyframe)
+	{
+		// For now, we make the assumption that our last serialized state and time match the parent simulation.
+		// This would not always be the case with low frequency simulated proxies. But could be handled by replicating the simulations together (at the replication level)
+		int32 NewHeadKeyframe = Buffers.Sync.GetHeadKeyframe()+1;
+		Buffers.Sync.ResetNextHeadKeyframe(NewHeadKeyframe);
+		Buffers.Input.ResetNextHeadKeyframe(NewHeadKeyframe);
+		TickInfo.SetTotalProcessedSimulationTime(LastSerializedSimulationTime, NewHeadKeyframe);
+		TickInfo.SetTotalAllowedSimulationTime(LastSerializedSimulationTime);
+
+		*Buffers.Sync.GetWriteNext() = LastSerializedSyncState;
+		*Buffers.Input.GetWriteNext() = TInputCmd();
+
+		TickInfo.LastProcessedInputKeyframe = Buffers.Input.GetHeadKeyframe();
+		TickInfo.MaxAllowedInputKeyframe = Buffers.Input.GetHeadKeyframe();
+
+		Driver->FinalizeFrame(LastSerializedSyncState);
+
+		LastSerializedSyncState.VisualLog( FVisualLoggingParameters(EVisualLoggingContext::FirstMispredicted, ParentKeyframe, EVisualLoggingLifetime::Persistent), Driver, Driver );
+	}
+
+	template<typename T, typename TDriver>
+	void DependentRollbackStep(TDriver* Driver, TNetworkSimBufferContainer<TBufferTypes>& Buffers, TSimulationTickState<TTickSettings>& TickInfo, const FNetworkSimTime& StepTime, const int32 ParentKeyframe, const bool bFinalStep)
+	{
+		TickInfo.SetTotalAllowedSimulationTime( TickInfo.GetTotalAllowedSimulationTime() + StepTime );
+
+		SimulationExtrapolation<T, TDriver>(Driver, Buffers, TickInfo, StepTime);
+
+		TSyncState* SyncState = Buffers.Sync.GetElementFromHead(0);
+		check(SyncState);
+		SyncState->VisualLog( FVisualLoggingParameters(bFinalStep ? EVisualLoggingContext::LastMispredicted : EVisualLoggingContext::OtherMispredicted, ParentKeyframe, EVisualLoggingLifetime::Persistent), Driver, Driver );
+	}
+
 private:
+
+	template<typename T, typename TDriver>
+	void SimulationExtrapolation(TDriver* Driver, TNetworkSimBufferContainer<TBufferTypes>& Buffers, TSimulationTickState<TTickSettings>& TickInfo, const FNetworkSimTime DeltaSimTime)
+	{
+		TInputCmd* LastCmd = Buffers.Input.GetElementFromHead(0);
+
+		// We have extrapolated ahead of the server. The latest network update is now "in the past" from what we rendered last frame.
+		// We will insert a new keyframe to make up the difference from the last known state to where we want to be in the now.
+				
+		TInputCmd* NewCmd = Buffers.Input.GetWriteNext();
+		*NewCmd = LastCmd ? *LastCmd : TInputCmd();
+		NewCmd->SetFrameDeltaTime(DeltaSimTime);
+
+		TSyncState* PrevSyncState = Buffers.Sync.GetElementFromHead(0);
+		TSyncState* NextSyncState = Buffers.Sync.GetWriteNext();
+
+		TAuxState Junk;
+
+		// Do the actual update
+		T::Update(Driver, NewCmd->GetFrameDeltaTime().ToRealTimeSeconds(), *NewCmd, *PrevSyncState, *NextSyncState, Junk);
+		TickInfo.IncrementTotalProcessedSimulationTime(NewCmd->GetFrameDeltaTime(), Buffers.Sync.GetHeadKeyframe());
+
+		// Set our LastProcessedInputKeyframe to fake that we handled it
+		TickInfo.LastProcessedInputKeyframe = Buffers.Input.GetHeadKeyframe();
+		TickInfo.MaxAllowedInputKeyframe = Buffers.Input.GetHeadKeyframe();
+	}
 	
 	FNetworkSimTime ReconcileSimulationTime;
 	FNetworkSimTime LastSerializedSimulationTime;
-	
+	TSyncState LastSerializedSyncState;
 };
 
 
@@ -509,6 +570,9 @@ struct TReplicator_Autonomous : public TBase
 	bool IsReconcileFaultDetected() const { return bReconcileFaultDetected; }
 	const FNetworkSimTime& GetLastSerializedSimTime() const { return SerializedTime; }
 
+	TArray<INetworkSimulationModel*> DependentSimulations;
+	bool bDependentSimulationNeedsReconcile = false;
+
 	int32 GetProxyDirtyCount(TNetworkSimBufferContainer<TBufferTypes>& Buffers) const
 	{
 		return Buffers.Sync.GetDirtyCount() ^ (TBase::GetProxyDirtyCount(Buffers) << 2); ;
@@ -556,8 +620,9 @@ struct TReplicator_Autonomous : public TBase
 			// The state the client predicted that corresponds to the state the server just serialized to us
 			if (TSyncState* ClientExistingState = Buffers.Sync.FindElementByKeyframe(SerializedHeadKeyframe))
 			{
-				if (ClientExistingState->ShouldReconcile(*SerializedState) || (NetworkSimulationModelCVars::ForceReconcile() > 0))
+				if (ClientExistingState->ShouldReconcile(*SerializedState) || (NetworkSimulationModelCVars::ForceReconcile() > 0) || (NetworkSimulationModelCVars::ForceReconcileSingle() > 0))
 				{
+					NetworkSimulationModelCVars::SetForceReconcileSingle(0);
 					UE_CLOG(!Buffers.Input.IsValidKeyframe(SerializedHeadKeyframe), LogNetworkSim, Error, TEXT("::NetSerialize: Client InputBuffer does not contain data for frame %d. {%s} {%s}"), SerializedHeadKeyframe, *Buffers.Input.GetBasicDebugStr(), *Buffers.Sync.GetBasicDebugStr());
 					bPendingReconciliation =  true;
 				}
@@ -587,11 +652,12 @@ struct TReplicator_Autonomous : public TBase
 	template<typename T, typename TDriver>
 	void Reconcile(TDriver* Driver, TNetworkSimBufferContainer<TBufferTypes>& Buffers, TSimulationTickState<TTickSettings>& TickInfo)
 	{
-		if (bPendingReconciliation == false)
+		if (bPendingReconciliation == false && bDependentSimulationNeedsReconcile == false)
 		{
 			return;
 		}
 		bPendingReconciliation = false;
+		bDependentSimulationNeedsReconcile = false;
 		check(ReconciliationBuffer.GetNumValidElements() == 1); // this function assumes only 1 element
 
 		const int32 ReconciliationKeyframe = ReconciliationBuffer.GetHeadKeyframe();
@@ -637,6 +703,8 @@ struct TReplicator_Autonomous : public TBase
 		check(ClientSyncState);
 		*ClientSyncState = *ServerState;
 
+		const FNetworkSimTime RollbackDeltaTime = SerializedTime - TickInfo.GetTotalProcessedSimulationTime();
+
 		// Set the canonical simulation time to what we received (we will advance it as we resimulate)
 		TickInfo.SetTotalProcessedSimulationTime(SerializedTime, ReconciliationKeyframe);
 		TickInfo.LastProcessedInputKeyframe = ReconciliationKeyframe;
@@ -647,6 +715,12 @@ struct TReplicator_Autonomous : public TBase
 			// If we aren't predicting at all, then we advanced the allowed sim time here, (since we aren't doing it in PreSimTick). This just keeps us constantly falling behind and not being able to toggle prediction on/off for debugging.
 			TickInfo.SetTotalAllowedSimulationTime(SerializedTime);
 		}
+
+		// Tell dependent simulations to rollback
+		for (INetworkSimulationModel* DependentSim : DependentSimulations)
+		{
+			DependentSim->BeginRollback(RollbackDeltaTime, ReconciliationKeyframe);
+		}
 		
 		// Resimulate all user commands 
 		const int32 LastKeyframeToProcess = TickInfo.MaxAllowedInputKeyframe;
@@ -688,6 +762,12 @@ struct TReplicator_Autonomous : public TBase
 
 			// Log out the newly predicted state that we got.
 			NextMotionState->VisualLog( FVisualLoggingParameters(Keyframe == LastKeyframeToProcess ? EVisualLoggingContext::LastPredicted : EVisualLoggingContext::OtherPredicted, Keyframe, EVisualLoggingLifetime::Persistent), Driver, Driver);
+
+			// Tell dependent simulations to advance
+			for (INetworkSimulationModel* DependentSim : DependentSimulations)
+			{
+				DependentSim->StepRollback(ResimulateCmd->GetFrameDeltaTime(), Keyframe, (Keyframe == LastKeyframeToProcess));
+			}
 		}
 	}
 
diff --git a/Engine/Plugins/Runtime/NetworkPrediction/readme.txt b/Engine/Plugins/Runtime/NetworkPrediction/readme.txt
--- a/Engine/Plugins/Runtime/NetworkPrediction/readme.txt
+++ b/Engine/Plugins/Runtime/NetworkPrediction/readme.txt
@@ -77,8 +77,6 @@ High level focus:
 -We intend to support RootMotion (Animation) / RootMotionSources (non anim) in the new movement system.
 
 Road Map:
--Interpolation/smoothing layer
--Dependent simulations
 -Aux Buffer
 -Event System
 -Ability System integration
@@ -108,6 +106,13 @@ Smoothing: Taking the output of the simulation and applying an additional layer
 // Release notes
 // ----------------------------------------------------------------------------------------------------------
 
+Update (10-1-19)
+-Forward Predict / Dependent simulation initial check in. This has some limitations and is not final
+	-Dependent Sim can't trigger reconcile yet
+	-Assumptions made about replication frequency. Dependents must replicate as same frequency as parent sim right now.
+	-General lack of precision here when running in variable tick simulation. (Sims tick at different rates on server, so hard to correlate client side).
+
+
 Update (9-26-19)
 -More refactors on general system architecture. UNetworkSimulationGlobalManager is solidifying and most boiler plate is out of the actor components (Not an exciting update but this was overdue!).
 -PostSimTick added and Interpolator refactored to be part of this step, dependent on network role (Eg., simulated proxy only when enabled).
