preprocessor fix
#rnx

diff --git a/Engine/Plugins/Runtime/NetworkPrediction/Source/NetworkPrediction/Public/NetworkSimulationModelCVars.h b/Engine/Plugins/Runtime/NetworkPrediction/Source/NetworkPrediction/Public/NetworkSimulationModelCVars.h
--- a/Engine/Plugins/Runtime/NetworkPrediction/Source/NetworkPrediction/Public/NetworkSimulationModelCVars.h
+++ b/Engine/Plugins/Runtime/NetworkPrediction/Source/NetworkPrediction/Public/NetworkSimulationModelCVars.h
@@ -37,7 +37,7 @@ struct FConditionalAutoConsoleRegister
 
 #if NETSIM_CONST_CVARS
 #define NETSIM_DEVCVAR_SHIPCONST_INT(Var,Value,VarName,Help) \
-	inline int32 Var() { return Value; }
+	inline int32 Var() { return Value; } \
 	inline void Set##Var(int32 V) { }
 #else
 #define NETSIM_DEVCVAR_SHIPCONST_INT(Var,Value,VarName,Help) \
@@ -58,7 +58,7 @@ struct FConditionalAutoConsoleRegister
 
 #if NETSIM_CONST_CVARS
 #define NETSIM_DEVCVAR_SHIPCONST_FLOAT(Var,Value,VarName,Help) \
-	inline float Var() { return Value; }
+	inline float Var() { return Value; } \
 	inline void Set##Vat() { }
 #else
 #define NETSIM_DEVCVAR_SHIPCONST_FLOAT(Var,Value,VarName,Help) \
