missed semi
#rnx

diff --git a/Engine/Plugins/Runtime/NetworkPrediction/Source/NetworkPrediction/Public/NetworkSimulationModelCVars.h b/Engine/Plugins/Runtime/NetworkPrediction/Source/NetworkPrediction/Public/NetworkSimulationModelCVars.h
--- a/Engine/Plugins/Runtime/NetworkPrediction/Source/NetworkPrediction/Public/NetworkSimulationModelCVars.h
+++ b/Engine/Plugins/Runtime/NetworkPrediction/Source/NetworkPrediction/Public/NetworkSimulationModelCVars.h
@@ -45,13 +45,13 @@ struct FConditionalAutoConsoleRegister
 	inline int32 Var() \
 	{ \
 		static const auto* Existing = IConsoleManager::Get().FindConsoleVariable(TEXT(VarName), false); \
-		check(Existing) \
+		check(Existing); \
 		return Existing->GetInt(); \
 	} \
 	inline void Set##Var(int32 V) \
 	{ \
 		static auto* Existing = IConsoleManager::Get().FindConsoleVariable(TEXT(VarName), false); \
-		check(Existing) \
+		check(Existing); \
 		Existing->Set(V, ECVF_SetByConsole); \
 	}
 #endif
@@ -66,13 +66,13 @@ struct FConditionalAutoConsoleRegister
 	inline float Var() \
 	{ \
 		static const auto* Existing = IConsoleManager::Get().FindConsoleVariable(TEXT(VarName), false); \
-		check(Existing) \
+		check(Existing); \
 		return Existing->GetFloat(); \
 	} \
 	inline void Set##Var(float V) \
 	{ \
 		static auto* Existing = IConsoleManager::Get().FindConsoleVariable(TEXT(VarName), false); \
-		check(Existing) \
+		check(Existing); \
 		Existing->Set(V, ECVF_SetByConsole); \
 	}
 #endif
\ No newline at end of file
