One more preprocessor fix
#rnx

diff --git a/Engine/Plugins/Runtime/NetworkPrediction/Source/NetworkPrediction/Public/NetworkSimulationModelCVars.h b/Engine/Plugins/Runtime/NetworkPrediction/Source/NetworkPrediction/Public/NetworkSimulationModelCVars.h
--- a/Engine/Plugins/Runtime/NetworkPrediction/Source/NetworkPrediction/Public/NetworkSimulationModelCVars.h
+++ b/Engine/Plugins/Runtime/NetworkPrediction/Source/NetworkPrediction/Public/NetworkSimulationModelCVars.h
@@ -59,7 +59,7 @@ struct FConditionalAutoConsoleRegister
 #if NETSIM_CONST_CVARS
 #define NETSIM_DEVCVAR_SHIPCONST_FLOAT(Var,Value,VarName,Help) \
 	inline float Var() { return Value; } \
-	inline void Set##Vat() { }
+	inline void Set##Var() { }
 #else
 #define NETSIM_DEVCVAR_SHIPCONST_FLOAT(Var,Value,VarName,Help) \
 	static FConditionalAutoConsoleRegister Var##Auto(TEXT(VarName),(float)Value,TEXT(Help)); \
