Network Prediction Plugin WIP
-Aux buffer hooked up
-TReplicatedBuffer replaced with leaner TNetworkSimSparseBuffer/ TNetworkSimContiguousBuffer. Simplified API around these.

#rnx

diff --git a/Engine/Plugins/Runtime/NetworkPrediction/Source/NetworkPrediction/Private/Movement/FlyingMovementComponent.cpp b/Engine/Plugins/Runtime/NetworkPrediction/Source/NetworkPrediction/Private/Movement/FlyingMovementComponent.cpp
--- a/Engine/Plugins/Runtime/NetworkPrediction/Source/NetworkPrediction/Private/Movement/FlyingMovementComponent.cpp
+++ b/Engine/Plugins/Runtime/NetworkPrediction/Source/NetworkPrediction/Private/Movement/FlyingMovementComponent.cpp
@@ -1,6 +1,5 @@
 // Copyright 1998-2019 Epic Games, Inc. All Rights Reserved.
 
-
 #include "Movement/FlyingMovement.h"
 #include "GameFramework/Actor.h"
 #include "Engine/World.h"
@@ -44,6 +43,12 @@ static FAutoConsoleVariableRef CVarPenetrationOverlapCheckInflation(TEXT("motion
 	TEXT("Distance added to inflation in penetration overlap check."),
 	ECVF_Default);
 
+static float MaxSpeed = 1200.f;
+static FAutoConsoleVariableRef CVarMaxSpeed(TEXT("motion.MaxSpeed"),
+	MaxSpeed,
+	TEXT("Temp value for testing changes to max speed."),
+	ECVF_Default);
+
 static int32 RequestMispredict = 0;
 static FAutoConsoleVariableRef CVarRequestMispredict(TEXT("fp.RequestMispredict"),
 	RequestMispredict, TEXT("Causes a misprediction by inserting random value into stream on authority side"), ECVF_Default);
@@ -66,6 +71,8 @@ INetworkSimulationModel* UFlyingMovementComponent::InstantiateNetworkSimulation(
 {
 	auto NewSim = new FlyingMovement::FMovementSystem<0>(this);
 	DO_NETSIM_MODEL_DEBUG(FNetworkSimulationModelDebuggerManager::Get().RegisterNetworkSimulationModel(NewSim, GetOwner()));
+	MovementSyncState = &NewSim->SyncAccessor;
+	MovementAuxState = &NewSim->AuxAccessor;
 	return NewSim;
 }
 
@@ -94,6 +101,20 @@ void UFlyingMovementComponent::TickComponent(float DeltaTime, enum ELevelTick Ti
 		FlyingMovement::FMovementSimulation::ForceMispredict = true;
 		FlyingMovementCVars::RequestMispredict = 0;
 	}
+
+	// Temp
+	/*
+	if (OwnerRole == ROLE_Authority)
+	{
+		if (MovementAuxState->Get()->MaxSpeed != FlyingMovementCVars::MaxSpeed)
+		{
+			MovementAuxState->Modify([](FlyingMovement::FAuxState& Aux)
+			{
+				Aux.MaxSpeed = FlyingMovementCVars::MaxSpeed;
+			});
+		}
+	}
+	*/
 }
 
 // ----------------------------------------------------------------------------------------------------------
@@ -106,13 +127,18 @@ void UFlyingMovementComponent::TickComponent(float DeltaTime, enum ELevelTick Ti
 //
 // ----------------------------------------------------------------------------------------------------------
 
-
 void UFlyingMovementComponent::InitSyncState(FlyingMovement::FMoveState& OutSyncState) const
 {
 	OutSyncState.Location = UpdatedComponent->GetComponentLocation();
 	OutSyncState.Rotation = UpdatedComponent->GetComponentQuat().Rotator();	
 }
 
+void UFlyingMovementComponent::InitAuxState(FlyingMovement::FAuxState& OutAuxState) const
+{
+	OutAuxState = FlyingMovement::FAuxState();
+	OutAuxState.MaxSpeed = FlyingMovementCVars::MaxSpeed;
+}
+
 void UFlyingMovementComponent::PreSimSync(const FlyingMovement::FMoveState& SyncState)
 {
 	// Does checking equality make any sense here? This is unfortunate
@@ -131,7 +157,7 @@ void UFlyingMovementComponent::ProduceInput(const FNetworkSimTime SimTime, Flyin
 	ProduceInputDelegate.ExecuteIfBound(SimTime, Cmd);
 }
 
-void UFlyingMovementComponent::FinalizeFrame(const FlyingMovement::FMoveState& SyncState)
+void UFlyingMovementComponent::FinalizeFrame(const FlyingMovement::FMoveState& SyncState, const FlyingMovement::FAuxState& AuxState)
 {
 	PreSimSync(SyncState);
 }
diff --git a/Engine/Plugins/Runtime/NetworkPrediction/Source/NetworkPrediction/Private/Movement/FlyingMovementSimulation.cpp b/Engine/Plugins/Runtime/NetworkPrediction/Source/NetworkPrediction/Private/Movement/FlyingMovementSimulation.cpp
--- a/Engine/Plugins/Runtime/NetworkPrediction/Source/NetworkPrediction/Private/Movement/FlyingMovementSimulation.cpp
+++ b/Engine/Plugins/Runtime/NetworkPrediction/Source/NetworkPrediction/Private/Movement/FlyingMovementSimulation.cpp
@@ -155,21 +155,12 @@ namespace FlyingMovement
 	}
 
 
-	void FMovementSimulation::Update(IMovementDriver* Driver, const float DeltaSeconds, const FInputCmd& InputCmd, const FMoveState& InputState, FMoveState& OutputState, const FAuxState& AuxState)
+	void FMovementSimulation::Update(IMovementDriver* Driver, const float DeltaSeconds, const FInputCmd& InputCmd, const FMoveState& InputState, FMoveState& OutputState, const FAuxState& AuxState, const TLazyStateAccessor<FAuxState>& OutAuxStateAccessor)
 	{
-		// ----------------------------------
-		// AuxState todo
-		float MaxSpeed = 1200.f;
-		float TurningBoost = 8.f;
-		float Deceleration = 8000.f;
-		float Acceleration = 4000.f;
-
-		// ----------------------------------------------------------------
-		
-		OutputState = InputState;
-
 		IBaseMovementDriver& BaseMovementDriver = Driver->GetBaseMovementDriver();
 
+		OutputState = InputState;
+
 		// --------------------------------------------------------------
 		//	Rotation Update
 		//	We do the rotational update inside the movement sim so that things like server side teleport will work.
@@ -196,7 +187,7 @@ namespace FlyingMovement
 			FVector Velocity = InputState.Velocity;
 
 			const float AnalogInputModifier = (ControlAcceleration.SizeSquared() > 0.f ? ControlAcceleration.Size() : 0.f);
-			const float MaxPawnSpeed = MaxSpeed * AnalogInputModifier;
+			const float MaxPawnSpeed = AuxState.MaxSpeed * AnalogInputModifier;
 			const bool bExceedingMaxSpeed = IsExceedingMaxSpeed(Velocity, MaxPawnSpeed);
 
 			if (AnalogInputModifier > 0.f && !bExceedingMaxSpeed)
@@ -205,7 +196,7 @@ namespace FlyingMovement
 				if (Velocity.SizeSquared() > 0.f)
 				{
 					// Change direction faster than only using acceleration, but never increase velocity magnitude.
-					const float TimeScale = FMath::Clamp(DeltaSeconds * TurningBoost, 0.f, 1.f);
+					const float TimeScale = FMath::Clamp(DeltaSeconds * AuxState.TurningBoost, 0.f, 1.f);
 					Velocity = Velocity + (ControlAcceleration * Velocity.Size() - Velocity) * TimeScale;
 				}
 			}
@@ -215,7 +206,7 @@ namespace FlyingMovement
 				if (Velocity.SizeSquared() > 0.f)
 				{
 					const FVector OldVelocity = Velocity;
-					const float VelSize = FMath::Max(Velocity.Size() - FMath::Abs(Deceleration) * DeltaSeconds, 0.f);
+					const float VelSize = FMath::Max(Velocity.Size() - FMath::Abs(AuxState.Deceleration) * DeltaSeconds, 0.f);
 					Velocity = Velocity.GetSafeNormal() * VelSize;
 
 					// Don't allow braking to lower us below max speed if we started above it.
@@ -228,7 +219,7 @@ namespace FlyingMovement
 
 			// Apply acceleration and clamp velocity magnitude.
 			const float NewMaxSpeed = (IsExceedingMaxSpeed(Velocity, MaxPawnSpeed)) ? Velocity.Size() : MaxPawnSpeed;
-			Velocity += ControlAcceleration * FMath::Abs(Acceleration) * DeltaSeconds;
+			Velocity += ControlAcceleration * FMath::Abs(AuxState.Acceleration) * DeltaSeconds;
 			Velocity = Velocity.GetClampedToMaxSize(NewMaxSpeed);
 
 			// Finally, output velocity that we calculated
diff --git a/Engine/Plugins/Runtime/NetworkPrediction/Source/NetworkPrediction/Private/Movement/ParametricMovement.cpp b/Engine/Plugins/Runtime/NetworkPrediction/Source/NetworkPrediction/Private/Movement/ParametricMovement.cpp
--- a/Engine/Plugins/Runtime/NetworkPrediction/Source/NetworkPrediction/Private/Movement/ParametricMovement.cpp
+++ b/Engine/Plugins/Runtime/NetworkPrediction/Source/NetworkPrediction/Private/Movement/ParametricMovement.cpp
@@ -42,7 +42,7 @@ namespace ParametricMovement
 {
 	const FName FMovementSimulation::GroupName(TEXT("Parametric"));
 
-	void FMovementSimulation::Update(IMovementDriver* Driver, const float DeltaSeconds, const FInputCmd& InputCmd, const FMoveState& InputState, FMoveState& OutputState, const FAuxState& AuxState)
+	void FMovementSimulation::Update(IMovementDriver* Driver, const float DeltaSeconds, const FInputCmd& InputCmd, const FMoveState& InputState, FMoveState& OutputState, const FAuxState& AuxState, const TLazyStateAccessor<FAuxState>& OutAuxStateAccessor)
 	{
 		IBaseMovementDriver& BaseMovementDriver = Driver->GetBaseMovementDriver();
 		//const float DeltaSeconds = SimTimeDeltaMS.ToRealTimeSeconds();
@@ -114,12 +114,16 @@ INetworkSimulationModel* UParametricMovementComponent::InstantiateNetworkSimulat
 		auto *NewSim = new ParametricMovement::FMovementSystem<16>(this);
 		NewSim->RepProxy_Simulated.bAllowSimulatedExtrapolation = !bEnableInterpolation;
 		DO_NETSIM_MODEL_DEBUG(FNetworkSimulationModelDebuggerManager::Get().RegisterNetworkSimulationModel(NewSim, GetOwner()));
+		MovementSyncState = &NewSim->SyncAccessor;
+		MovementAuxState = &NewSim->AuxAccessor;
 		return NewSim;
 	}
 	
 	auto *NewSim = new ParametricMovement::FMovementSystem<>(this);
 	NewSim->RepProxy_Simulated.bAllowSimulatedExtrapolation = !bEnableInterpolation;
 	DO_NETSIM_MODEL_DEBUG(FNetworkSimulationModelDebuggerManager::Get().RegisterNetworkSimulationModel(NewSim, GetOwner()));
+	MovementSyncState = &NewSim->SyncAccessor;
+	MovementAuxState = &NewSim->AuxAccessor;
 	return NewSim;
 }
 
@@ -148,7 +152,11 @@ void UParametricMovementComponent::InitSyncState(ParametricMovement::FMoveState&
 	OutSyncState.PlayRate = 0.f;
 }
 
-void UParametricMovementComponent::FinalizeFrame(const ParametricMovement::FMoveState& SyncState)
+void UParametricMovementComponent::InitAuxState(ParametricMovement::FAuxState& OutAuxState) const
+{
+}
+
+void UParametricMovementComponent::FinalizeFrame(const ParametricMovement::FMoveState& SyncState, const ParametricMovement::FAuxState& AuxState)
 {
 	FTransform NewTransform;
 	MapTimeToTransform(SyncState.Position, NewTransform);
diff --git a/Engine/Plugins/Runtime/NetworkPrediction/Source/NetworkPrediction/Private/NetworkPredictionComponent.cpp b/Engine/Plugins/Runtime/NetworkPrediction/Source/NetworkPrediction/Private/NetworkPredictionComponent.cpp
--- a/Engine/Plugins/Runtime/NetworkPrediction/Source/NetworkPrediction/Private/NetworkPredictionComponent.cpp
+++ b/Engine/Plugins/Runtime/NetworkPrediction/Source/NetworkPrediction/Private/NetworkPredictionComponent.cpp
@@ -15,32 +15,34 @@ void UNetworkPredictionComponent::InitializeComponent()
 {
 	Super::InitializeComponent();
 
-	// Child class will instantiate
-	INetworkSimulationModel* NewNetworkSim = InstantiateNetworkSimulation();
-	NetworkSim.Reset(NewNetworkSim);
-
-	if (NewNetworkSim == nullptr)
+	UWorld* World = GetWorld();
+	UNetworkSimulationGlobalManager* NetworkSimGlobalManager = GetWorld()->GetSubsystem<UNetworkSimulationGlobalManager>();
+	if (NetworkSimGlobalManager)
 	{
-		// Its ok to not instantiate a sim
-		return;
-	}
 
-	// Init RepProxies
-	ReplicationProxy_ServerRPC.Init(NewNetworkSim, EReplicationProxyTarget::ServerRPC);
-	ReplicationProxy_Autonomous.Init(NewNetworkSim, EReplicationProxyTarget::AutonomousProxy);
-	ReplicationProxy_Simulated.Init(NewNetworkSim, EReplicationProxyTarget::SimulatedProxy);
-	ReplicationProxy_Replay.Init(NewNetworkSim, EReplicationProxyTarget::Replay);
-#if NETSIM_MODEL_DEBUG
-	ReplicationProxy_Debug.Init(NewNetworkSim, EReplicationProxyTarget::Debug);
-#endif
+		// Child class will instantiate
+		INetworkSimulationModel* NewNetworkSim = InstantiateNetworkSimulation();
+		NetworkSim.Reset(NewNetworkSim);
 
-	// Register with GlobalManager
-	UNetworkSimulationGlobalManager* NetworkSimGlobalManager = GetWorld()->GetSubsystem<UNetworkSimulationGlobalManager>();
-	check(NetworkSimGlobalManager);
+		if (NewNetworkSim == nullptr)
+		{
+			// Its ok to not instantiate a sim
+			return;
+		}
 
-	NetworkSimGlobalManager->RegisterModel(NewNetworkSim, GetOwner());
-	
-	CheckOwnerRoleChange();
+		// Init RepProxies
+		ReplicationProxy_ServerRPC.Init(NewNetworkSim, EReplicationProxyTarget::ServerRPC);
+		ReplicationProxy_Autonomous.Init(NewNetworkSim, EReplicationProxyTarget::AutonomousProxy);
+		ReplicationProxy_Simulated.Init(NewNetworkSim, EReplicationProxyTarget::SimulatedProxy);
+		ReplicationProxy_Replay.Init(NewNetworkSim, EReplicationProxyTarget::Replay);
+	#if NETSIM_MODEL_DEBUG
+		ReplicationProxy_Debug.Init(NewNetworkSim, EReplicationProxyTarget::Debug);
+	#endif
+
+		// Register with GlobalManager			
+		NetworkSimGlobalManager->RegisterModel(NewNetworkSim, GetOwner());
+		CheckOwnerRoleChange();
+	}
 }
 
 void UNetworkPredictionComponent::UninitializeComponent()
@@ -50,16 +52,16 @@ void UNetworkPredictionComponent::UninitializeComponent()
 	if (NetworkSim.IsValid())
 	{
 		UNetworkSimulationGlobalManager* NetworkSimGlobalManager = GetWorld()->GetSubsystem<UNetworkSimulationGlobalManager>();
-		check(NetworkSimGlobalManager);
-
-		NetworkSimGlobalManager->UnregisterModel(NetworkSim.Get(), GetOwner());
-
-		if (ServerRPCHandle.IsValid())
+		if (ensure(NetworkSimGlobalManager))
 		{
-			NetworkSimGlobalManager->TickServerRPCDelegate.Remove(ServerRPCHandle);
-		}
+			NetworkSimGlobalManager->UnregisterModel(NetworkSim.Get(), GetOwner());
+			if (ServerRPCHandle.IsValid())
+			{
+				NetworkSimGlobalManager->TickServerRPCDelegate.Remove(ServerRPCHandle);
+			}
 
-		NetworkSim.Reset(nullptr);
+			NetworkSim.Reset(nullptr);
+		}
 	}
 }
 
diff --git a/Engine/Plugins/Runtime/NetworkPrediction/Source/NetworkPrediction/Private/NetworkPredictionTests.cpp b/Engine/Plugins/Runtime/NetworkPrediction/Source/NetworkPrediction/Private/NetworkPredictionTests.cpp
--- a/Engine/Plugins/Runtime/NetworkPrediction/Source/NetworkPrediction/Private/NetworkPredictionTests.cpp
+++ b/Engine/Plugins/Runtime/NetworkPrediction/Source/NetworkPrediction/Private/NetworkPredictionTests.cpp
@@ -6,7 +6,7 @@
 //	Testing Fixed vs Variable tick rate struct declration
 // ------------------------------------------------------------------------------------------------------------
 
-
+DEFINE_LOG_CATEGORY_STATIC(NetworkSimulationTests, Log, All);
 
 struct FNetSimDummy { };
 using FDummyBufferTypes = TNetworkSimBufferTypes<FNetSimDummy, FNetSimDummy, FNetSimDummy>;
@@ -38,3 +38,47 @@ void NetSimulationTypesTest()
 	//TNetworkSimTickSettings<20, 60> StaticAssertSettings;
 	//TickSettingsTest(StaticAssertSettings);
 }
+
+template<typename TBuffer>
+void NetSimulationTestBuffer(TBuffer& Buffer)
+{
+	for (int32 i=0; i < 10; ++i)
+	{
+		typename TBuffer::ElementType NewData;
+		*Buffer.WriteKeyframe(i) = NewData;
+		UE_LOG(NetworkSimulationTests, Display, TEXT("%s"), *Buffer.GetBasicDebugStr());
+
+		for (auto It = Buffer.CreateIterator(); It; ++It)
+		{
+			UE_LOG(NetworkSimulationTests, Display, TEXT("   [%d] %s"), It.Keyframe(), *It.Element()->ToString());
+		}
+
+		for (auto It = Buffer.CreateConstIterator(); It; ++It)
+		{
+			UE_LOG(NetworkSimulationTests, Display, TEXT("   [%d] %s"), It.Keyframe(), *It.Element()->ToString());
+		}
+	}
+}
+
+void NetSimulationBufferTest()
+{
+	struct FState
+	{
+		int32 x = 0;
+		int32 y = 0;
+		int32 z = 0;
+		float w = 1.f;
+
+		FString ToString() const { return FString::Printf(TEXT("%d %d %d %f"), x, y, z, w); }
+	};
+
+	TNetworkSimSparseBuffer<FState> SparseBuffer;
+	TNetworkSimContiguousBuffer<FState> ContiguousBuffer;
+
+	NetSimulationTestBuffer(SparseBuffer);
+	NetSimulationTestBuffer(ContiguousBuffer);
+
+
+	//TNetworkSimBufferIterator<TNetworkSimSparseBuffer<FState>> It(SparseBuffer);
+	//It.Keyframe();
+}
diff --git a/Engine/Plugins/Runtime/NetworkPrediction/Source/NetworkPrediction/Private/NetworkSimulationModelDebugger.cpp b/Engine/Plugins/Runtime/NetworkPrediction/Source/NetworkPrediction/Private/NetworkSimulationModelDebugger.cpp
--- a/Engine/Plugins/Runtime/NetworkPrediction/Source/NetworkPrediction/Private/NetworkSimulationModelDebugger.cpp
+++ b/Engine/Plugins/Runtime/NetworkPrediction/Source/NetworkPrediction/Private/NetworkSimulationModelDebugger.cpp
@@ -122,7 +122,7 @@ FAutoConsoleCommandWithWorldAndArgs NetworkSimulationModelDebugClassCmd(TEXT("nm
 
 		for (TObjectIterator<UClass> ClassIt; ClassIt; ++ClassIt)
 		{
-			if (ClassIt->GetName().Contains(Args[0]))
+			if (ClassIt->IsChildOf(AActor::StaticClass()) && ClassIt->GetName().Contains(Args[0]))
 			{
 				Class = *ClassIt;
 				UE_LOG(LogNetworkSimDebug, Display, TEXT("Found: %s"), *Class->GetName());
@@ -139,6 +139,7 @@ FAutoConsoleCommandWithWorldAndArgs NetworkSimulationModelDebugClassCmd(TEXT("nm
 	}
 
 	// Probably need to do distance based sorting or a way to specify which one you want.
+	bool bFound = false;
 	for (TActorIterator<AActor> It(World); It; ++It)
 	{
 		if (It->GetClass()->IsChildOf(Class))
@@ -147,7 +148,12 @@ FAutoConsoleCommandWithWorldAndArgs NetworkSimulationModelDebugClassCmd(TEXT("nm
 			{
 				UE_LOG(LogNetworkSimDebug, Display, TEXT("Toggling NetworkSim debugger for %s"), *It->GetName());
 				FNetworkSimulationModelDebuggerManager::Get().ToggleDebuggerActive(*It);
+				bFound = true;
 			}
 		}
 	}
+	if (!bFound)
+	{
+		UE_LOG(LogNetworkSimDebug, Display, TEXT("Unable to find actors matching search criteria. Class: %s"), *GetNameSafe(Class));
+	}
 }));
\ No newline at end of file
diff --git a/Engine/Plugins/Runtime/NetworkPrediction/Source/NetworkPrediction/Public/Movement/FlyingMovement.h b/Engine/Plugins/Runtime/NetworkPrediction/Source/NetworkPrediction/Public/Movement/FlyingMovement.h
--- a/Engine/Plugins/Runtime/NetworkPrediction/Source/NetworkPrediction/Public/Movement/FlyingMovement.h
+++ b/Engine/Plugins/Runtime/NetworkPrediction/Source/NetworkPrediction/Public/Movement/FlyingMovement.h
@@ -104,11 +104,18 @@ namespace FlyingMovement
 	// Auxiliary state that is input into the simulation. Doesn't change during the simulation tick.
 	// (It can change and even be predicted but doing so will trigger more bookeeping, etc. Changes will happen "next tick")
 	struct FAuxState
-	{
-		float Multiplier=1;
+	{	
+		float MaxSpeed = 1200.f;
+		float TurningBoost = 8.f;
+		float Deceleration = 8000.f;
+		float Acceleration = 4000.f;
+
 		void NetSerialize(const FNetSerializeParams& P)
 		{
-			P.Ar << Multiplier;
+			P.Ar << MaxSpeed;
+			P.Ar << TurningBoost;
+			P.Ar << Deceleration;
+			P.Ar << Acceleration;
 		}
 	};
 
@@ -134,7 +141,7 @@ namespace FlyingMovement
 	{
 	public:
 		/** Main update function */
-		static void Update(IMovementDriver* Driver, const float DeltaSeconds, const FInputCmd& InputCmd, const FMoveState& InputState, FMoveState& OutputState, const FAuxState& AuxState);
+		static void Update(IMovementDriver* Driver, const float DeltaSeconds, const FInputCmd& InputCmd, const FMoveState& InputState, FMoveState& OutputState, const FAuxState& AuxState, const TLazyStateAccessor<FAuxState>& OutAuxStateAccessor);
 
 		/** Tick group the simulation maps to */
 		static const FName GroupName;
@@ -178,14 +185,20 @@ public:
 
 	FString GetDebugName() const override;
 	const UObject* GetVLogOwner() const override;
+
 	void InitSyncState(FlyingMovement::FMoveState& OutSyncState) const override;
-	void PreSimSync(const FlyingMovement::FMoveState& SyncState) override;
+	void InitAuxState(FlyingMovement::FAuxState& OutAuxState) const override;
 	void ProduceInput(const FNetworkSimTime SimTime, FlyingMovement::FInputCmd& Cmd) override;
-	void FinalizeFrame(const FlyingMovement::FMoveState& SyncState) override;
+	void FinalizeFrame(const FlyingMovement::FMoveState& SyncState, const FlyingMovement::FAuxState& AuxState) override;
+
+	void PreSimSync(const FlyingMovement::FMoveState& SyncState) override;
 
 	DECLARE_DELEGATE_TwoParams(FProduceFlyingInput, const FNetworkSimTime /*SimTime*/, FlyingMovement::FInputCmd& /*Cmd*/)
 	FProduceFlyingInput ProduceInputDelegate;
 
+	TPredictedStateAccessor<FlyingMovement::FMoveState>* MovementSyncState = nullptr;
+	TPredictedStateAccessor<FlyingMovement::FAuxState>*	 MovementAuxState = nullptr;
+
 protected:
 
 	// Network Prediction
diff --git a/Engine/Plugins/Runtime/NetworkPrediction/Source/NetworkPrediction/Public/Movement/ParametricMovement.h b/Engine/Plugins/Runtime/NetworkPrediction/Source/NetworkPrediction/Public/Movement/ParametricMovement.h
--- a/Engine/Plugins/Runtime/NetworkPrediction/Source/NetworkPrediction/Public/Movement/ParametricMovement.h
+++ b/Engine/Plugins/Runtime/NetworkPrediction/Source/NetworkPrediction/Public/Movement/ParametricMovement.h
@@ -114,7 +114,7 @@ namespace ParametricMovement
 	class FMovementSimulation
 	{
 	public:
-		static void Update(IMovementDriver* Driver, const float DeltaSeconds, const FInputCmd& InputCmd, const FMoveState& InputState, FMoveState& OutputState, const FAuxState& AuxState);
+		static void Update(IMovementDriver* Driver, const float DeltaSeconds, const FInputCmd& InputCmd, const FMoveState& InputState, FMoveState& OutputState, const FAuxState& AuxState, const TLazyStateAccessor<FAuxState>& OutAuxStateAccessor);
 		static const FName GroupName;
 	};
 
@@ -154,8 +154,9 @@ class NETWORKPREDICTION_API UParametricMovementComponent : public UBaseMovementC
 	FString GetDebugName() const override;
 	const UObject* GetVLogOwner() const override;
 	void InitSyncState(ParametricMovement::FMoveState& OutSyncState) const override;
-	void FinalizeFrame(const ParametricMovement::FMoveState& SyncState) override;
+	void InitAuxState(ParametricMovement::FAuxState& OutAuxState) const override;
 	void ProduceInput(const FNetworkSimTime SimTime, ParametricMovement::FInputCmd& Cmd);
+	void FinalizeFrame(const ParametricMovement::FMoveState& SyncState, const ParametricMovement::FAuxState& AuxState) override;	
 
 	// Base Movement Driver
 	IBaseMovementDriver& GetBaseMovementDriver() override final { return *static_cast<IBaseMovementDriver*>(this); }
@@ -164,6 +165,10 @@ class NETWORKPREDICTION_API UParametricMovementComponent : public UBaseMovementC
 	virtual void AdvanceParametricTime(const float InPosition, const float InPlayRate, float &OutPosition, float& OutPlayRate, const float DeltaTimeSeconds) const override;
 	virtual void MapTimeToTransform(const float InPosition, FTransform& OutTransform) const override;
 
+
+	TPredictedStateAccessor<ParametricMovement::FMoveState>* MovementSyncState = nullptr;
+	TPredictedStateAccessor<ParametricMovement::FAuxState>*	 MovementAuxState = nullptr;
+
 protected:
 
 	virtual INetworkSimulationModel* InstantiateNetworkSimulation() override;
diff --git a/Engine/Plugins/Runtime/NetworkPrediction/Source/NetworkPrediction/Public/NetworkSimulationModel.h b/Engine/Plugins/Runtime/NetworkPrediction/Source/NetworkPrediction/Public/NetworkSimulationModel.h
--- a/Engine/Plugins/Runtime/NetworkPrediction/Source/NetworkPrediction/Public/NetworkSimulationModel.h
+++ b/Engine/Plugins/Runtime/NetworkPrediction/Source/NetworkPrediction/Public/NetworkSimulationModel.h
@@ -51,18 +51,10 @@ public:
 
 	using TSimTime = FNetworkSimTime;
 	using TRealTime = FNetworkSimTime::FRealTime;
-
-	class IDriver
-	{
-	public:
-		virtual FString GetDebugName() const = 0; // Used for debugging. Recommended to emit the simulation name and the actor name/role.
-
-		virtual void InitSyncState(TSyncState& OutSyncState) const = 0;	// Called to create initial value of the sync state.
-		virtual void ProduceInput(const FNetworkSimTime, typename TUserBufferTypes::TInputCmd&) = 0; // Called when the sim is ready to process new local input
-		virtual void FinalizeFrame(const TSyncState& SyncState) = 0; // Called from the Network Sim at the end of the sim frame when there is new sync data.
-	};
 	
 	TNetworkedSimulationModel(TDriver* InDriver)
+		: SyncAccessor(Buffers.SyncMods, &TickInfo.LastProcessedInputKeyframe)
+		, AuxAccessor(Buffers.AuxMods, &TickInfo.LastProcessedInputKeyframe)
 	{
 		Driver = InDriver;
 	}
@@ -78,7 +70,7 @@ public:
 		// (property replication and ServerRPC get sent after the tick, rather than forcing awkward callback into the NetSim post replication, we can check it here)
 		if (auto* DebugBuffer = GetLocalDebugBuffer())
 		{
-			if (TDebugState* const PrevDebugState = DebugBuffer->FindElementByKeyframe(DebugBuffer->GetHeadKeyframe()))
+			if (TDebugState* const PrevDebugState = DebugBuffer->Get(DebugBuffer->HeadKeyframe()))
 			{
 				if (Parameters.Role == ROLE_AutonomousProxy)
 				{
@@ -133,7 +125,7 @@ public:
 		// -------------------------------------------------------------------------------------------------------------------------------------------------
 		//												Input Processing & Simulation Update
 		// -------------------------------------------------------------------------------------------------------------------------------------------------
-		if (Buffers.Input.GetHeadKeyframe() > Buffers.Sync.GetHeadKeyframe())
+		if (Buffers.Input.HeadKeyframe() > Buffers.Sync.HeadKeyframe())
 		{
 			// -------------------------------------------------------------------------------------------------------------
 			//	
@@ -142,27 +134,32 @@ public:
 			//		-Its HeadKeyframe should be one behind the Keyframe we are about to process.
 			//
 			//	Note, InputCmds start @ Keyframe=1. The first SyncedState that Update produces will go in KeyFrame=1.
-			//	(E.g, InputCmd @ keyframe=X is used to generate MotionState @ keyframe=X)
+			//	(E.g, InputCmd @ keyframe=X is used to generate SyncState @ keyframe=X)
 			//	This means that SyncedState @ keyframe=0 is always created here via InitSyncState.
 			//	This also means that we never actually process InputCmd @ keyframe=0. Which is why LastProcessedInputKeyframe is initialized to 0 ("already processed")
 			//	and the buffer has an empty element inserted in InitializeForNetworkRole.
 			// -------------------------------------------------------------------------------------------------------------
 
-			if (Buffers.Sync.GetHeadKeyframe() != TickInfo.LastProcessedInputKeyframe)
+			if (Buffers.Sync.HeadKeyframe() != TickInfo.LastProcessedInputKeyframe)
 			{
 				if (TickInfo.LastProcessedInputKeyframe != 0)
 				{
 					// This shouldn't happen, but is not fatal. We are reseting the sync state buffer.
-					UE_LOG(LogNetworkSim, Warning, TEXT("%s. Break in SyncState continuity. LastProcessedInputKeyframe: %d. SyncBuffer.GetHeadKeyframe(): %d."), *Driver->GetDebugName() , TickInfo.LastProcessedInputKeyframe, Buffers.Sync.GetHeadKeyframe());
+					UE_LOG(LogNetworkSim, Warning, TEXT("%s. Break in SyncState continuity. LastProcessedInputKeyframe: %d. SyncBuffer.HeadKeyframe(): %d."), *Driver->GetDebugName() , TickInfo.LastProcessedInputKeyframe, Buffers.Sync.HeadKeyframe());
 				}
 
 				// We need an initial/current state. Get this from the sim driver
-				Buffers.Sync.ResetNextHeadKeyframe(TickInfo.LastProcessedInputKeyframe);
-				TSyncState* StartingState = Buffers.Sync.GetWriteNext();
+				TSyncState* StartingState = Buffers.Sync.WriteKeyframe(TickInfo.LastProcessedInputKeyframe);
 				Driver->InitSyncState(*StartingState);
 
 				// Reset time tracking buffer too
-				TickInfo.SetTotalProcessedSimulationTime(TickInfo.GetTotalProcessedSimulationTime(), Buffers.Sync.GetHeadKeyframe());
+				TickInfo.SetTotalProcessedSimulationTime(TickInfo.GetTotalProcessedSimulationTime(), Buffers.Sync.HeadKeyframe());
+			}
+
+			if (Buffers.Aux.Num() == 0)
+			{
+				TAuxState* StartingAuxState = Buffers.Aux.WriteKeyframe(TickInfo.LastProcessedInputKeyframe);
+				Driver->InitAuxState(*StartingAuxState);
 			}
 		
 			// -------------------------------------------------------------------------------------------------------------
@@ -176,7 +173,7 @@ public:
 					break;
 				}
 
-				if (TInputCmd* NextCmd = Buffers.Input.FindElementByKeyframe(Keyframe))
+				if (TInputCmd* NextCmd = Buffers.Input[Keyframe])
 				{
 					// We have an unprocessed command, do we have enough allotted simulation time to process it?
 					if (TickInfo.GetRemainingAllowedSimulationTime() >= NextCmd->GetFrameDeltaTime())
@@ -184,20 +181,20 @@ public:
 						// -------------------------------------------------------------------------------------------------
 						//	The core process input command and call ::Update block!
 						// -------------------------------------------------------------------------------------------------
-						TSyncState* PrevSyncState = Buffers.Sync.FindElementByKeyframe(TickInfo.LastProcessedInputKeyframe);
-						TSyncState* NextSyncState = Buffers.Sync.GetWriteNext();
+						TSyncState* PrevSyncState = Buffers.Sync[TickInfo.LastProcessedInputKeyframe];
+						TSyncState* NextSyncState = Buffers.Sync.WriteKeyframe(Keyframe);
+						TAuxState* AuxState = Buffers.Aux[TickInfo.LastProcessedInputKeyframe];
 
 						check(PrevSyncState != nullptr);
 						check(NextSyncState != nullptr);
-						check(Buffers.Sync.GetHeadKeyframe() == Keyframe);
+						check(Buffers.Sync.HeadKeyframe() == Keyframe);
 				
 						if (DebugState)
 						{
 							DebugState->ProcessedKeyframes.Add(Keyframe);
 						}
-					
-						TAuxState AuxState; // Temp: aux buffer not implemented yet
-						TSimulation::Update(Driver, NextCmd->GetFrameDeltaTime().ToRealTimeSeconds(), *NextCmd, *PrevSyncState, *NextSyncState, AuxState);
+						
+						TSimulation::Update(Driver, NextCmd->GetFrameDeltaTime().ToRealTimeSeconds(), *NextCmd, *PrevSyncState, *NextSyncState, *AuxState, Buffers.Aux.WriteKeyframeFunc(Keyframe+1));
 					
 						TickInfo.IncrementTotalProcessedSimulationTime(NextCmd->GetFrameDeltaTime(), Keyframe);
 						TickInfo.LastProcessedInputKeyframe = Keyframe;
@@ -241,7 +238,7 @@ public:
 		if (DebugState)
 		{
 			DebugState->LastProcessedKeyframe = TickInfo.LastProcessedInputKeyframe;
-			DebugState->HeadKeyframe = Buffers.Input.GetHeadKeyframe();
+			DebugState->HeadKeyframe = Buffers.Input.HeadKeyframe();
 			DebugState->RemainingAllowedSimulationTimeSeconds = (float)TickInfo.GetRemainingAllowedSimulationTime().ToRealTimeSeconds();
 		}
 
@@ -279,26 +276,31 @@ public:
 	
 	void InitializeForNetworkRole(const ENetRole Role, const FNetworkSimulationModelInitParameters& Parameters) final override
 	{
-		Buffers.Input.SetBufferSize(Parameters.InputBufferSize);
-		Buffers.Sync.SetBufferSize(Parameters.SyncedBufferSize);
-		Buffers.Aux.SetBufferSize(Parameters.AuxBufferSize);
+		// FIXME: buffer sizes are now inlined allocated but we want to support role based buffer sizes
+
+		//Buffers.Input.SetBufferSize(Parameters.InputBufferSize);
+		//Buffers.Sync.SetBufferSize(Parameters.SyncedBufferSize);
+		//Buffers.Aux.SetBufferSize(Parameters.AuxBufferSize); AUXFIXME
 
 		if (GetLocalDebugBuffer())
 		{
-			GetLocalDebugBuffer()->SetBufferSize(Parameters.DebugBufferSize);
+			//GetLocalDebugBuffer()->SetBufferSize(Parameters.DebugBufferSize);
 		}
 
 		if (auto* MyHistoricBuffers = GetHistoricBuffers(true))
 		{
-			MyHistoricBuffers->Input.SetBufferSize(Parameters.HistoricBufferSize);
-			MyHistoricBuffers->Sync.SetBufferSize(Parameters.HistoricBufferSize);
-			MyHistoricBuffers->Aux.SetBufferSize(Parameters.HistoricBufferSize);
+			//MyHistoricBuffers->Input.SetBufferSize(Parameters.HistoricBufferSize);
+			//MyHistoricBuffers->Sync.SetBufferSize(Parameters.HistoricBufferSize);
+			//MyHistoricBuffers->Aux.SetBufferSize(Parameters.HistoricBufferSize); AUXFIXME
 		}
 
-		TickInfo.InitSimulationTimeBuffer(Parameters.SyncedBufferSize);
+		//TickInfo.InitSimulationTimeBuffer(Parameters.SyncedBufferSize);
+
+		Buffers.SyncMods.SetIsAuthority(Role == ROLE_Authority);
+		Buffers.AuxMods.SetIsAuthority(Role == ROLE_Authority);
 
 		// We want to start with an empty command in the input buffer. The sync buffer will be populated @ frame 0 with the "current" state when we actually sim. This keeps them in sync
-		*Buffers.Input.GetWriteNext() = TInputCmd();
+		*Buffers.Input.WriteKeyframe(0) = TInputCmd();
 	}
 
 	void NetSerializeProxy(EReplicationProxyTarget Target, const FNetSerializeParams& Params) final override
@@ -418,11 +420,19 @@ public:
 
 	TNetworkSimBufferContainer<TBufferTypes> Buffers;
 
+	TPredictedStateAccessor<TSyncState> SyncAccessor;
+	TPredictedStateAccessor<TAuxState> AuxAccessor;
+
 	TRepProxyServerRPC RepProxy_ServerRPC;
 	TRepProxyAutonomous RepProxy_Autonomous;
 	TRepProxySimulated RepProxy_Simulated;
 	TRepProxyReplay RepProxy_Replay;
 
+	using TInputBuffer = typename TNetworkSimBufferContainer<TBufferTypes>::TInputBuffer;
+	using TSyncBuffer = typename TNetworkSimBufferContainer<TBufferTypes>::TSyncBuffer;
+	using TAuxBuffer = typename TNetworkSimBufferContainer<TBufferTypes>::TAuxBuffer;
+	using TDebugBuffer = typename TNetworkSimBufferContainer<TBufferTypes>::TDebugBuffer;
+
 	// ------------------------------------------------------------------
 	// RPC Sending helper: provides basic send frequency settings for tracking when the Server RPC can be invoked.
 	// Note that the Driver is the one that must call the RPC, that cannot be rolled into this templated structure.
@@ -458,20 +468,20 @@ private:
 public:
 
 #if NETSIM_MODEL_DEBUG
-	TReplicationBuffer<TDebugState>* GetLocalDebugBuffer() {	return &Buffers.Debug; }
-	TDebugState* GetNextLocalDebugStateWrite() { return Buffers.Debug.GetWriteNext(); }
+	TDebugBuffer* GetLocalDebugBuffer() {	return &Buffers.Debug; }
+	TDebugState* GetNextLocalDebugStateWrite() { return Buffers.Debug.WriteKeyframe( Buffers.Debug.HeadKeyframe() + 1 ); }
 	TNetworkSimBufferContainer<TBufferTypes>* GetHistoricBuffers(bool bCreate=false)
 	{
 		if (HistoricBuffers.IsValid() == false && bCreate) { HistoricBuffers.Reset(new TNetworkSimBufferContainer<TBufferTypes>()); }
 		return HistoricBuffers.Get();
 	}
 
-	TReplicationBuffer<TDebugState>* GetRemoteDebugBuffer() {	return &RepProxy_Debug.ReceivedBuffer; }
+	TDebugBuffer* GetRemoteDebugBuffer() {	return &RepProxy_Debug.ReceivedBuffer; }
 #else
-	TReplicationBuffer<TDebugState>* GetLocalDebugBuffer() {	return nullptr; }
+	TDebugBuffer* GetLocalDebugBuffer() {	return nullptr; }
 	TDebugState* GetNextLocalDebugStateWrite() { return nullptr; }
 	TNetworkSimBufferContainer<TBufferTypes>* GetHistoricBuffers(bool bCreate=false) { return nullptr; }
-	TReplicationBuffer<TDebugState>* GetRemoteDebugBuffer() {	return nullptr; }
+	TDebugBuffer* GetRemoteDebugBuffer() {	return nullptr; }
 #endif
 
 private:
diff --git a/Engine/Plugins/Runtime/NetworkPrediction/Source/NetworkPrediction/Public/NetworkSimulationModelBuffer.h b/Engine/Plugins/Runtime/NetworkPrediction/Source/NetworkPrediction/Public/NetworkSimulationModelBuffer.h
--- a/Engine/Plugins/Runtime/NetworkPrediction/Source/NetworkPrediction/Public/NetworkSimulationModelBuffer.h
+++ b/Engine/Plugins/Runtime/NetworkPrediction/Source/NetworkPrediction/Public/NetworkSimulationModelBuffer.h
@@ -2,227 +2,291 @@
 
 #pragma once
 
-// ----------------------------------------------------------------------------------------------------------------------------------------------------------------------
-//	TReplicationBuffer
-//
-//	Generic Cyclic buffer. Has canonical head position. This is the "Client Frame" or "Keyframe" identifier used by the rest of the system.
-//	Contract: Elements in buffer are contiguously valid. We don't allow "gaps" in the buffer!
-//	Use ::GetWriteNext<T> to write to the buffer.
-//	Use ::CreateIterator to iterate from Tail->Head
-//
-//	for (auto It = Buffer.CreateIterator(); It; ++It) {
-//		FMyType* Element = It.Element();
-//		int32 Keyframe = It.Keyframe();
-//	}
-//
-// ----------------------------------------------------------------------------------------------------------------------------------------------------------------------
-template<typename T>
-struct TReplicationBuffer
+template<typename TBuffer, typename ElementType>
+struct TNetworkSimBufferIterator
 {
-	// Set the buffer size. Resizing is supported but not optimized because we preserve head/tail iteration (not a simple memcopy since where buffer wraps around will change)
-	// In general, avoid resizing outside of startup/initialization.
-	void SetBufferSize(int32 NewMaxNumElements)
-	{
-		check(NewMaxNumElements >= 0);	
-		if (Data.Num() == NewMaxNumElements)
-		{
-			return;
-		}
-		
-		if (Data.Num() == 0)
-		{
-			Data.SetNum(NewMaxNumElements, false);
-		}
-		else
-		{
-			// Grow or shrink. This is far from optimal but this operation should be rare
-			TReplicationBuffer<T> Old(MoveTemp(*this)); // Move old data to a copy
-			Data.SetNum(NewMaxNumElements, false);
-			ResetNextHeadKeyframe(Old.GetTailKeyframe(), true); // Reset so our next write is at the old copies tail keyframe
+	TNetworkSimBufferIterator(TBuffer& InBuffer) : Buffer(InBuffer)	{ CurrentPos = Buffer.IterStartPos(); }
+	int32 Keyframe() const { return Buffer.IterKeyframe(CurrentPos); }
+	ElementType* Element() const { return Buffer.IterElement(CurrentPos); }
 
-			for (auto It = Old.CreateIterator(); It; ++It)
-			{
-				T* WriteNext = GetWriteNext();
-				T* ReadNext = It.Element();
-				FMemory::Memcpy(WriteNext, ReadNext, sizeof(T));
-			}
+	void operator++() { CurrentPos++; }
+	operator bool()	{ return CurrentPos <= Buffer.IterEndPos(); }
 
-			// Blow away the temp copy so we don't run the dtor on it on scope end (because we memcpy'd it over, not copy assigned)
-			Old.Data.Reset();
-			Old.NumValidElements = 0;
-		}
-	}
+private:
+	TBuffer& Buffer;
+	int32 CurrentPos = -1;
+};
 
-	// Gets the oldest element in the buffer (or null)
-	const T* GetElementFromTail(int32 OffsetFromTail) const
+template<typename T>
+struct TNetworkSimBufferBase
+{
+	FString GetBasicDebugStr() const
 	{
-		return GetElementFromTailImpl(OffsetFromTail);
+		return FString::Printf(TEXT("Elements: [%d/%d]. Keyframes: [%d-%d]"), ((T*)this)->Num(), ((T*)this)->Max(), ((T*)this)->TailKeyframe(), ((T*)this)->HeadKeyframe());
 	}
 
-	T* GetElementFromTail(int32 OffsetFromTail)
-	{
-		const T* Element = const_cast<TReplicationBuffer<T>*>(this)->GetElementFromTailImpl(OffsetFromTail);
-		return const_cast<T*>(Element);
+	bool IsValidKeyframe(int32 Keyframe) const 
+	{ 
+		return Keyframe >= ((T*)this)->TailKeyframe() && Keyframe <= ((T*)this)->HeadKeyframe(); 
 	}
 
-	// Gets the newest element in the buffer (or null)
-	const T* GetElementFromHead(int32 OffsetFromHead) const
+	// Copies each element of SourceBuffer into this buffer. This could break continuity so for TNetworkSimContiguousBuffer, it is possible to lose elements by doing this.
+	template<typename TSource>
+	void CopyAndMerge(const TSource& SourceBuffer)
 	{
-		return GetElementFromHeadImpl(OffsetFromHead);
+		for (auto It = SourceBuffer.CreateConstIterator(); It; ++It)
+		{
+			*((T*)this)->WriteKeyframe( It.Keyframe() ) = *It.Element();
+		}
 	}
+};
 
-	T* GetElementFromHead(int32 OffsetFromHead)
-	{
-		const T* Element = const_cast<TReplicationBuffer<T>*>(this)->GetElementFromHeadImpl(OffsetFromHead);
-		return const_cast<T*>(Element);
-	}
-	
-	const T* FindElementByKeyframe(int32 Key) const
+template<typename T, int32 NumElements=32>
+struct TNetworkSimContiguousBuffer : public TNetworkSimBufferBase<TNetworkSimContiguousBuffer<T, NumElements>>
+{
+	using ElementType = T;
+	TNetworkSimContiguousBuffer()
 	{
-		return FindElementByKeyframeImpl(Key);
+		Data.SetNum(NumElements);
 	}
 
-	T* FindElementByKeyframe(int32 Key)
-	{
-		const T* Element = const_cast<TReplicationBuffer<T>*>(this)->FindElementByKeyframeImpl(Key);
-		return const_cast<T*>(Element);
-	}
+	T* operator[](int32 Keyframe) {	return const_cast<T*>(GetImpl(Keyframe)); }
+	const T* operator[](int32 Keyframe) const { return GetImpl(Keyframe); }
 
-	// Returns the next element for writing. The contents of this element are unknown (could be stale content!). This also advances the head. Note the element returned is immediately considered "valid"!
-	T* GetWriteNext()
-	{
-		check(Data.Num() > 0); // Buffer must be initialized first
-		
-		++Head;
-		++DirtyCount;
-		NumValidElements = FMath::Min<int32>(NumValidElements+1, Data.Num());
+	T* Get(int32 Keyframe) { return const_cast<T*>(GetImpl(Keyframe)); }
+	const T* Get(int32 Keyframe) const { return GetImpl(Keyframe); }
 
-		const int32 Offset = (Head % Data.Num());
-		return &Data[Offset];
-	}
+	const T* HeadElement() const { return HeadElementImpl(); }
+	T* HeadElement() { return const_cast<T*>(HeadElementImpl()); }
 
-	// Resets the head keyframen, so that NextHeadKeyframe will be the keyframe of the next GetWriteNext<>(). E.g, call this before you want to write to NextHeadKeyframe.
-	// To remphasize: NextHeadKeyframe is where the NEXT write will go. NOT "where the current head is".
-	// Contents will be preserved if possible. E.g, if NextHeadKeyframe-1 is a valid keyframe, it will remain valid. If its not (if this creates "gaps" in the buffer) then
-	// the contents will effectively be cleared. bForceClearContents will force this behavior even if the previous keyframe was valid.
-	void ResetNextHeadKeyframe(int32 NextHeadKeyframe=0, bool bForceClearContents=false)
+	const T* TailElement() const { return TailElementImpl(); }
+	T* TailElement() { return const_cast<T*>(TailElementImpl()); }
+
+	int32 HeadKeyframe() const { return Head; }
+	int32 TailKeyframe() const { return GetTail(); }
+
+	int32 Num() const { return Head == INDEX_NONE ? 0 : (Head - GetTail() + 1); }
+	int32 Max() const { return Data.Max(); }
+	int32 GetDirtyCount() const { return DirtyCount; }
+
+	T* WriteKeyframe(int32 Keyframe)
 	{
-		const int32 NewHeadKeyframe = NextHeadKeyframe - 1;
-		if (bForceClearContents || NewHeadKeyframe < GetTailKeyframe() || NewHeadKeyframe > GetHeadKeyframe())
+		check(Keyframe >= 0);
+
+		const int32 Tail = GetTail();
+		if (Head == INDEX_NONE || Keyframe < Tail || Keyframe > Head+1)
 		{
-			NumValidElements = 0;
+			// Writing outside the current range (+1) results in a full wipe of valid contents
+			NumValidElements = 1;
 		}
 		else
 		{
-			NumValidElements += (NewHeadKeyframe - GetHeadKeyframe());
-			check(NumValidElements >= 0 && NumValidElements <= Data.Num());
+			// Writing inside the current range (+1) preserves valid elements up to the size of our buffer
+			NumValidElements = FMath::Min<int32>(Keyframe - Tail + 1, Data.Num());
 		}
-
-		Head =  NewHeadKeyframe;
-		DirtyCount++;
+		
+		Head = Keyframe;
+		++DirtyCount;
+		return &Data[Keyframe % Data.Num()];
 	}
 
-	// Copies the contents of the SourceBuffer into this buffer. This existing buffer's data will be preserved as much as possible, but the guarantee is that the SourceBuffer's
-	// data will all be added to this buffer.
-	
-	// Example 1: Target={1...5}, Source={3...9}  --> Target={1...9}
-	// Example 2: Target={1...5}, Source={7...9}  --> Target={7...9}
-	// Exmaple 3: Target={6...9}, Source={1...4}  --> Target={1...4}
-	void CopyAndMerge(const TReplicationBuffer<T>& SourceBuffer)
+	TUniqueFunction<T*()> WriteKeyframeFunc(int32 Keyframe)
 	{
-		const int32 StartingHeadKeyframe = GetHeadKeyframe();
+		return [this, Keyframe]() { return WriteKeyframe(Keyframe); };
+	}
+
+	using IteratorType = TNetworkSimBufferIterator<TNetworkSimContiguousBuffer<ElementType, NumElements>, ElementType>;
+	using ConstIteratorType = TNetworkSimBufferIterator<const TNetworkSimContiguousBuffer<ElementType, NumElements>, const ElementType>;
+
+	IteratorType CreateIterator() { return IteratorType(*this); }
+	ConstIteratorType CreateConstIterator() const { return ConstIteratorType(*this); }
 
-		ResetNextHeadKeyframe(SourceBuffer.GetTailKeyframe());
-		for (int32 Keyframe = SourceBuffer.GetTailKeyframe(); Keyframe <= SourceBuffer.GetHeadKeyframe(); ++Keyframe)
+private:
+
+	const T* GetImpl(int32 Keyframe) const
+	{
+		const int32 RelativeToHead = Keyframe - Head;
+		if (RelativeToHead > 0 || RelativeToHead <= -NumValidElements)
 		{
-			T* TargetData = GetWriteNext();
-			check(GetHeadKeyframe() == Keyframe);
+			return nullptr;
+		}
 
-			const T* SourceData = SourceBuffer.FindElementByKeyframe(Keyframe);
-			check(SourceData != nullptr);
+		return &Data[Keyframe % Data.Num()];
+	}
 
-			*TargetData = *SourceData;
+	const T* HeadElementImpl() const
+	{
+		if (Head != INDEX_NONE)
+		{
+			return &Data[Head % Data.Num()];
 		}
+		return nullptr;
 	}
 
-	FString GetBasicDebugStr() const
+	const T* TailElementImpl() const
+	{
+		const int32 Tail = GetTail();
+		if (Tail != INDEX_NONE)
+		{
+			return &Data[Tail % Data.Num()];
+		}
+		return nullptr;
+	}
+
+	int32 GetTail() const { return Head == INDEX_NONE ? INDEX_NONE : (Head - NumValidElements + 1); };
+
+	int32 DirtyCount = 0;
+	int32 Head = INDEX_NONE;
+	int32 NumValidElements = 0;
+	TArray<ElementType, TInlineAllocator<NumElements>> Data;
+
+	// Efficient position based accessing through TNetworkSimBufferIterator.
+	template<typename, typename>
+	friend struct TNetworkSimBufferIterator;
+	int32 IterStartPos() const { return GetTail(); }
+	int32 IterEndPos() const { return Head; }
+	int32 IterKeyframe(int32 Pos) const { return Pos; }
+	T* IterElement(int32 Pos) { return &Data[Pos % Data.Num()]; }
+	const T* IterElement(int32 Pos) const { return &Data[Pos % Data.Num()]; }
+};
+
+// Keyframe: arbitrary identifier for data. Not contiguous or controlled by us (always passed in)
+// Position/Pos: increasing counter for position in array. Mod with Data.Num to get Index.
+// Index/Idx: actual index into Data Array
+template<typename T, int32 NumElements=32>
+struct TNetworkSimSparseBuffer : public TNetworkSimBufferBase<TNetworkSimSparseBuffer<T, NumElements>>
+{
+	using ElementType = T;
+	TNetworkSimSparseBuffer()
 	{
-		return FString::Printf(TEXT("Elements: [%d/%d]. Keyframes: [%d-%d]"), NumValidElements, Data.Num(), GetTailKeyframe(), GetHeadKeyframe());
+		Data.SetNum(NumElements);
 	}
 
-	int32 GetNumValidElements() const { return NumValidElements; }
-	int32 GetMaxNumElements() const { return Data.Num(); }
-	int32 GetHeadKeyframe() const { return Head; }
-	int32 GetTailKeyframe() const { return Head - NumValidElements + 1; }
-	bool IsValidKeyframe(int32 Keyframe) const { return Keyframe >= GetTailKeyframe() && Keyframe <= GetHeadKeyframe(); }
+	T* operator[](int32 Keyframe) {	return const_cast<T*>(GetImpl(Keyframe)); }
+	const T* operator[](int32 Keyframe) const { return GetImpl(Keyframe); }
+
+	T* Get(int32 Keyframe) { return const_cast<T*>(GetImpl(Keyframe)); }
+	const T* Get(int32 Keyframe) const { return GetImpl(Keyframe); }
+
+	const T* HeadElement() const { return HeadElementImpl(); }
+	T* HeadElement() { return const_cast<T*>(HeadElementImpl()); }
+
+	const T* TailElement() const { return TailElementImpl(); }
+	T* TailElement() { return const_cast<T*>(TailElementImpl()); }
+
+	int32 HeadKeyframe() const { return Data[HeadPos % Data.Num()].Keyframe; }
+	int32 TailKeyframe() const { return Data[GetTailPos() % Data.Num()].Keyframe; }
+
+	int32 Num() const { return HeadPos == INDEX_NONE ? 0 : (HeadPos - GetTailPos() + 1); }
+	int32 Max() const { return Data.Max(); }
 	int32 GetDirtyCount() const { return DirtyCount; }
 
-	// Create an iterator from tail->head. Note that no matter what template class is, this WILL iterate correctly across the element. The templated type
-	// is just for casting the return element type. E.g, it is fine to use <uint8> in generic code, this will always step by 'StructSize'.
-	// FIXME: this is slow using FindElementByKeyframe. In general this entire class needs to be optimized
-	struct TGenericIterator
+	// Returns the element @ keyframe for writing. The contents of this element are unknown (could be stale content!). Note the element returned is immediately considered "valid" by Num(), Iterators, etc!
+	T* WriteKeyframe(int32 Keyframe)
 	{
-		TGenericIterator(TReplicationBuffer<T>& InBuffer) : Buffer(InBuffer)
+		check(Keyframe >= 0);
+		++DirtyCount;
+
+		int32 Pos = HeadPos;
+		if (HeadPos != INDEX_NONE)
 		{
-			CurrentKeyframe = Buffer.GetTailKeyframe();
+			// If we have elements, find where this Keyframe would go
+			const int32 Tail = GetTailPos();
+			for (; Pos >= Tail; --Pos)
+			{
+				TInternal& InternalData = Data[Pos % Data.Num()];
+				if (InternalData.Keyframe == Keyframe)
+				{
+					return &InternalData.Element;
+				}
+				if (InternalData.Keyframe <= Keyframe)
+				{
+					break;
+				}
+			}
 		}
 
-		int32 Keyframe() const { return CurrentKeyframe; }
-		T* Element() const { return Buffer.FindElementByKeyframe(CurrentKeyframe); }
-		void operator++() { CurrentKeyframe++; }
-		operator bool()	{ return CurrentKeyframe <= Buffer.GetHeadKeyframe(); }
+		// Write frame to next pos
+		HeadPos = Pos+1;
+		TInternal& WriteData = Data[HeadPos % Data.Num()];
+		WriteData.Keyframe = Keyframe;
+		return &WriteData.Element;
+	}
 
-	private:
-		TReplicationBuffer<T>& Buffer;
-		int32 CurrentKeyframe = -1;
-	};
-		
-	TGenericIterator CreateIterator() { return TGenericIterator(*this); }
+	TUniqueFunction<T*()> WriteKeyframeFunc(int32 Keyframe)
+	{
+		return [this, Keyframe]() { return WriteKeyframe(Keyframe); };
+	}
+
+	using IteratorType = TNetworkSimBufferIterator<TNetworkSimSparseBuffer<ElementType, NumElements>, ElementType>;
+	using ConstIteratorType = TNetworkSimBufferIterator<const TNetworkSimSparseBuffer<ElementType, NumElements>, const ElementType>;
+
+	IteratorType CreateIterator() { return IteratorType(*this); }
+	ConstIteratorType CreateConstIterator() const { return ConstIteratorType(*this); }
 
 private:
 
-	const T* FindElementByKeyframeImpl(int32 Key) const
+	struct TInternal
 	{
-		const int32 RelativeToHead = Key - Head;
-		if (RelativeToHead > 0 || RelativeToHead <= -NumValidElements)
-		{
-			return nullptr;
+		int32 Keyframe = -1;
+		T Element;
+	};
+
+	const T* GetImpl(int32 Keyframe) const
+	{
+		const int32 Pos = GetPosForKeyframe(Keyframe);
+		if (Pos != INDEX_NONE)
+		{		
+			return &Data[Pos % Data.Num()].Element;
 		}
 
-		return &Data[Key % Data.Num()];
+		return nullptr;
 	}
 
-	const T* GetElementFromTailImpl(int32 OffsetFromTail) const
+	const T* HeadElementImpl() const
 	{
-		check(OffsetFromTail >= 0); // Must always pass in valid offset	
-		if (NumValidElements <= OffsetFromTail)
+		if (HeadPos != INDEX_NONE)
 		{
-			return nullptr;
+			return &Data[HeadPos % Data.Num()].Element;
 		}
-		
-		const int32 Tail = GetTailKeyframe();
-		const int32 Position = Tail + OffsetFromTail;
-		const int32 Offset = (Position % Data.Num());
-		return &Data[Offset];
+		return nullptr;
 	}
 
-	const T* GetElementFromHeadImpl(int32 OffsetFromHead) const
+	const T* TailElementImpl() const
 	{
-		check(OffsetFromHead >= 0); // Must always pass in valid offset
-		if (NumValidElements <= OffsetFromHead)
+		const int32 TailPos = GetTailPos();
+		if (TailPos != INDEX_NONE)
 		{
-			return nullptr;
+			return &Data[TailPos % Data.Num()].Element;
 		}
-		
-		const int32 Position = Head - OffsetFromHead;
-		check(Position >= 0);
-		const int32 Offset = (Position % Data.Num());
-		return &Data[Offset];
+		return nullptr;
+	}
+	
+	int32 GetPosForKeyframe(int32 Keyframe) const
+	{
+		const int32 TailPos = GetTailPos();
+		for (int32 Pos=HeadPos; Pos >= TailPos; --Pos)
+		{
+			if (Data[Pos % Data.Num()].Keyframe <= Keyframe)
+			{
+				return Pos;
+			}
+		}
+		return INDEX_NONE;
 	}
 
-	int32 Head = INDEX_NONE;
-	int32 DirtyCount = 0;
-	int32 NumValidElements = 0;
+	int32 GetTailPos() const { return HeadPos == INDEX_NONE ? INDEX_NONE : FMath::Max<int32>(0, HeadPos - Data.Num()); }
+
+	int32 DirtyCount=0;
+	int32 HeadPos=INDEX_NONE;
+	TArray<TInternal, TInlineAllocator<NumElements>> Data;
 
-	TArray<T> Data;
+	// Efficient position based accessing through TNetworkSimBufferIterator.
+	template<typename, typename>
+	friend struct TNetworkSimBufferIterator;
+	int32 IterStartPos() const { return GetTailPos(); }
+	int32 IterEndPos() const { return HeadPos; }
+	int32 IterKeyframe(int32 Pos) const { return Data[Pos % Data.Num()].Keyframe; }
+	T* IterElement(int32 Pos) { return &Data[Pos % Data.Num()].Element; }
+	const T* IterElement(int32 Pos) const { return &Data[Pos % Data.Num()].Element; }
 };
\ No newline at end of file
diff --git a/Engine/Plugins/Runtime/NetworkPrediction/Source/NetworkPrediction/Public/NetworkSimulationModelDebugger.h b/Engine/Plugins/Runtime/NetworkPrediction/Source/NetworkPrediction/Public/NetworkSimulationModelDebugger.h
--- a/Engine/Plugins/Runtime/NetworkPrediction/Source/NetworkPrediction/Public/NetworkSimulationModelDebugger.h
+++ b/Engine/Plugins/Runtime/NetworkPrediction/Source/NetworkPrediction/Public/NetworkSimulationModelDebugger.h
@@ -207,7 +207,7 @@ struct NETWORKPREDICTION_API FNetworkSimulationModelDebuggerManager: public FTic
 		FStandardLoggingParameters LocalParameters = Parameters;
 		LocalParameters.Ar = &StrOut;
 
-		auto* Element = Buffer.FindElementByKeyframe(LocalParameters.Keyframe);
+		auto* Element = Buffer[LocalParameters.Keyframe];
 		if (Element)
 		{
 			Element->Log(LocalParameters);
@@ -364,14 +364,15 @@ struct TNetworkSimulationModelDebugger : public INetworkSimulationModelDebugger
 		TMap<int32, FScreenPositions> Keyframes;
 	};
 
-	void GatherDebugGraph(FNetworkSimulationModelDebuggerManager& Out, UCanvas* Canvas, TReplicationBuffer<TDebugState>* DebugBuffer, FRect DrawRect, const float MaxColumnTimeSeconds, const float MaxLocalFrameTime, const FString& Header, FCachedScreenPositionMap& SendCache, FCachedScreenPositionMap& RecvCache)
+	template<typename TBuffer>
+	void GatherDebugGraph(FNetworkSimulationModelDebuggerManager& Out, UCanvas* Canvas, TBuffer* DebugBuffer, FRect DrawRect, const float MaxColumnTimeSeconds, const float MaxLocalFrameTime, const FString& Header, FCachedScreenPositionMap& SendCache, FCachedScreenPositionMap& RecvCache)
 	{
 		static float Pad = 2.f;
 		static float BaseLineYPCT = 0.8f;
 
 		auto& InputBuffer = NetworkSim->GetHistoricBuffers() ? NetworkSim->GetHistoricBuffers()->Input : NetworkSim->Buffers.Input;
 
-		if (Canvas && DebugBuffer && DebugBuffer->GetNumValidElements() > 0)
+		if (Canvas && DebugBuffer && DebugBuffer->Num() > 0)
 		{
 			// Outline + Header
 			Out.EmitLine( FVector2D(DrawRect.Min.X, DrawRect.Min.Y), FVector2D(DrawRect.Min.X, DrawRect.Max.Y), FColor::White );
@@ -418,7 +419,7 @@ struct TNetworkSimulationModelDebugger : public INetworkSimulationModelDebugger
 
 				for (int32 Keyframe : DebugState->ProcessedKeyframes)
 				{
-					auto* Cmd = InputBuffer.FindElementByKeyframe(Keyframe);
+					auto* Cmd = InputBuffer[Keyframe];
 					if (Cmd)
 					{
 						float ClientSizeX = FixedWidth;
@@ -434,7 +435,7 @@ struct TNetworkSimulationModelDebugger : public INetworkSimulationModelDebugger
 				// Unprocessed InputCmds (above processed)				
 				for (int32 Keyframe = DebugState->LastProcessedKeyframe+1; Keyframe <= DebugState->HeadKeyframe; ++Keyframe)
 				{
-					if (auto* Cmd = InputBuffer.FindElementByKeyframe(Keyframe))
+					if (auto* Cmd = InputBuffer[Keyframe])
 					{
 						float ClientSizeX = FixedWidth;
 						float ClientSizeY = SecondsToPixelsY * Cmd->GetFrameDeltaTime().ToRealTimeSeconds();
@@ -489,7 +490,8 @@ struct TNetworkSimulationModelDebugger : public INetworkSimulationModelDebugger
 		}
 	}
 
-	void CalcMaxColumnFrameTime(TReplicationBuffer<TDebugState>* DebugBuffer, float& MaxInputTime, float& MaxLocalFrameTime)
+	template<typename TBuffer>
+	void CalcMaxColumnFrameTime(TBuffer* DebugBuffer, float& MaxInputTime, float& MaxLocalFrameTime)
 	{
 		auto& InputBuffer = NetworkSim->GetHistoricBuffers() ? NetworkSim->GetHistoricBuffers()->Input : NetworkSim->Buffers.Input;
 
@@ -500,7 +502,7 @@ struct TNetworkSimulationModelDebugger : public INetworkSimulationModelDebugger
 			auto* DebugState = It.Element();
 			for (int32 Keyframe : DebugState->ProcessedKeyframes)
 			{
-				auto* Cmd = InputBuffer.FindElementByKeyframe(Keyframe);
+				auto* Cmd = InputBuffer[Keyframe];
 				if (Cmd)
 				{
 					ColumnTime += (float)Cmd->GetFrameDeltaTime().ToRealTimeSeconds();
@@ -508,7 +510,7 @@ struct TNetworkSimulationModelDebugger : public INetworkSimulationModelDebugger
 			}
 			for (int32 Keyframe = DebugState->LastProcessedKeyframe+1; Keyframe <= DebugState->HeadKeyframe; ++Keyframe)
 			{
-				if (auto* Cmd = InputBuffer.FindElementByKeyframe(Keyframe))
+				if (auto* Cmd = InputBuffer[Keyframe])
 				{
 					ColumnTime += (float)Cmd->GetFrameDeltaTime().ToRealTimeSeconds();
 				}
@@ -532,7 +534,7 @@ struct TNetworkSimulationModelDebugger : public INetworkSimulationModelDebugger
 		// ------------------------------------------------------------------------------------------------------------------------------------------------
 
 		Out.Emit(FString::Printf(TEXT("%s - %s"), *Owner->GetName(), *UEnum::GetValueAsString(TEXT("Engine.ENetRole"), Owner->GetLocalRole())), FColor::Yellow);
-		Out.Emit(FString::Printf(TEXT("LastProcessedInputKeyframe: %d (%d Buffered)"), NetworkSim->TickInfo.LastProcessedInputKeyframe, NetworkSim->Buffers.Input.GetHeadKeyframe() - NetworkSim->TickInfo.LastProcessedInputKeyframe));
+		Out.Emit(FString::Printf(TEXT("LastProcessedInputKeyframe: %d (%d Buffered)"), NetworkSim->TickInfo.LastProcessedInputKeyframe, NetworkSim->Buffers.Input.HeadKeyframe() - NetworkSim->TickInfo.LastProcessedInputKeyframe));
 				
 		if (Owner->GetLocalRole() == ROLE_AutonomousProxy)
 		{			
@@ -544,15 +546,15 @@ struct TNetworkSimulationModelDebugger : public INetworkSimulationModelDebugger
 			// Calc how much predicted time we have processed. Note that we use the motionstate buffer to iterate but the MS is on the input cmd. (if we are buffering cmds, don't want to count them)
 			
 			TSimTime PredictedMS;
-			for (int32 PredKeyrame = LastSerializedKeyframe+1; PredKeyrame <= NetworkSim->Buffers.Sync.GetHeadKeyframe(); ++PredKeyrame)
+			for (int32 PredKeyframe = LastSerializedKeyframe+1; PredKeyframe <= NetworkSim->Buffers.Sync.HeadKeyframe(); ++PredKeyframe)
 			{
-				if (auto* Cmd = NetworkSim->Buffers.Input.FindElementByKeyframe(PredKeyrame))
+				if (auto* Cmd = NetworkSim->Buffers.Input[PredKeyframe])
 				{
 					PredictedMS += Cmd->GetFrameDeltaTime();
 				}
 			}
 
-			FString ConfirmedFrameStr = FString::Printf(TEXT("LastConfirmedFrame: %d. Prediction: %d Frames, %s MS"), LastSerializedKeyframe, NetworkSim->Buffers.Sync.GetHeadKeyframe() - LastSerializedKeyframe, *PredictedMS.ToString());
+			FString ConfirmedFrameStr = FString::Printf(TEXT("LastConfirmedFrame: %d. Prediction: %d Frames, %s MS"), LastSerializedKeyframe, NetworkSim->Buffers.Sync.HeadKeyframe() - LastSerializedKeyframe, *PredictedMS.ToString());
 			if (FaultDetected)
 			{
 				ConfirmedFrameStr += TEXT(" RECONCILE FAULT DETECTED!");
@@ -565,7 +567,7 @@ struct TNetworkSimulationModelDebugger : public INetworkSimulationModelDebugger
 				*NetworkSim->RepProxy_Autonomous.GetLastSerializedSimTime().ToString(), *(NetworkSim->TickInfo.GetTotalProcessedSimulationTime() - NetworkSim->RepProxy_Autonomous.GetLastSerializedSimTime()).ToString());
 			Out.Emit(*SimulationTimeString, Color);
 
-			FString AllowedSimulationTimeString = FString::Printf(TEXT("Allowed Simulation Time: %s. Keyframe: %d/%d/%d"), *NetworkSim->TickInfo.GetRemainingAllowedSimulationTime().ToString(), NetworkSim->TickInfo.MaxAllowedInputKeyframe, NetworkSim->TickInfo.LastProcessedInputKeyframe, NetworkSim->Buffers.Input.GetHeadKeyframe());
+			FString AllowedSimulationTimeString = FString::Printf(TEXT("Allowed Simulation Time: %s. Keyframe: %d/%d/%d"), *NetworkSim->TickInfo.GetRemainingAllowedSimulationTime().ToString(), NetworkSim->TickInfo.MaxAllowedInputKeyframe, NetworkSim->TickInfo.LastProcessedInputKeyframe, NetworkSim->Buffers.Input.HeadKeyframe());
 			Out.Emit(*AllowedSimulationTimeString, Color);
 		}
 		else if (Owner->GetLocalRole() == ROLE_SimulatedProxy)
@@ -581,7 +583,7 @@ struct TNetworkSimulationModelDebugger : public INetworkSimulationModelDebugger
 			Out.Emit(FString::Printf(TEXT("//////////////// %s ///////////////"), *BufferName), FColor::Yellow);
 			Out.Emit(FString::Printf(TEXT("%s"), *Buffer.GetBasicDebugStr()));		
 			Out.Emit();
-			Out.EmitElement(Buffer, FStandardLoggingParameters(nullptr, EStandardLoggingContext::Full, Buffer.GetHeadKeyframe()));
+			Out.EmitElement(Buffer, FStandardLoggingParameters(nullptr, EStandardLoggingContext::Full, Buffer.HeadKeyframe()));
 		};
 
 		EmitBuffer(TEXT("InputBuffer"), NetworkSim->Buffers.Input);
@@ -658,34 +660,34 @@ struct TNetworkSimulationModelDebugger : public INetworkSimulationModelDebugger
 
 		UWorld* World = Owner->GetWorld();
 
-		if (auto* LatestSync = NetworkSim->Buffers.Sync.GetElementFromHead(0))
+		if (auto* LatestSync = NetworkSim->Buffers.Sync.HeadElement())
 		{
-			LatestSync->VisualLog( FVisualLoggingParameters(EVisualLoggingContext::LastPredicted, NetworkSim->Buffers.Sync.GetHeadKeyframe(), EVisualLoggingLifetime::Transient), NetworkSim->Driver, NetworkSim->Driver);
+			LatestSync->VisualLog( FVisualLoggingParameters(EVisualLoggingContext::LastPredicted, NetworkSim->Buffers.Sync.HeadKeyframe(), EVisualLoggingLifetime::Transient), NetworkSim->Driver, NetworkSim->Driver);
 		}
 
 		FStuff ServerPIEStuff = GetServerPIEStuff();
 		if (ServerPIEStuff.NetworkSim)
 		{
-			if (auto* ServerLatestSync = ServerPIEStuff.NetworkSim->Buffers.Sync.GetElementFromHead(0))
+			if (auto* ServerLatestSync = ServerPIEStuff.NetworkSim->Buffers.Sync.HeadElement())
 			{
-				ServerLatestSync->VisualLog( FVisualLoggingParameters(EVisualLoggingContext::CurrentServerPIE, ServerPIEStuff.NetworkSim->Buffers.Sync.GetHeadKeyframe(), EVisualLoggingLifetime::Transient), ServerPIEStuff.NetworkSim->Driver, NetworkSim->Driver);
+				ServerLatestSync->VisualLog( FVisualLoggingParameters(EVisualLoggingContext::CurrentServerPIE, ServerPIEStuff.NetworkSim->Buffers.Sync.HeadKeyframe(), EVisualLoggingLifetime::Transient), ServerPIEStuff.NetworkSim->Driver, NetworkSim->Driver);
 			}
 		}
 
 		if (Owner->GetLocalRole() == ROLE_AutonomousProxy)
 		{
-			for (int32 Keyframe = NetworkSim->RepProxy_Autonomous.GetLastSerializedKeyframe(); Keyframe < NetworkSim->Buffers.Sync.GetHeadKeyframe(); ++Keyframe)
+			for (int32 Keyframe = NetworkSim->RepProxy_Autonomous.GetLastSerializedKeyframe(); Keyframe < NetworkSim->Buffers.Sync.HeadKeyframe(); ++Keyframe)
 			{
-				if (auto* SyncState = NetworkSim->Buffers.Sync.FindElementByKeyframe(Keyframe))
+				if (auto* SyncState = NetworkSim->Buffers.Sync[Keyframe])
 				{
 					const EVisualLoggingContext Context = (Keyframe == NetworkSim->RepProxy_Autonomous.GetLastSerializedKeyframe()) ? EVisualLoggingContext::LastConfirmed : EVisualLoggingContext::OtherPredicted;
-					SyncState->VisualLog( FVisualLoggingParameters(Context, NetworkSim->Buffers.Sync.GetHeadKeyframe(), EVisualLoggingLifetime::Transient), NetworkSim->Driver, NetworkSim->Driver);
+					SyncState->VisualLog( FVisualLoggingParameters(Context, NetworkSim->Buffers.Sync.HeadKeyframe(), EVisualLoggingLifetime::Transient), NetworkSim->Driver, NetworkSim->Driver);
 				}
 			}
 		}
 		else if (Owner->GetLocalRole() == ROLE_SimulatedProxy)
 		{
-			NetworkSim->RepProxy_Simulated.GetLastSerializedSyncState().VisualLog( FVisualLoggingParameters(EVisualLoggingContext::LastConfirmed, NetworkSim->Buffers.Sync.GetHeadKeyframe(), EVisualLoggingLifetime::Transient), NetworkSim->Driver, NetworkSim->Driver); 
+			NetworkSim->RepProxy_Simulated.GetLastSerializedSyncState().VisualLog( FVisualLoggingParameters(EVisualLoggingContext::LastConfirmed, NetworkSim->Buffers.Sync.HeadKeyframe(), EVisualLoggingLifetime::Transient), NetworkSim->Driver, NetworkSim->Driver); 
 			if (NetworkSim->GetSimulatedUpdateMode() != ESimulatedUpdateMode::Interpolate)
 			{
 				FVector2D ServerSimulationTimeData(Owner->GetWorld()->GetTimeSeconds(), NetworkSim->RepProxy_Simulated.GetLastSerializedSimulationTime().ToRealTimeMS());
diff --git a/Engine/Plugins/Runtime/NetworkPrediction/Source/NetworkPrediction/Public/NetworkSimulationModelInterpolator.h b/Engine/Plugins/Runtime/NetworkPrediction/Source/NetworkPrediction/Public/NetworkSimulationModelInterpolator.h
--- a/Engine/Plugins/Runtime/NetworkPrediction/Source/NetworkPrediction/Public/NetworkSimulationModelInterpolator.h
+++ b/Engine/Plugins/Runtime/NetworkPrediction/Source/NetworkPrediction/Public/NetworkSimulationModelInterpolator.h
@@ -17,9 +17,9 @@ namespace NetworkInterpolationDebugCVars
 
 	NETSIM_DEVCVAR_SHIPCONST_FLOAT(WaitSlack, 0.05, "ni.WaitSlack", "How much slack to wait for when waiting");
 
-	NETSIM_DEVCVAR_SHIPCONST_FLOAT(CatchUpThreshold, 0.200, "ni.CatchUpThreshold", "When we start catching up (seconds from head)");
-	NETSIM_DEVCVAR_SHIPCONST_FLOAT(CatchUpGoal, 0.100, "ni.CatchUpGoal", "When we stop cathcing up (seconds from head)");
-	NETSIM_DEVCVAR_SHIPCONST_FLOAT(CatchUpFactor, 1.25, "ni.CatchUpFactor", "Factor we use to catch up");
+	NETSIM_DEVCVAR_SHIPCONST_FLOAT(CatchUpThreshold, 0.300, "ni.CatchUpThreshold", "When we start catching up (seconds from head)");
+	NETSIM_DEVCVAR_SHIPCONST_FLOAT(CatchUpGoal, 0.010, "ni.CatchUpGoal", "When we stop cathcing up (seconds from head)");
+	NETSIM_DEVCVAR_SHIPCONST_FLOAT(CatchUpFactor, 1.50, "ni.CatchUpFactor", "Factor we use to catch up");
 }
 
 template<typename TBufferTypes, typename TTickSettings>
@@ -28,6 +28,7 @@ struct TInterpolator
 	using FSimTime = FNetworkSimTime;
 	using FRealTime = FNetworkSimTime::FRealTime;
 	using TSyncState = typename TBufferTypes::TSyncState;
+	using TAuxState = typename TBufferTypes::TAuxState;
 
 	bool bEnableVisualLog = true;
 
@@ -39,49 +40,50 @@ struct TInterpolator
 
 		const UObject* LogOwner = Driver->GetVLogOwner();
 
-		//
-
+		// Interpolation disabled
 		if (NetworkInterpolationDebugCVars::Disable() > 0)
 		{			
-			if (const TSyncState* HeadState = Buffers.Sync.GetElementFromHead(0))
+			if (const TSyncState* HeadState = Buffers.Sync.HeadElement())
 			{
-				Driver->FinalizeFrame(*HeadState);
+				const TAuxState* AuxState = Buffers.Aux.HeadElement();
+				check(AuxState);
+
+				Driver->FinalizeFrame(*HeadState, *AuxState);
 			}
 			return;
 		}
 		
-		if (TickInfo.SimulationTimeBuffer.GetNumValidElements() <= 1)
+		if (TickInfo.SimulationTimeBuffer.Num() <= 1)
 		{
 			// Cant interpolate yet	
 			return;
 		}
 
-
 		auto& SimulationTimeBuffer = TickInfo.SimulationTimeBuffer;
 
 		// Starting off: start at the tail end
 		if (InterpolationTime <= 0.f)
 		{
-			InterpolationTime = SimulationTimeBuffer.GetElementFromTail(0)->ToRealTimeSeconds();
-			InterpolationKeyframe = SimulationTimeBuffer.GetTailKeyframe();
-			InterpolationState = *Buffers.Sync.GetElementFromTail(0);
+			InterpolationTime = SimulationTimeBuffer.TailElement()->ToRealTimeSeconds();
+			InterpolationKeyframe = SimulationTimeBuffer.TailKeyframe();
+			InterpolationState = *Buffers.Sync.TailElement();
 		}
 
 		// Wait if we were too far ahead
 		if (WaitUntilTime > 0.f)
 		{
-			if (WaitUntilTime <= SimulationTimeBuffer.GetElementFromHead(0)->ToRealTimeSeconds())
+			if (WaitUntilTime <= SimulationTimeBuffer.HeadElement()->ToRealTimeSeconds())
 			{
 				// done waiting, we can continue
 				WaitUntilTime = 0.f;
-				UE_VLOG(LogOwner, LogNetInterpolation, Log, TEXT("Done Waiting! Head: %s"), *LexToString(SimulationTimeBuffer.GetElementFromHead(0)->ToRealTimeSeconds()));
+				UE_VLOG(LogOwner, LogNetInterpolation, Log, TEXT("Done Waiting! Head: %s"), *LexToString(SimulationTimeBuffer.HeadElement()->ToRealTimeSeconds()));
 			}
 			else
 			{
 				if (bDoVLog)
 				{
 					// Still waiting, return
-					UE_VLOG(LogOwner, LogNetInterpolation, Log, TEXT("Still Waiting! %s < %s"), *LexToString(WaitUntilTime), *LexToString(SimulationTimeBuffer.GetElementFromHead(0)->ToRealTimeSeconds()));
+					UE_VLOG(LogOwner, LogNetInterpolation, Log, TEXT("Still Waiting! %s < %s"), *LexToString(WaitUntilTime), *LexToString(SimulationTimeBuffer.HeadElement()->ToRealTimeSeconds()));
 
 					const FName LocalInterpolationTimeName("Local Interpolation Time");
 					FVector2D LocalTimeVsInterpolationTime(LogOwner->GetWorld()->GetTimeSeconds(), (int32)(InterpolationTime * 1000.f));
@@ -99,10 +101,10 @@ struct TInterpolator
 			FRealTime Step = DeltaSeconds;
 			
 			// Speed up if we are too far behind
-			FRealTime CatchUpThreshold = SimulationTimeBuffer.GetElementFromHead(0)->ToRealTimeSeconds() - NetworkInterpolationDebugCVars::CatchUpThreshold();
+			FRealTime CatchUpThreshold = SimulationTimeBuffer.HeadElement()->ToRealTimeSeconds() - NetworkInterpolationDebugCVars::CatchUpThreshold();
 			if (CatchUpUntilTime <= 0.f && InterpolationTime < CatchUpThreshold)
 			{
-				CatchUpUntilTime = SimulationTimeBuffer.GetElementFromHead(0)->ToRealTimeSeconds() - NetworkInterpolationDebugCVars::CatchUpGoal();
+				CatchUpUntilTime = SimulationTimeBuffer.HeadElement()->ToRealTimeSeconds() - NetworkInterpolationDebugCVars::CatchUpGoal();
 			}
 
 			if (CatchUpUntilTime > 0.f)
@@ -124,31 +126,33 @@ struct TInterpolator
 			NewInterpolationTime += Step;
 
 			// Did this put us too far ahead, and now we need to start waiting?
-			if (NewInterpolationTime > SimulationTimeBuffer.GetElementFromHead(0)->ToRealTimeSeconds())
+			if (NewInterpolationTime > SimulationTimeBuffer.HeadElement()->ToRealTimeSeconds())
 			{
-				UE_VLOG(LogOwner, LogNetInterpolation, Log, TEXT("Too far ahead! Starting to wait! Head: %s"), *LexToString(SimulationTimeBuffer.GetElementFromHead(0)->ToRealTimeSeconds()));
-				WaitUntilTime = SimulationTimeBuffer.GetElementFromHead(0)->ToRealTimeSeconds() + NetworkInterpolationDebugCVars::WaitSlack();
-				NewInterpolationTime = SimulationTimeBuffer.GetElementFromHead(0)->ToRealTimeSeconds();
+				UE_VLOG(LogOwner, LogNetInterpolation, Log, TEXT("Too far ahead! Starting to wait! Head: %s"), *LexToString(SimulationTimeBuffer.HeadElement()->ToRealTimeSeconds()));
+				WaitUntilTime = SimulationTimeBuffer.HeadElement()->ToRealTimeSeconds() + NetworkInterpolationDebugCVars::WaitSlack();
+				NewInterpolationTime = SimulationTimeBuffer.HeadElement()->ToRealTimeSeconds();
 			}
 		}
 
 		// Find "To" keyframe
 		const TSyncState* ToState = nullptr;
+		const TAuxState* AuxState = nullptr;
 		FRealTime ToTime = 0.f;
 
-		for (int32 Keyframe = SimulationTimeBuffer.GetTailKeyframe(); Keyframe <= SimulationTimeBuffer.GetHeadKeyframe(); ++Keyframe)
+		for (auto It = SimulationTimeBuffer.CreateConstIterator(); It; ++It)
 		{
-			FSimTime ElementSimTime = *SimulationTimeBuffer.FindElementByKeyframe(Keyframe);
+			FSimTime ElementSimTime = *It.Element();
 			if (NewInterpolationTime <= ElementSimTime.ToRealTimeSeconds())
 			{
-				InterpolationKeyframe = Keyframe;
+				InterpolationKeyframe = It.Keyframe();
 				ToTime = ElementSimTime.ToRealTimeSeconds();
-				ToState = Buffers.Sync.FindElementByKeyframe(Keyframe);
+				ToState = Buffers.Sync[It.Keyframe()];
+				AuxState = Buffers.Aux[It.Keyframe()];
 				break;
 			}
 		}
 
-		if (ensure(ToState))
+		if (ensure(ToState && AuxState))
 		{
 			const FRealTime FromRealTime = InterpolationTime;
 			const FRealTime ToRealTime = ToTime;
@@ -162,7 +166,7 @@ struct TInterpolator
 				TSyncState NewInterpolatedState;
 				TSyncState::Interpolate(InterpolationState, *ToState, InterpolationPCT, NewInterpolatedState);
 
-				Driver->FinalizeFrame(NewInterpolatedState);
+				Driver->FinalizeFrame(NewInterpolatedState, *AuxState);
 				
 				if (bDoVLog)
 				{
@@ -171,29 +175,29 @@ struct TInterpolator
 					// Graph Interpolation Time vs Buffer Head/Tail times
 					const FName ServerSimulationGraphName("ServerSimulationTimeGraph");
 					const FName ServerSimTimeName("Server Simulation Time");
-					FVector2D LocalTimeVsServerSimTime(LogOwner->GetWorld()->GetTimeSeconds(), (int32)TickInfo.SimulationTimeBuffer.GetElementFromHead(0)->ToRealTimeMS());
+					FVector2D LocalTimeVsServerSimTime(LogOwner->GetWorld()->GetTimeSeconds(), (int32)TickInfo.SimulationTimeBuffer.HeadElement()->ToRealTimeMS());
 					UE_VLOG_HISTOGRAM(LogOwner, LogNetInterpolation, Log, ServerSimulationGraphName, ServerSimTimeName, LocalTimeVsServerSimTime);
 
 					const FName BufferTailSimTimeName("Buffer Tail Simulation Time");
-					FVector2D LocalTimeVsBufferTailSim(LogOwner->GetWorld()->GetTimeSeconds(), (int32)TickInfo.SimulationTimeBuffer.GetElementFromTail(0)->ToRealTimeMS());
+					FVector2D LocalTimeVsBufferTailSim(LogOwner->GetWorld()->GetTimeSeconds(), (int32)TickInfo.SimulationTimeBuffer.TailElement()->ToRealTimeMS());
 					UE_VLOG_HISTOGRAM(LogOwner, LogNetInterpolation, Log, ServerSimulationGraphName, BufferTailSimTimeName, LocalTimeVsBufferTailSim);
 
 					const FName LocalInterpolationTimeName("Local Interpolation Time");
 					FVector2D LocalTimeVsInterpolationTime(LogOwner->GetWorld()->GetTimeSeconds(), (int32)(NewInterpolationTime * 1000.f));
 					UE_VLOG_HISTOGRAM(LogOwner, LogNetInterpolation, Log, ServerSimulationGraphName, LocalInterpolationTimeName, LocalTimeVsInterpolationTime);
 					
-					FVector2D LocalTimeVsCatchUpThreshold(LogOwner->GetWorld()->GetTimeSeconds(), (SimulationTimeBuffer.GetElementFromHead(0)->ToRealTimeSeconds() - NetworkInterpolationDebugCVars::CatchUpThreshold()) * 1000.f);
+					FVector2D LocalTimeVsCatchUpThreshold(LogOwner->GetWorld()->GetTimeSeconds(), (SimulationTimeBuffer.HeadElement()->ToRealTimeSeconds() - NetworkInterpolationDebugCVars::CatchUpThreshold()) * 1000.f);
 					UE_VLOG_HISTOGRAM(LogOwner, LogNetInterpolation, Log, ServerSimulationGraphName, "Catch Up Threshold", LocalTimeVsCatchUpThreshold);
 
-					FVector2D LocalTimeVsCatchUpGoal(LogOwner->GetWorld()->GetTimeSeconds(), (SimulationTimeBuffer.GetElementFromHead(0)->ToRealTimeSeconds() - NetworkInterpolationDebugCVars::CatchUpGoal()) * 1000.f);
+					FVector2D LocalTimeVsCatchUpGoal(LogOwner->GetWorld()->GetTimeSeconds(), (SimulationTimeBuffer.HeadElement()->ToRealTimeSeconds() - NetworkInterpolationDebugCVars::CatchUpGoal()) * 1000.f);
 					UE_VLOG_HISTOGRAM(LogOwner, LogNetInterpolation, Log, ServerSimulationGraphName, "Catch Up Goal", LocalTimeVsCatchUpGoal);
 
 					// VLog the actual motion states
-					const TSyncState* DebugTail = Buffers.Sync.GetElementFromTail(0);
-					const TSyncState* DebugHead = Buffers.Sync.GetElementFromHead(0);
+					const TSyncState* DebugTail = Buffers.Sync.TailElement();
+					const TSyncState* DebugHead = Buffers.Sync.HeadElement();
 
-					DebugTail->VisualLog( FVisualLoggingParameters(EVisualLoggingContext::InterpolationBufferTail, Buffers.Sync.GetTailKeyframe(), EVisualLoggingLifetime::Transient), Driver, Driver );
-					DebugHead->VisualLog( FVisualLoggingParameters(EVisualLoggingContext::InterpolationBufferHead, Buffers.Sync.GetHeadKeyframe(), EVisualLoggingLifetime::Transient), Driver, Driver );
+					DebugTail->VisualLog( FVisualLoggingParameters(EVisualLoggingContext::InterpolationBufferTail, Buffers.Sync.TailKeyframe(), EVisualLoggingLifetime::Transient), Driver, Driver );
+					DebugHead->VisualLog( FVisualLoggingParameters(EVisualLoggingContext::InterpolationBufferHead, Buffers.Sync.HeadKeyframe(), EVisualLoggingLifetime::Transient), Driver, Driver );
 
 					InterpolationState.VisualLog( FVisualLoggingParameters(EVisualLoggingContext::InterpolationFrom, InterpolationKeyframe, EVisualLoggingLifetime::Transient), Driver, Driver );
 					ToState->VisualLog( FVisualLoggingParameters(EVisualLoggingContext::InterpolationTo, InterpolationKeyframe, EVisualLoggingLifetime::Transient), Driver, Driver );
diff --git a/Engine/Plugins/Runtime/NetworkPrediction/Source/NetworkPrediction/Public/NetworkSimulationModelReplicators.h b/Engine/Plugins/Runtime/NetworkPrediction/Source/NetworkPrediction/Public/NetworkSimulationModelReplicators.h
--- a/Engine/Plugins/Runtime/NetworkPrediction/Source/NetworkPrediction/Public/NetworkSimulationModelReplicators.h
+++ b/Engine/Plugins/Runtime/NetworkPrediction/Source/NetworkPrediction/Public/NetworkSimulationModelReplicators.h
@@ -107,9 +107,9 @@ struct TReplicatorBase
 	void PostSimTick(TDriver* Driver, const TNetworkSimBufferContainer<TBufferTypes>& Buffers, const TSimulationTickState<TTickSettings>& TickInfo, const FNetSimTickParameters& TickParameters)
 	{
 		// Sync to latest frame if there is any
-		if (Buffers.Sync.GetNumValidElements() > 0)
+		if (Buffers.Sync.Num() > 0)
 		{
-			Driver->FinalizeFrame(*Buffers.Sync.GetElementFromHead(0));
+			Driver->FinalizeFrame(*Buffers.Sync.HeadElement(), *Buffers.Aux.HeadElement());
 		}
 	}
 };
@@ -157,29 +157,27 @@ struct TReplicator_Sequence : public TBase
 
 		auto& Buffer = Buffers.template Get<BufferId>();
 		FArchive& Ar = P.Ar;
-		uint8 SerializedNumElements = FMath::Min<uint8>(MaxNumElements, Buffer.GetNumValidElements());
+		uint8 SerializedNumElements = FMath::Min<uint8>(MaxNumElements, Buffer.Num());
 		Ar << SerializedNumElements;
 
-		const int32 HeadKeyframe = FNetworkSimulationSerialization::SerializeKeyframe(Ar, Buffer.GetHeadKeyframe());
+		const int32 HeadKeyframe = FNetworkSimulationSerialization::SerializeKeyframe(Ar, Buffer.HeadKeyframe());
 		const int32 StartingKeyframe = FMath::Max(0, HeadKeyframe - SerializedNumElements + 1);
 
 		if (Ar.IsLoading())
 		{
-			const int32 PrevHead = Buffer.GetHeadKeyframe();
+			const int32 PrevHead = Buffer.HeadKeyframe();
 			if (PrevHead < StartingKeyframe && PrevHead >= 0)
 			{
 				// There is a gap in the stream. In some cases, we want this to be a "fault" and bubble up. We may want to synthesize state or maybe we just skip ahead.
 				UE_LOG(LogNetworkSim, Warning, TEXT("Fault: gap in received %s buffer. PrevHead: %d. Received: %d-%d. Reseting previous buffer contents"), *LexToString(BufferId), PrevHead, StartingKeyframe, HeadKeyframe);
 			}
-
-			Buffer.ResetNextHeadKeyframe(StartingKeyframe);
 		}
 
 		for (int32 Keyframe = StartingKeyframe; Keyframe <= HeadKeyframe; ++Keyframe)
 		{
 			// This, as is, is bad. The intention is that these functions serialize multiple items in some delta compressed fashion.
 			// As is, we are just serializing the elements individually.
-			auto* Cmd = Ar.IsLoading() ? Buffer.GetWriteNext() : Buffer.FindElementByKeyframe(Keyframe);
+			auto* Cmd = Ar.IsLoading() ? Buffer.WriteKeyframe(Keyframe) : Buffer[Keyframe];
 			Cmd->NetSerialize(P);
 		}
 
@@ -245,13 +243,13 @@ struct TReplicator_Server : public TBase
 		// After receiving input, server may process up to the latest received frames.
 		// (If we needed to buffer input server side for whatever reason, we would do it here)
 		// (also note that we will implicitly guard against speed hacks int he core update loop by not processing cmds past what we have been "allowed")
-		TickInfo.MaxAllowedInputKeyframe = Buffers.Input.GetHeadKeyframe();
+		TickInfo.MaxAllowedInputKeyframe = Buffers.Input.HeadKeyframe();
 
 		// Check for gaps in commands
-		if ( TickInfo.LastProcessedInputKeyframe+1 < Buffers.Input.GetTailKeyframe() )
+		if ( TickInfo.LastProcessedInputKeyframe+1 < Buffers.Input.TailKeyframe() )
 		{
 			UE_LOG(LogNetworkSim, Warning, TEXT("TReplicator_Server::Reconcile missing inputcmds. LastProcessedInputKeyframe: %d. %s"), TickInfo.LastProcessedInputKeyframe, *Buffers.Input.GetBasicDebugStr());
-			TickInfo.LastProcessedInputKeyframe = Buffers.Input.GetTailKeyframe()+1;
+			TickInfo.LastProcessedInputKeyframe = Buffers.Input.TailKeyframe()+1;
 		}
 	}
 
@@ -265,12 +263,12 @@ struct TReplicator_Server : public TBase
 			FNetworkSimTime DeltaSimTime = TickInfo.GetRemainingAllowedSimulationTime();
 			if (DeltaSimTime.IsPositive())
 			{
-				if (TInputCmd* InputCmd = Buffers.Input.GetWriteNext())
+				if (TInputCmd* InputCmd = Buffers.Input.WriteKeyframe(Buffers.Input.HeadKeyframe() + 1))
 				{
 					*InputCmd = TInputCmd();
 					InputCmd->SetFrameDeltaTime(DeltaSimTime);
 					Driver->ProduceInput(DeltaSimTime, *InputCmd);
-					TickInfo.MaxAllowedInputKeyframe = Buffers.Input.GetHeadKeyframe();
+					TickInfo.MaxAllowedInputKeyframe = Buffers.Input.HeadKeyframe();
 				}
 			}
 		}
@@ -328,16 +326,18 @@ struct TReplicator_Simulated : public TBase
 		LastSerializedSimulationTime.NetSerialize(P.Ar);
 
 		// Serialize latest element
-		TSyncState* State = nullptr;
+		TSyncState* SyncState = nullptr;
+		TAuxState* AuxState = nullptr;
 		
 		if (Ar.IsSaving())
 		{
-			State = Buffers.Sync.GetElementFromHead(0);
-			check(State); // We should not be here if the buffer is empty. Want to avoid serializing an "empty" flag at the top here.
+			SyncState = Buffers.Sync.HeadElement();
+			AuxState = Buffers.Aux.HeadElement();
+			check(SyncState && AuxState	); // We should not be here if the buffer is empty. Want to avoid serializing an "empty" flag at the top here.
 		}
 		else
 		{
-			check(TickInfo.SimulationTimeBuffer.GetHeadKeyframe() == Buffers.Sync.GetHeadKeyframe());
+			check(TickInfo.SimulationTimeBuffer.HeadKeyframe() == Buffers.Sync.HeadKeyframe());
 			check(TickInfo.GetTotalProcessedSimulationTime() <= TickInfo.GetTotalAllowedSimulationTime());
 
 			// Cache off our "starting" time before possibly overwriting it. We will use this in Reconcile to catch back up in some cases.
@@ -351,14 +351,14 @@ struct TReplicator_Simulated : public TBase
 			if (LastSerializedSimulationTime > TickInfo.GetTotalProcessedSimulationTime())
 			{
 				// We are getting new state that is ahead of what we have locally, so it can safety go right to head
-				DestinationKeyframe = TickInfo.SimulationTimeBuffer.GetHeadKeyframe()+1;
+				DestinationKeyframe = TickInfo.SimulationTimeBuffer.HeadKeyframe()+1;
 			}
 			else
 			{
 				// We are getting state that is behind what we have locally
-				for (int32 Keyframe = TickInfo.SimulationTimeBuffer.GetHeadKeyframe(); Keyframe >= TickInfo.SimulationTimeBuffer.GetTailKeyframe(); --Keyframe)
+				for (int32 Keyframe = TickInfo.SimulationTimeBuffer.HeadKeyframe(); Keyframe >= TickInfo.SimulationTimeBuffer.TailKeyframe(); --Keyframe)
 				{
-					if (LastSerializedSimulationTime > *TickInfo.SimulationTimeBuffer.FindElementByKeyframe(Keyframe))
+					if (LastSerializedSimulationTime > *TickInfo.SimulationTimeBuffer[Keyframe])
 					{
 						DestinationKeyframe = Keyframe+1;
 						break;
@@ -367,21 +367,22 @@ struct TReplicator_Simulated : public TBase
 
 				if (DestinationKeyframe == INDEX_NONE)
 				{
-					FNetworkSimTime TotalTimeAhead = *TickInfo.SimulationTimeBuffer.GetElementFromHead(0) - LastSerializedSimulationTime;
+					FNetworkSimTime TotalTimeAhead = TickInfo.GetTotalProcessedSimulationTime() - LastSerializedSimulationTime;
 					FNetworkSimTime SerializeDelta = LastSerializedSimulationTime - PrevLastSerializedSimulationTime;
 
 					// We are way far ahead of the server... we will need to clear out sync buffers, take what they gave us, and catch up in reconcile
 					//UE_LOG(LogNetworkSim, Warning, TEXT("!!! TReplicator_Simulated. Large gap detected. SerializedTime: %s. Buffer time: [%s-%s]. %d Elements. DeltaFromHead: %s. DeltaSerialize: %s"), *LastSerializedSimulationTime.ToString(), 
 					//	*TickInfo.SimulationTimeBuffer.GetElementFromTail(0)->ToString(), *TickInfo.SimulationTimeBuffer.GetElementFromHead(0)->ToString(), TickInfo.SimulationTimeBuffer.GetNumValidElements(), *TotalTimeAhead.ToString(), *SerializeDelta.ToString());
-					DestinationKeyframe = TickInfo.SimulationTimeBuffer.GetHeadKeyframe()+2; // (Skip ahead 2 to force a break in continuity)
+					DestinationKeyframe = TickInfo.SimulationTimeBuffer.HeadKeyframe()+2; // (Skip ahead 2 to force a break in continuity)
 				}
 			}
 
 			check(DestinationKeyframe != INDEX_NONE);
 
-			// "Finalize" our buffers and time keeping such that we serialize the latest state from the server in the right spot			
-			Buffers.Sync.ResetNextHeadKeyframe(DestinationKeyframe);
-			State = Buffers.Sync.GetWriteNext();
+			// "Finalize" our buffers and time keeping such that we serialize the latest state from the server in the right spot
+			SyncState = Buffers.Sync.WriteKeyframe(DestinationKeyframe);
+			AuxState = Buffers.Aux.WriteKeyframe(DestinationKeyframe);
+			Buffers.Input.WriteKeyframe(DestinationKeyframe);
 
 			// Update tick info
 			TickInfo.SetTotalProcessedSimulationTime(LastSerializedSimulationTime, DestinationKeyframe);
@@ -393,17 +394,17 @@ struct TReplicator_Simulated : public TBase
 			check(TickInfo.GetTotalProcessedSimulationTime() <= TickInfo.GetTotalAllowedSimulationTime());
 
 			TickInfo.LastProcessedInputKeyframe = DestinationKeyframe;
-			TickInfo.MaxAllowedInputKeyframe = DestinationKeyframe;
-
-			Buffers.Input.ResetNextHeadKeyframe(DestinationKeyframe+1);
+			TickInfo.MaxAllowedInputKeyframe = DestinationKeyframe;	
 		}
 
-		check(State);
-		State->NetSerialize(Ar);
+		check(SyncState && AuxState);
+		SyncState->NetSerialize(Ar);
+		AuxState->NetSerialize(Ar);
 
 		if (Ar.IsLoading())
 		{
-			LastSerializedSyncState = *State;
+			LastSerializedSyncState = *SyncState;
+			LastSerializedAuxState = *AuxState;
 		}
 	}
 
@@ -419,13 +420,13 @@ struct TReplicator_Simulated : public TBase
 
 		if (bAllowSimulatedExtrapolation && ParentSimulation == nullptr && NetworkSimulationModelCVars::EnableSimulatedExtrapolation() && NetworkSimulationModelCVars::EnableSimulatedReconcile())
 		{
-			TInputCmd* LastCmd = Buffers.Input.GetElementFromHead(0);
+			TInputCmd* LastCmd = Buffers.Input.HeadElement();
 
 			// Simulated Reconcile requires the input buffer to be kept up to date with the Sync buffer
 			// Generate a new, fake, command since we just added a new sync state to head
-			while (Buffers.Input.GetHeadKeyframe() < Buffers.Sync.GetHeadKeyframe())
+			while (Buffers.Input.HeadKeyframe() < Buffers.Sync.HeadKeyframe())
 			{
-				TInputCmd* Next = Buffers.Input.GetWriteNext();
+				TInputCmd* Next = Buffers.Input.WriteKeyframe(Buffers.Input.HeadKeyframe()+1);
 				*Next = LastCmd ? *LastCmd : TInputCmd();
 			}
 
@@ -458,14 +459,14 @@ struct TReplicator_Simulated : public TBase
 				FNetworkSimTime DeltaSimTime = TickInfo.GetRemainingAllowedSimulationTime();
 				if (DeltaSimTime.IsPositive())
 				{
-					if (TInputCmd* InputCmd = Buffers.Input.GetWriteNext())
+					if (TInputCmd* InputCmd = Buffers.Input.WriteKeyframe(Buffers.Input.HeadKeyframe() + 1))
 					{
 						//UE_LOG(LogNetworkSim, Warning, TEXT("   Extrapolating %s"), *DeltaSimTime.ToString());
 
 						*InputCmd = TInputCmd();
 						InputCmd->SetFrameDeltaTime(DeltaSimTime);
 						Driver->ProduceInput(DeltaSimTime, *InputCmd);
-						TickInfo.MaxAllowedInputKeyframe = Buffers.Input.GetHeadKeyframe();
+						TickInfo.MaxAllowedInputKeyframe = Buffers.Input.HeadKeyframe();
 					}
 				}
 			}
@@ -478,9 +479,9 @@ struct TReplicator_Simulated : public TBase
 		if (bAllowSimulatedExtrapolation || ParentSimulation)
 		{
 			// Sync to latest frame if there is any
-			if (Buffers.Sync.GetNumValidElements() > 0)
+			if (Buffers.Sync.Num() > 0)
 			{
-				Driver->FinalizeFrame(*Buffers.Sync.GetElementFromHead(0));
+				Driver->FinalizeFrame(*Buffers.Sync.HeadElement(), *Buffers.Aux.HeadElement());
 			}
 			
 		}
@@ -495,19 +496,17 @@ struct TReplicator_Simulated : public TBase
 	{
 		// For now, we make the assumption that our last serialized state and time match the parent simulation.
 		// This would not always be the case with low frequency simulated proxies. But could be handled by replicating the simulations together (at the replication level)
-		int32 NewHeadKeyframe = Buffers.Sync.GetHeadKeyframe()+1;
-		Buffers.Sync.ResetNextHeadKeyframe(NewHeadKeyframe);
-		Buffers.Input.ResetNextHeadKeyframe(NewHeadKeyframe);
+		int32 NewHeadKeyframe = Buffers.Sync.HeadKeyframe()+1;
+		
 		TickInfo.SetTotalProcessedSimulationTime(LastSerializedSimulationTime, NewHeadKeyframe);
 		TickInfo.SetTotalAllowedSimulationTime(LastSerializedSimulationTime);
+		TickInfo.LastProcessedInputKeyframe = NewHeadKeyframe;
+		TickInfo.MaxAllowedInputKeyframe = NewHeadKeyframe;
 
-		*Buffers.Sync.GetWriteNext() = LastSerializedSyncState;
-		*Buffers.Input.GetWriteNext() = TInputCmd();
+		*Buffers.Sync.WriteKeyframe(NewHeadKeyframe) = LastSerializedSyncState;
+		*Buffers.Input.WriteKeyframe(NewHeadKeyframe) = TInputCmd();
 
-		TickInfo.LastProcessedInputKeyframe = Buffers.Input.GetHeadKeyframe();
-		TickInfo.MaxAllowedInputKeyframe = Buffers.Input.GetHeadKeyframe();
-
-		Driver->FinalizeFrame(LastSerializedSyncState);
+		Driver->FinalizeFrame(LastSerializedSyncState, LastSerializedAuxState);
 
 		LastSerializedSyncState.VisualLog( FVisualLoggingParameters(EVisualLoggingContext::FirstMispredicted, ParentKeyframe, EVisualLoggingLifetime::Persistent), Driver, Driver );
 	}
@@ -519,7 +518,7 @@ struct TReplicator_Simulated : public TBase
 
 		SimulationExtrapolation<T, TDriver>(Driver, Buffers, TickInfo, StepTime);
 
-		TSyncState* SyncState = Buffers.Sync.GetElementFromHead(0);
+		TSyncState* SyncState = Buffers.Sync.HeadElement();
 		check(SyncState);
 		SyncState->VisualLog( FVisualLoggingParameters(bFinalStep ? EVisualLoggingContext::LastMispredicted : EVisualLoggingContext::OtherMispredicted, ParentKeyframe, EVisualLoggingLifetime::Persistent), Driver, Driver );
 	}
@@ -529,32 +528,38 @@ private:
 	template<typename T, typename TDriver>
 	void SimulationExtrapolation(TDriver* Driver, TNetworkSimBufferContainer<TBufferTypes>& Buffers, TSimulationTickState<TTickSettings>& TickInfo, const FNetworkSimTime DeltaSimTime)
 	{
-		TInputCmd* LastCmd = Buffers.Input.GetElementFromHead(0);
-
 		// We have extrapolated ahead of the server. The latest network update is now "in the past" from what we rendered last frame.
 		// We will insert a new keyframe to make up the difference from the last known state to where we want to be in the now.
-				
-		TInputCmd* NewCmd = Buffers.Input.GetWriteNext();
-		*NewCmd = LastCmd ? *LastCmd : TInputCmd();
-		NewCmd->SetFrameDeltaTime(DeltaSimTime);
 
-		TSyncState* PrevSyncState = Buffers.Sync.GetElementFromHead(0);
-		TSyncState* NextSyncState = Buffers.Sync.GetWriteNext();
+		ensure(Buffers.Input.HeadKeyframe() == Buffers.Sync.HeadKeyframe());
 
-		TAuxState Junk;
+		TInputCmd* LastCmd = Buffers.Input.HeadElement();
+		const int32 LastKeyframe = Buffers.Input.HeadKeyframe();
+		const int32 Keyframe = LastKeyframe + 1;
+
+		// Create fake cmd				
+		TInputCmd* NewCmd = Buffers.Input.WriteKeyframe(Keyframe);
+		*NewCmd = LastCmd ? *LastCmd : TInputCmd();
+		NewCmd->SetFrameDeltaTime(DeltaSimTime);	
+		
+		// Create new sync state to write to
+		TSyncState* PrevSyncState = Buffers.Sync.HeadElement();
+		TSyncState* NextSyncState = Buffers.Sync.WriteKeyframe(Keyframe);
+		TAuxState* AuxState = Buffers.Aux.HeadElement();
 
 		// Do the actual update
-		T::Update(Driver, NewCmd->GetFrameDeltaTime().ToRealTimeSeconds(), *NewCmd, *PrevSyncState, *NextSyncState, Junk);
-		TickInfo.IncrementTotalProcessedSimulationTime(NewCmd->GetFrameDeltaTime(), Buffers.Sync.GetHeadKeyframe());
+		T::Update(Driver, NewCmd->GetFrameDeltaTime().ToRealTimeSeconds(), *NewCmd, *PrevSyncState, *NextSyncState, *AuxState, Buffers.Aux.WriteKeyframeFunc(Keyframe+1));
+		TickInfo.IncrementTotalProcessedSimulationTime(NewCmd->GetFrameDeltaTime(), Buffers.Sync.HeadKeyframe());
 
 		// Set our LastProcessedInputKeyframe to fake that we handled it
-		TickInfo.LastProcessedInputKeyframe = Buffers.Input.GetHeadKeyframe();
-		TickInfo.MaxAllowedInputKeyframe = Buffers.Input.GetHeadKeyframe();
+		TickInfo.LastProcessedInputKeyframe = Keyframe;
+		TickInfo.MaxAllowedInputKeyframe = Keyframe;
 	}
 	
 	FNetworkSimTime ReconcileSimulationTime;
 	FNetworkSimTime LastSerializedSimulationTime;
 	TSyncState LastSerializedSyncState;
+	TAuxState LastSerializedAuxState;	// Temp? This should be conditional or optional. We want to support not replicating the aux state to simulated proxies
 };
 
 
@@ -566,7 +571,7 @@ struct TReplicator_Autonomous : public TBase
 	using TSyncState = typename TBufferTypes::TSyncState;
 	using TAuxState = typename TBufferTypes::TAuxState;
 
-	int32 GetLastSerializedKeyframe() const { return LastSerializedKeyframe; }
+	int32 GetLastSerializedKeyframe() const { return SerializedKeyframe; }
 	bool IsReconcileFaultDetected() const { return bReconcileFaultDetected; }
 	const FNetworkSimTime& GetLastSerializedSimTime() const { return SerializedTime; }
 
@@ -575,7 +580,7 @@ struct TReplicator_Autonomous : public TBase
 
 	int32 GetProxyDirtyCount(TNetworkSimBufferContainer<TBufferTypes>& Buffers) const
 	{
-		return Buffers.Sync.GetDirtyCount() ^ (TBase::GetProxyDirtyCount(Buffers) << 2); ;
+		return (Buffers.Sync.GetDirtyCount()) ^ (Buffers.Aux.GetDirtyCount() << 1) ^ (TBase::GetProxyDirtyCount(Buffers) << 2); 
 	}
 
 	// --------------------------------------------------------------------
@@ -585,51 +590,42 @@ struct TReplicator_Autonomous : public TBase
 	{
 		FArchive& Ar = P.Ar;
 
-		const int32 SerializedHeadKeyframe = FNetworkSimulationSerialization::SerializeKeyframe(Ar, Buffers.Sync.GetHeadKeyframe());
-		TSyncState* SerializedState = nullptr;
-
-		// Serialize total simulation time. This isn't really necessary since we have the keyframe above. 
+		SerializedKeyframe = FNetworkSimulationSerialization::SerializeKeyframe(Ar, Buffers.Sync.HeadKeyframe());
+		
 		SerializedTime = TickInfo.GetTotalProcessedSimulationTime();
 		SerializedTime.NetSerialize(P.Ar);
 
 		if (Ar.IsSaving())
 		{
 			// Server serialize the latest state
-			SerializedState = Buffers.Sync.GetElementFromHead(0);
+			Buffers.Sync.HeadElement()->NetSerialize(Ar);
+			Buffers.Aux.HeadElement()->NetSerialize(Ar);
 		}
 		else
 		{
-			// Lazy init the reconciliation buffer. We don't need this on the server/writing side. (Fixme, with templated approach we could easily specialize this to an inline element instead of buffer?)
-			if (ReconciliationBuffer.GetMaxNumElements() == 0)
-			{
-				ReconciliationBuffer.SetBufferSize(1);
-			}
-
-			// Reset the head to match the server. This is not critical, we just need a place to store this, but it makes the code in the reconciliation a bit nicer since the reconcilated frame number = head frame.
-			ReconciliationBuffer.ResetNextHeadKeyframe(SerializedHeadKeyframe);
-			SerializedState = ReconciliationBuffer.GetWriteNext();
+			SerializedSyncState.NetSerialize(Ar);
+			SerializedAuxState.NetSerialize(Ar);
 		}
 
-		SerializedState->NetSerialize(Ar);
-
 		if (Ar.IsLoading())
 		{
 			bReconcileFaultDetected = false;
 			bPendingReconciliation = false;
 
 			// The state the client predicted that corresponds to the state the server just serialized to us
-			if (TSyncState* ClientExistingState = Buffers.Sync.FindElementByKeyframe(SerializedHeadKeyframe))
+			if (TSyncState* ClientExistingState = Buffers.Sync[SerializedKeyframe])
 			{
-				if (ClientExistingState->ShouldReconcile(*SerializedState) || (NetworkSimulationModelCVars::ForceReconcile() > 0) || (NetworkSimulationModelCVars::ForceReconcileSingle() > 0))
+				// TODO: AuxState->ShouldReconcile
+				if (ClientExistingState->ShouldReconcile(SerializedSyncState) || (NetworkSimulationModelCVars::ForceReconcile() > 0) || (NetworkSimulationModelCVars::ForceReconcileSingle() > 0))
 				{
 					NetworkSimulationModelCVars::SetForceReconcileSingle(0);
-					UE_CLOG(!Buffers.Input.IsValidKeyframe(SerializedHeadKeyframe), LogNetworkSim, Error, TEXT("::NetSerialize: Client InputBuffer does not contain data for frame %d. {%s} {%s}"), SerializedHeadKeyframe, *Buffers.Input.GetBasicDebugStr(), *Buffers.Sync.GetBasicDebugStr());
+					UE_CLOG(!Buffers.Input.IsValidKeyframe(SerializedKeyframe), LogNetworkSim, Error, TEXT("::NetSerialize: Client InputBuffer does not contain data for frame %d. {%s} {%s}"), SerializedKeyframe, *Buffers.Input.GetBasicDebugStr(), *Buffers.Sync.GetBasicDebugStr());
 					bPendingReconciliation =  true;
 				}
 			}
 			else
 			{
-				if (SerializedHeadKeyframe < Buffers.Sync.GetTailKeyframe())
+				if (SerializedKeyframe < Buffers.Sync.TailKeyframe())
 				{
 					// Case 1: the serialized state is older than what we've kept in our buffer. A bigger buffer would solve this! (at the price of more resimulated frames to recover when this happens)
 					// This is a reconcile fault and we just need to chill. We'll stop sending user commands until the cmds in flight flush through the system and we catch back up.
@@ -642,8 +638,6 @@ struct TReplicator_Autonomous : public TBase
 				}
 			}
 		}
-
-		LastSerializedKeyframe = SerializedHeadKeyframe;
 	}
 
 	// --------------------------------------------------------------------
@@ -658,30 +652,21 @@ struct TReplicator_Autonomous : public TBase
 		}
 		bPendingReconciliation = false;
 		bDependentSimulationNeedsReconcile = false;
-		check(ReconciliationBuffer.GetNumValidElements() == 1); // this function assumes only 1 element
 
-		const int32 ReconciliationKeyframe = ReconciliationBuffer.GetHeadKeyframe();
-		TSyncState* ServerState = ReconciliationBuffer.GetElementFromHead(0);
-
-		TInputCmd* ClientInputCmd = Buffers.Input.FindElementByKeyframe( ReconciliationKeyframe );
+		TInputCmd* ClientInputCmd = Buffers.Input[SerializedKeyframe];
 		if (ClientInputCmd == nullptr)
 		{
 			// Fault: no historic command for this frame to reconcile with.
-			UE_LOG(LogNetworkSim, Error, TEXT("Client InputBuffer does not contain data for frame %d. {%s} {%s}"), ReconciliationKeyframe, *Buffers.Input.GetBasicDebugStr(), *Buffers.Sync.GetBasicDebugStr());
+			UE_LOG(LogNetworkSim, Error, TEXT("Client InputBuffer does not contain data for frame %d. {%s} {%s}"), SerializedKeyframe, *Buffers.Input.GetBasicDebugStr(), *Buffers.Sync.GetBasicDebugStr());
 			return;
 		}
 
-		// -------------------------------------------------------------------------------------------------------------------------
-		// Resimulate
-		// -------------------------------------------------------------------------------------------------------------------------
-
-		TSyncState* ClientSyncState = Buffers.Sync.FindElementByKeyframe( ReconciliationKeyframe );
-
+		TSyncState* ClientSyncState = Buffers.Sync[SerializedKeyframe];
 		const bool bDoVisualLog = NetworkSimulationModelCVars::EnableLocalPrediction() > 0; // don't visual log if we have prediction disabled
 		
 		if (bDoVisualLog)
 		{
-			ServerState->VisualLog( FVisualLoggingParameters(EVisualLoggingContext::LastConfirmed, ReconciliationKeyframe, EVisualLoggingLifetime::Persistent), Driver, Driver);
+			SerializedSyncState.VisualLog( FVisualLoggingParameters(EVisualLoggingContext::LastConfirmed, SerializedKeyframe, EVisualLoggingLifetime::Persistent), Driver, Driver);
 		}
 
 		if (ClientSyncState)
@@ -689,25 +674,27 @@ struct TReplicator_Autonomous : public TBase
 			// Existing ClientSyncState, log it before overwriting it
 			if (bDoVisualLog)
 			{
-				ClientSyncState->VisualLog( FVisualLoggingParameters(EVisualLoggingContext::FirstMispredicted, ReconciliationKeyframe, EVisualLoggingLifetime::Persistent), Driver, Driver);
+				ClientSyncState->VisualLog( FVisualLoggingParameters(EVisualLoggingContext::FirstMispredicted, SerializedKeyframe, EVisualLoggingLifetime::Persistent), Driver, Driver);
 			}
 		}
 		else
 		{
 			// No existing state, so create add it explicitly
-			Buffers.Sync.ResetNextHeadKeyframe( ReconciliationKeyframe );
-			ClientSyncState = Buffers.Sync.GetWriteNext();
+			ClientSyncState = Buffers.Sync.WriteKeyframe( SerializedKeyframe );
 		}
 
 		// Set client's sync state to the server version
 		check(ClientSyncState);
-		*ClientSyncState = *ServerState;
+		*ClientSyncState = SerializedSyncState;
+
+		// Set Client's aux state to the server version
+		*Buffers.Aux.WriteKeyframe(SerializedKeyframe) = SerializedAuxState;
 
 		const FNetworkSimTime RollbackDeltaTime = SerializedTime - TickInfo.GetTotalProcessedSimulationTime();
 
 		// Set the canonical simulation time to what we received (we will advance it as we resimulate)
-		TickInfo.SetTotalProcessedSimulationTime(SerializedTime, ReconciliationKeyframe);
-		TickInfo.LastProcessedInputKeyframe = ReconciliationKeyframe;
+		TickInfo.SetTotalProcessedSimulationTime(SerializedTime, SerializedKeyframe);
+		TickInfo.LastProcessedInputKeyframe = SerializedKeyframe;
 		TickInfo.MaxAllowedInputKeyframe = FMath::Max(TickInfo.MaxAllowedInputKeyframe, TickInfo.LastProcessedInputKeyframe); // Make sure this doesn't lag behind. This is the only place we should need to do this.
 
 		if (NetworkSimulationModelCVars::EnableLocalPrediction() == 0)
@@ -719,49 +706,37 @@ struct TReplicator_Autonomous : public TBase
 		// Tell dependent simulations to rollback
 		for (INetworkSimulationModel* DependentSim : DependentSimulations)
 		{
-			DependentSim->BeginRollback(RollbackDeltaTime, ReconciliationKeyframe);
+			DependentSim->BeginRollback(RollbackDeltaTime, SerializedKeyframe);
 		}
 		
 		// Resimulate all user commands 
 		const int32 LastKeyframeToProcess = TickInfo.MaxAllowedInputKeyframe;
-		for (int32 Keyframe = ReconciliationKeyframe+1; Keyframe <= LastKeyframeToProcess; ++Keyframe)
+		for (int32 Keyframe = SerializedKeyframe+1; Keyframe <= LastKeyframeToProcess; ++Keyframe)
 		{
 			// Keyframe is the frame we are resimulating right now.
-			TInputCmd* ResimulateCmd	= Buffers.Input.FindElementByKeyframe(Keyframe);
-			TSyncState* PrevMotionState = Buffers.Sync.FindElementByKeyframe(Keyframe - 1);
-			TSyncState* NextMotionState = Buffers.Sync.FindElementByKeyframe(Keyframe);
+			TInputCmd* ResimulateCmd  = Buffers.Input[Keyframe];
+			TSyncState* PrevSyncState = Buffers.Sync[Keyframe - 1];
+			TSyncState* NextSyncState = Buffers.Sync.WriteKeyframe(Keyframe);
 			
 			check(ResimulateCmd);
-			check(PrevMotionState);
-			if (NextMotionState == nullptr)
-			{
-				// This should only happen if we transition from no local prediction to local prediction, since we will traverse through out input buffer but not have predicted sync states to look at.
-				NextMotionState = Buffers.Sync.GetWriteNext();
-				check(Keyframe == Buffers.Sync.GetHeadKeyframe());
-			}
-
-			// TEMP (Aux buffer not fully plumbed through the system yet)
-			// ------------------------------------------------------
-			TAuxState TempAuxState;
-			TAuxState* AuxState = Buffers.Aux.FindElementByKeyframe(Keyframe);
-			if (AuxState == nullptr)
-			{
-				AuxState = &TempAuxState;
-			}
-			// ------------------------------------------------------
+			check(PrevSyncState);
+			check(NextSyncState);
+			
+			TAuxState* AuxState = Buffers.Aux[Keyframe];
+			check(AuxState);
 
 			// Log out the Mispredicted state that we are about to overwrite.
-			NextMotionState->VisualLog( FVisualLoggingParameters(Keyframe == LastKeyframeToProcess ? EVisualLoggingContext::LastMispredicted : EVisualLoggingContext::OtherMispredicted, Keyframe, EVisualLoggingLifetime::Persistent), Driver, Driver);
+			NextSyncState->VisualLog( FVisualLoggingParameters(Keyframe == LastKeyframeToProcess ? EVisualLoggingContext::LastMispredicted : EVisualLoggingContext::OtherMispredicted, Keyframe, EVisualLoggingLifetime::Persistent), Driver, Driver);
 
 			// Do the actual update
-			T::Update(Driver, ResimulateCmd->GetFrameDeltaTime().ToRealTimeSeconds(), *ResimulateCmd, *PrevMotionState, *NextMotionState, *AuxState);
+			T::Update(Driver, ResimulateCmd->GetFrameDeltaTime().ToRealTimeSeconds(), *ResimulateCmd, *PrevSyncState, *NextSyncState, *AuxState, Buffers.Aux.WriteKeyframeFunc(Keyframe+1));
 			
 			// Update TickInfo
 			TickInfo.IncrementTotalProcessedSimulationTime(ResimulateCmd->GetFrameDeltaTime(), Keyframe);
 			TickInfo.LastProcessedInputKeyframe = Keyframe;
 
 			// Log out the newly predicted state that we got.
-			NextMotionState->VisualLog( FVisualLoggingParameters(Keyframe == LastKeyframeToProcess ? EVisualLoggingContext::LastPredicted : EVisualLoggingContext::OtherPredicted, Keyframe, EVisualLoggingLifetime::Persistent), Driver, Driver);
+			NextSyncState->VisualLog( FVisualLoggingParameters(Keyframe == LastKeyframeToProcess ? EVisualLoggingContext::LastPredicted : EVisualLoggingContext::OtherPredicted, Keyframe, EVisualLoggingLifetime::Persistent), Driver, Driver);
 
 			// Tell dependent simulations to advance
 			for (INetworkSimulationModel* DependentSim : DependentSimulations)
@@ -792,12 +767,12 @@ struct TReplicator_Autonomous : public TBase
 				const FNetworkSimTime DeltaSimTime = TickInfo.GetRemainingAllowedSimulationTime();
 				if (DeltaSimTime.IsPositive())
 				{
-					if (TInputCmd* InputCmd = Buffers.Input.GetWriteNext())
+					if (TInputCmd* InputCmd = Buffers.Input.WriteKeyframe(Buffers.Input.HeadKeyframe() + 1))
 					{
 						*InputCmd = TInputCmd();
 						InputCmd->SetFrameDeltaTime(DeltaSimTime);
 						Driver->ProduceInput(DeltaSimTime, *InputCmd);
-						TickInfo.MaxAllowedInputKeyframe = Buffers.Input.GetHeadKeyframe(); // Allow the new command to be processed by the local simulation
+						TickInfo.MaxAllowedInputKeyframe = Buffers.Input.HeadKeyframe(); // Allow the new command to be processed by the local simulation
 					}
 				}
 			}
@@ -811,7 +786,7 @@ struct TReplicator_Autonomous : public TBase
 				NonPredictedInputTimeAccumulator.Accumulate(NonPredictedInputTime, TickParameters.LocalDeltaTimeSeconds);
 				if (NonPredictedInputTime.IsPositive())
 				{
-					if (TInputCmd* InputCmd = Buffers.Input.GetWriteNext())
+					if (TInputCmd* InputCmd = Buffers.Input.WriteKeyframe(Buffers.Input.HeadKeyframe() + 1))
 					{
 						*InputCmd = TInputCmd();
 						InputCmd->SetFrameDeltaTime(NonPredictedInputTime);
@@ -824,13 +799,15 @@ struct TReplicator_Autonomous : public TBase
 
 private:
 	
-	TReplicationBuffer<TSyncState> ReconciliationBuffer;
-	FNetworkSimTime SerializedTime; // last serialized time keeper
-	TRealTimeAccumulator<TTickSettings> NonPredictedInputTimeAccumulator; // for tracking input time in the non predictive case
+	TSyncState SerializedSyncState;
+	TAuxState SerializedAuxState;
+	FNetworkSimTime SerializedTime;
+	int32 SerializedKeyframe = -1;
 
-	int32 LastSerializedKeyframe = -1;
 	bool bPendingReconciliation = false;	// Reconciliation is pending: we need to reconcile state from the server that differs from the locally predicted state
 	bool bReconcileFaultDetected = false;	// A fault was detected: we received state from the server that we are unable to reconcile with locally predicted state
+
+	TRealTimeAccumulator<TTickSettings> NonPredictedInputTimeAccumulator; // for tracking input time in the non predictive case
 };
 
 /** Special replicator for debug buffer, this preserves the local buffer and receives into a replicator-owned buffer (we want these buffers to be distinct/not merged) */
@@ -838,6 +815,7 @@ template<typename TBufferTypes, typename TTickSettings, int32 MaxNumElements=5,
 struct TReplicator_Debug : public TBase
 {
 	using TDebugState = typename TBufferTypes::TDebugState;
+	using TDebugBuffer = typename TNetworkSimBufferContainer<TBufferTypes>::TDebugBuffer;
 
 	int32 GetProxyDirtyCount(TNetworkSimBufferContainer<TBufferTypes>& Buffers) const 
 	{
@@ -849,40 +827,32 @@ struct TReplicator_Debug : public TBase
 		TBase::NetSerialize(P, Buffers, TickInfo);
 		FArchive& Ar = P.Ar;
 
-		TReplicationBuffer<TDebugState>& Buffer = Ar.IsSaving() ? Buffers.Debug : ReceivedBuffer;
+		TDebugBuffer& Buffer = Ar.IsSaving() ? Buffers.Debug : ReceivedBuffer;
 		
-		uint8 SerializedNumElements = FMath::Min<uint8>(MaxNumElements, Buffer.GetNumValidElements());
+		uint8 SerializedNumElements = FMath::Min<uint8>(MaxNumElements, Buffer.Num());
 		Ar << SerializedNumElements;
 
-		const int32 HeadKeyframe = FNetworkSimulationSerialization::SerializeKeyframe(Ar, Buffer.GetHeadKeyframe());
+		const int32 HeadKeyframe = FNetworkSimulationSerialization::SerializeKeyframe(Ar, Buffer.HeadKeyframe());
 		const int32 StartingKeyframe = FMath::Max(0, HeadKeyframe - SerializedNumElements + 1);
 
 		if (Ar.IsLoading())
 		{
-			// Lazy init on recieve
-			if (ReceivedBuffer.GetMaxNumElements() != Buffers.Debug.GetMaxNumElements())
-			{
-				ReceivedBuffer.SetBufferSize(Buffers.Debug.GetMaxNumElements());
-			}
-
-			const int32 PrevHead = Buffer.GetHeadKeyframe();
+			const int32 PrevHead = Buffer.HeadKeyframe();
 			if (PrevHead < StartingKeyframe && PrevHead >= 0)
 			{
 				// There is a gap in the stream. In some cases, we want this to be a "fault" and bubble up. We may want to synthesize state or maybe we just skip ahead.
 				UE_LOG(LogNetworkSim, Warning, TEXT("Fault: gap in received Debug buffer. PrevHead: %d. Received: %d-%d. Reseting previous buffer contents"), PrevHead, StartingKeyframe, HeadKeyframe);
 			}
-
-			Buffer.ResetNextHeadKeyframe(StartingKeyframe);
 		}
 
 		for (int32 Keyframe = StartingKeyframe; Keyframe <= HeadKeyframe; ++Keyframe)
 		{
 			// This, as is, is bad. The intention is that these functions serialize multiple items in some delta compressed fashion.
 			// As is, we are just serializing the elements individually.
-			auto* Cmd = Ar.IsLoading() ? Buffer.GetWriteNext() : Buffer.FindElementByKeyframe(Keyframe);
+			auto* Cmd = Ar.IsLoading() ? Buffer.WriteKeyframe(Keyframe) : Buffer[Keyframe];
 			Cmd->NetSerialize(P);
 		}
 	}
 
-	TReplicationBuffer<TDebugState> ReceivedBuffer;
+	TDebugBuffer ReceivedBuffer;
 };
\ No newline at end of file
diff --git a/Engine/Plugins/Runtime/NetworkPrediction/Source/NetworkPrediction/Public/NetworkSimulationModelTypes.h b/Engine/Plugins/Runtime/NetworkPrediction/Source/NetworkPrediction/Public/NetworkSimulationModelTypes.h
--- a/Engine/Plugins/Runtime/NetworkPrediction/Source/NetworkPrediction/Public/NetworkSimulationModelTypes.h
+++ b/Engine/Plugins/Runtime/NetworkPrediction/Source/NetworkPrediction/Public/NetworkSimulationModelTypes.h
@@ -60,6 +60,19 @@ struct TSelectTypeHelper<TBufferTypes, Debug>
 	using type = typename TBufferTypes::TDebugState;
 };
 
+enum ENetworkSimBufferAllocationType
+{
+	Contiguous,
+	Sparse
+};
+
+template<ENetworkSimBufferAllocationType InAllocationType, int32 InSize>
+struct TNetworkSimBufferAllocation
+{
+	enum { Size = InSize };
+	static constexpr ENetworkSimBufferAllocationType AllocationType() { return InAllocationType; }
+};
+
 // A collection of the system's buffer types. This allows us to collapse the 4 types into a single type to use a template argument elsewhere.
 template<typename InInputCmd, typename InSyncState, typename InAuxState, typename InDebugState = FNetSimProcessedFrameDebugInfo>
 struct TNetworkSimBufferTypes
@@ -127,44 +140,166 @@ struct TBufferGetterHelper<TContainer, ENetworkSimBufferTypeId::Debug>
 	}
 };
 
+// -------------------------------------------------------------------
+
+template<typename InElementType>
+struct TBaseStateAccessor
+{
+	using ElementType = InElementType;
+	virtual ~TBaseStateAccessor() { }
+	virtual const ElementType* Get() const = 0;
+	virtual void Modify(TUniqueFunction<void(ElementType&)>&& Func, int32 Keyframe) = 0;
+};
+
+// Stored modifications to a NetSim Buffer. Needed for replaying predictive mutations that happen outside of the core Update function
+template<typename TBuffer>
+struct TPredictedModBuffer : public TBaseStateAccessor<typename TBuffer::ElementType>
+{
+	using ElementType = typename TBuffer::ElementType;
+
+	TPredictedModBuffer(TBuffer& InBuffer)
+		: Buffer(InBuffer) { }
+
+	const ElementType* Get() const final override
+	{
+		return Buffer.HeadElement();
+	}
+
+	void Modify(TUniqueFunction<void(ElementType&)>&& Func, int32 Keyframe) final override
+	{
+		ElementType* State = Buffer[Keyframe];
+		check(State);
+		Func(*State);
+
+		if (!bAuthority)
+		{
+			PredictedMods.FindOrAdd(Keyframe).Add(MoveTemp(Func));
+		}
+	}
+
+	// -------------------------------------------------------
+
+	void ApplySavedMods(int32 Keyframe)
+	{
+
+	}
+
+	void TrimSavedMods(int32 Keyframe)
+	{
+
+	}
+
+	void SetIsAuthority(bool bInAuthority)
+	{
+		bAuthority = bInAuthority;
+	}
+
+private:
+
+	TSortedMap<int32, TArray<TUniqueFunction<void(ElementType&)>>> PredictedMods;
+	TBuffer& Buffer;
+	bool bAuthority = false; 
+};
+
+// Accessor for predicted state. This is what GameCode outside of the core ::Update function should use to mutate state in the Sync or Aux buffer/
+// This builds off TPredictedModBuffer: it stores a pointer to the simulation's current frame so that game code doesn't need to supply it.
 template<typename T>
+struct TPredictedStateAccessor
+{
+	TPredictedStateAccessor(TBaseStateAccessor<T>& InModBuffer, int32* InLastProcessedKeyframe)
+		: ModBuffer(InModBuffer), LastProcessedKeyframe(InLastProcessedKeyframe) { }
+
+	// Returns the current state, cannot be modified directly
+	const T* Get() const
+	{
+		return ModBuffer.Get();
+	}
+
+	// Applies a mod to the current state. If in prediction mode, will save the mod to be reapplied during a rollback
+	void Modify(TUniqueFunction<void(T&)>&& Func)
+	{
+		ModBuffer.Modify(MoveTemp(Func), *LastProcessedKeyframe);
+	}
+
+private:
+
+	TBaseStateAccessor<T>& ModBuffer;
+	int32* LastProcessedKeyframe = nullptr;
+};
+
+// Struct that encapsulates writing a new element to a buffer. This is used to allow a new Aux state to be created in the ::Update loop.
+template<typename T>
+struct TLazyStateAccessor
+{
+	TLazyStateAccessor(TUniqueFunction<T*()> && InFunc)
+		: GetWriteNextFunc(MoveTemp(InFunc)) { }
+
+	T* GetWriteNext() const
+	{
+		if (CachedWriteNext == nullptr)
+		{
+			CachedWriteNext = GetWriteNextFunc();
+		}
+		return CachedWriteNext;
+	}
+
+private:
+
+	T* CachedWriteNext = nullptr;
+	TUniqueFunction<T*()> GetWriteNextFunc;
+};
+
+// The main container for all of our buffers
+template<typename InBufferTypes>
 struct TNetworkSimBufferContainer
 {
+	TNetworkSimBufferContainer()
+		: SyncMods(Sync), AuxMods(Aux) { }
+
 	// Collection of types we were assigned
-	using TBufferTypes = T;
+	using TBufferTypes = InBufferTypes;
 
 	// helper that returns the buffer type for a given BufferId (not the underlying type: the actual TReplicatedBuffer<TUnderlyingType>
 	template<ENetworkSimBufferTypeId TypeId>
 	struct select_buffer_type
 	{
-		// We are just wrapping a TReplicationBuffer around whatever type TBufferTypes::select_type returns (which returns the underlying type)
-		using type = TReplicationBuffer<typename TBufferTypes::template select_type<TypeId>::type>;
+		// We are just wrapping a TNetworkSimContiguousBuffer around whatever type TBufferTypes::select_type returns (which returns the underlying type)
+		using type = TNetworkSimContiguousBuffer<typename TBufferTypes::template select_type<TypeId>::type>;
+	};
+
+	template<ENetworkSimBufferTypeId TypeId>
+	struct select_buffer_type_sparse
+	{
+		// We are just wrapping a TNetworkSimSparseBuffer around whatever type TBufferTypes::select_type returns (which returns the underlying type)
+		using type = TNetworkSimSparseBuffer<typename TBufferTypes::template select_type<TypeId>::type>;
 	};
 
 	// The buffer types. This may look intimidating but is not that bad!
 	// We are using select_buffer_type to decide what the type should be for a given ENetworkSimBufferTypeId.
 	using TInputBuffer = typename select_buffer_type<ENetworkSimBufferTypeId::Input>::type;
 	using TSyncBuffer = typename select_buffer_type<ENetworkSimBufferTypeId::Sync>::type;
-	using TAuxBuffer = typename select_buffer_type<ENetworkSimBufferTypeId::Aux>::type;
+	using TAuxBuffer = typename select_buffer_type_sparse<ENetworkSimBufferTypeId::Aux>::type;
 	using TDebugBuffer = typename select_buffer_type<ENetworkSimBufferTypeId::Debug>::type;
 
 	// The buffers themselves. Types are already declared above.
-	// If you are reading just to find the damn underlying type here, its TReplicationBuffer< whatever your struct type is >
+	// If you are reading just to find the damn underlying type here, its (TNetworkSimSparseBuffer||TNetworkSimContiguousBuffer) < whatever your struct type is >
 	TInputBuffer Input;
 	TSyncBuffer Sync;
 	TAuxBuffer Aux;
 	TDebugBuffer Debug;
 
+	TPredictedModBuffer<TSyncBuffer> SyncMods;
+	TPredictedModBuffer<TAuxBuffer> AuxMods;
+
 	// Finally, template accessor for getting buffers based on enum. This is really what all this junk is about.
 	// This allows other templated classes in the system to access a specific buffer from another templated argument
 	template<ENetworkSimBufferTypeId BufferId>
 	typename select_buffer_type<BufferId>::type& Get()
 	{
-		return TBufferGetterHelper<TNetworkSimBufferContainer<T>, BufferId>::Get(*this);
+		return TBufferGetterHelper<TNetworkSimBufferContainer<TBufferTypes>, BufferId>::Get(*this);
 	}
 };
 
-
 // ----------------------------------------------------------------------------------------------------------------------------------------------
 //		Tick and time keeping related structures
 // ----------------------------------------------------------------------------------------------------------------------------------------------
@@ -328,24 +463,17 @@ struct TSimulationTickState
 	void SetTotalProcessedSimulationTime(const FNetworkSimTime& SimTime, int32 Keyframe)
 	{
 		TotalProcessedSimulationTime = SimTime;
-		SimulationTimeBuffer.ResetNextHeadKeyframe(Keyframe);
-		*SimulationTimeBuffer.GetWriteNext() = SimTime;
+		*SimulationTimeBuffer.WriteKeyframe(Keyframe) = SimTime;
 	}
 
 	void IncrementTotalProcessedSimulationTime(const FNetworkSimTime& DeltaSimTime, int32 Keyframe)
 	{
 		TotalProcessedSimulationTime += DeltaSimTime;
-		*SimulationTimeBuffer.GetWriteNext() = TotalProcessedSimulationTime;
-		ensure(SimulationTimeBuffer.GetHeadKeyframe() == Keyframe);
-	}
-
-	void InitSimulationTimeBuffer(int32 Size)
-	{
-		SimulationTimeBuffer.SetBufferSize(Size);
+		*SimulationTimeBuffer.WriteKeyframe(Keyframe) = TotalProcessedSimulationTime;
 	}
 		
 	// Historic tracking of simulation time. This allows us to timestamp sync data as its produced
-	TReplicationBuffer<FNetworkSimTime>	SimulationTimeBuffer;
+	TNetworkSimContiguousBuffer<FNetworkSimTime>	SimulationTimeBuffer;
 
 	void SetTotalAllowedSimulationTime(const FNetworkSimTime& SimTime)
 	{
@@ -415,7 +543,7 @@ struct TInternalBufferTypes : TNetworkSimBufferTypes<
 	// InputCmds are wrapped in TFrameCmd, which will store an explicit sim delta time if we are not a fixed tick sim
 	TFrameCmd< typename TUserBufferTypes::TInputCmd , TTickSettings>,
 
-	typename TUserBufferTypes::TSyncState,	// SyncState Passes through
+	typename TUserBufferTypes::TSyncState,	// SyncState passes through
 	typename TUserBufferTypes::TAuxState,	// Auxstate passes through
 	typename TUserBufferTypes::TDebugState	// Debugstate passes through
 >
@@ -430,7 +558,15 @@ public:
 	virtual FString GetDebugName() const = 0; // Used for debugging. Recommended to emit the simulation name and the actor name/role.
 	virtual const UObject* GetVLogOwner() const = 0; // Owning object for Visual Logs
 
-	virtual void InitSyncState(typename TBufferTypes::TSyncState& OutSyncState) const = 0;	// Called to create initial value of the sync state.
-	virtual void ProduceInput(const FNetworkSimTime SimTime, typename TBufferTypes::TInputCmd&) = 0; // Called when the sim is ready to process new local input
-	virtual void FinalizeFrame(const typename TBufferTypes::TSyncState& SyncState) = 0; // Called from the Network Sim at the end of the sim frame when there is new sync data.
+	// Called to create initial value of the sync state.
+	virtual void InitSyncState(typename TBufferTypes::TSyncState& OutSyncState) const = 0;
+
+	// Called to create initial value of the aux state.
+	virtual void InitAuxState(typename TBufferTypes::TAuxState& OutAuxState) const = 0;
+	
+	// Called whenever the sim is ready to process new local input.
+	virtual void ProduceInput(const FNetworkSimTime SimTime, typename TBufferTypes::TInputCmd&) = 0;
+	
+	// Called from the Network Sim at the end of the sim frame when there is new sync data.
+	virtual void FinalizeFrame(const typename TBufferTypes::TSyncState& SyncState, const typename TBufferTypes::TAuxState& AuxState) = 0;
 };
\ No newline at end of file
diff --git a/Engine/Plugins/Runtime/NetworkPrediction/readme.txt b/Engine/Plugins/Runtime/NetworkPrediction/readme.txt
--- a/Engine/Plugins/Runtime/NetworkPrediction/readme.txt
+++ b/Engine/Plugins/Runtime/NetworkPrediction/readme.txt
@@ -77,8 +77,7 @@ High level focus:
 -We intend to support RootMotion (Animation) / RootMotionSources (non anim) in the new movement system.
 
 Road Map:
--Aux Buffer
--Event System
+-Event System / "Game API"
 -Ability System integration
 -New Movement System
  
@@ -106,6 +105,10 @@ Smoothing: Taking the output of the simulation and applying an additional layer
 // Release notes
 // ----------------------------------------------------------------------------------------------------------
 
+Update (10-15-19)
+-Aux buffer hooked up. Still some decisions to make here, expect a few more refactors.
+-TReplicationBuffer replaced with simpler API.
+
 Update (10-1-19)
 -Forward Predict / Dependent simulation initial check in. This has some limitations and is not final
 	-Dependent Sim can't trigger reconcile yet
diff --git a/Engine/Plugins/Runtime/NetworkPredictionExtras/Source/NetworkPredictionExtras/Private/MockNetworkSimulation.cpp b/Engine/Plugins/Runtime/NetworkPredictionExtras/Source/NetworkPredictionExtras/Private/MockNetworkSimulation.cpp
--- a/Engine/Plugins/Runtime/NetworkPredictionExtras/Source/NetworkPredictionExtras/Private/MockNetworkSimulation.cpp
+++ b/Engine/Plugins/Runtime/NetworkPredictionExtras/Source/NetworkPredictionExtras/Private/MockNetworkSimulation.cpp
@@ -62,7 +62,7 @@ static bool ForceMispredict = false;
 
 const FName FMockNetworkSimulation::GroupName(TEXT("Mock"));
 
-void FMockNetworkSimulation::Update(IMockDriver* Driver, const float DeltaTimeSeconds, const FMockInputCmd& InputCmd, const FMockSyncState& InputState, FMockSyncState& OutputState, const FMockAuxState& AuxState)
+void FMockNetworkSimulation::Update(IMockDriver* Driver, const float DeltaTimeSeconds, const FMockInputCmd& InputCmd, const FMockSyncState& InputState, FMockSyncState& OutputState, const FMockAuxState& AuxState, const TLazyStateAccessor<FMockAuxState>& OutAuxStateAccessor)
 {
 	OutputState.Total = InputState.Total + (InputCmd.InputValue * AuxState.Multiplier * DeltaTimeSeconds);
 
@@ -171,7 +171,12 @@ void UMockNetworkSimulationComponent::InitSyncState(FMockSyncState& OutSyncState
 	OutSyncState.Total = MockValue;
 }
 
-void UMockNetworkSimulationComponent::FinalizeFrame(const FMockSyncState& SyncState)
+void UMockNetworkSimulationComponent::InitAuxState(FMockAuxState& OutAuxState) const
+{
+
+}
+
+void UMockNetworkSimulationComponent::FinalizeFrame(const FMockSyncState& SyncState, const FMockAuxState& AuxState)
 {
 	MockValue = SyncState.Total;
 }
diff --git a/Engine/Plugins/Runtime/NetworkPredictionExtras/Source/NetworkPredictionExtras/Private/NetworkPredictionExtrasFlyingPawn.cpp b/Engine/Plugins/Runtime/NetworkPredictionExtras/Source/NetworkPredictionExtras/Private/NetworkPredictionExtrasFlyingPawn.cpp
--- a/Engine/Plugins/Runtime/NetworkPredictionExtras/Source/NetworkPredictionExtras/Private/NetworkPredictionExtrasFlyingPawn.cpp
+++ b/Engine/Plugins/Runtime/NetworkPredictionExtras/Source/NetworkPredictionExtras/Private/NetworkPredictionExtrasFlyingPawn.cpp
@@ -14,6 +14,10 @@
 #include "Components/CapsuleComponent.h"
 #include "DrawDebugHelpers.h"
 
+#include "Misc/AssertionMacros.h"
+#include "HAL/PlatformStackWalk.h"
+#include "HAL/ThreadHeartBeat.h"
+
 namespace FlyingPawnCVars
 {
 
@@ -32,8 +36,13 @@ static FAutoConsoleVariableRef CVarBindAutomatically(TEXT("NetworkPredictionExtr
 ANetworkPredictionExtrasFlyingPawn::ANetworkPredictionExtrasFlyingPawn()
 {
 	FlyingMovementComponent = CreateDefaultSubobject<UFlyingMovementComponent>(TEXT("FlyingMovementComponent"));
-	
-	if (HasAnyFlags(RF_ClassDefaultObject) == false)
+}
+
+void ANetworkPredictionExtrasFlyingPawn::BeginPlay()
+{
+	Super::BeginPlay();
+
+	if (UWorld* World = GetWorld())
 	{
 		FlyingMovementComponent->ProduceInputDelegate.BindUObject(this, &ANetworkPredictionExtrasFlyingPawn::ProduceInput);
 
@@ -42,13 +51,42 @@ ANetworkPredictionExtrasFlyingPawn::ANetworkPredictionExtrasFlyingPawn()
 		{
 			if (ULocalPlayer* LocalPlayer = GetWorld()->GetFirstLocalPlayerFromController())
 			{
-				LocalPlayer->Exec(GetWorld(), TEXT("setbind Nine nms.Debug.LocallyControlledPawn"), *GLog);
-				LocalPlayer->Exec(GetWorld(), TEXT("setbind Zero nms.Debug.ToggleContinous"), *GLog);
+				LocalPlayer->Exec(World, TEXT("setbind Nine nms.Debug.LocallyControlledPawn"), *GLog);
+				LocalPlayer->Exec(World, TEXT("setbind Zero nms.Debug.ToggleContinous"), *GLog);
 			}
 		}
 	}
+}
 
+UNetConnection* ANetworkPredictionExtrasFlyingPawn::GetNetConnection() const
+{
+	UNetConnection* SuperNetConnection = Super::GetNetConnection();
+	if (SuperNetConnection)
+	{
+		return SuperNetConnection;
+	}
 
+	if (bFakeAutonomousProxy)
+	{
+		if (GetLocalRole() == ROLE_Authority)
+		{
+			for (FConstPlayerControllerIterator Iterator = GetWorld()->GetPlayerControllerIterator(); Iterator; ++Iterator)
+			{
+				APlayerController* PC = Iterator->Get();
+				if (PC->GetNetConnection() && PC->GetPawn())
+				{
+					return PC->GetNetConnection();
+				}
+			}
+		}
+		if (GetLocalRole() == ROLE_AutonomousProxy && GetWorld()->GetFirstPlayerController() && GetWorld()->GetFirstPlayerController()->GetPawn())
+		{
+			return GetWorld()->GetFirstPlayerController()->GetNetConnection();
+		}
+	}
+	
+
+	return nullptr;
 }
 
 void ANetworkPredictionExtrasFlyingPawn::SetupPlayerInputComponent(UInputComponent* PlayerInputComponent)
@@ -133,6 +171,46 @@ void ANetworkPredictionExtrasFlyingPawn::Tick( float DeltaSeconds)
 {
 	Super::Tick(DeltaSeconds);
 	// Do whatever you want here. By now we have the latest movement state and latest input processed.
+
+	if (bFakeAutonomousProxy && GetLocalRole() == ROLE_Authority)
+	{
+		if (GetRemoteRole() != ROLE_AutonomousProxy)
+		{
+			SetAutonomousProxy(true);
+		}
+	}
+}
+
+void ANetworkPredictionExtrasFlyingPawn::PrintDebug()
+{
+	UE_LOG(LogTemp, Warning, TEXT("======== ANetworkPredictionExtrasFlyingPawn::PrintDebug ========"));
+
+	FSlowHeartBeatScope SuspendHeartBeat;
+	FDisableHitchDetectorScope SuspendGameThreadHitch;
+
+	PrintScriptCallstack();
+
+	const SIZE_T StackTraceSize = 65535;
+	ANSICHAR* StackTrace = (ANSICHAR*)FMemory::SystemMalloc(StackTraceSize);
+	if (StackTrace != nullptr)
+	{
+		StackTrace[0] = 0;
+		// Walk the stack and dump it to the allocated memory.
+		FPlatformStackWalk::StackWalkAndDump(StackTrace, StackTraceSize, 1);
+		UE_LOG(LogTemp, Log, TEXT("Call Stack:\n%s"), ANSI_TO_TCHAR(StackTrace));
+		FMemory::SystemFree(StackTrace);
+	}
+}
+
+void ANetworkPredictionExtrasFlyingPawn::SetMaxMoveSpeed(float NewMaxMoveSpeed)
+{
+	if (ensure(FlyingMovementComponent))
+	{
+		FlyingMovementComponent->MovementAuxState->Modify([NewMaxMoveSpeed](FlyingMovement::FAuxState& State)
+		{
+			State.MaxSpeed = NewMaxMoveSpeed;
+		});
+	}
 }
 
 void ANetworkPredictionExtrasFlyingPawn::ProduceInput(const FNetworkSimTime SimTime, FlyingMovement::FInputCmd& Cmd)
@@ -149,6 +227,14 @@ void ANetworkPredictionExtrasFlyingPawn::ProduceInput(const FNetworkSimTime SimT
 	// answers and will depend on the game and its specific needs. In general, at this time, I'd recommend aim assist and lock on 
 	// targeting systems to happen /outside/ of the system, i.e, here. But I can think of scenarios where that may not be ideal too.
 
+	if (InputPreset == ENetworkPredictionExtrasFlyingInputPreset::Forward)
+	{
+		Cmd.MovementInput = FVector(1.f, 0.f, 0.f);
+		Cmd.RotationInput = FRotator::ZeroRotator;
+		return;
+	}
+
+
 	if (Controller == nullptr)
 	{
 		// We don't have a local controller so we can't run the code below. This is ok. Simulated proxies will just use previous input when extrapolating
diff --git a/Engine/Plugins/Runtime/NetworkPredictionExtras/Source/NetworkPredictionExtras/Public/MockNetworkSimulation.h b/Engine/Plugins/Runtime/NetworkPredictionExtras/Source/NetworkPredictionExtras/Public/MockNetworkSimulation.h
--- a/Engine/Plugins/Runtime/NetworkPredictionExtras/Source/NetworkPredictionExtras/Public/MockNetworkSimulation.h
+++ b/Engine/Plugins/Runtime/NetworkPredictionExtras/Source/NetworkPredictionExtras/Public/MockNetworkSimulation.h
@@ -137,7 +137,7 @@ class FMockNetworkSimulation
 public:
 
 	/** Main update function */
-	static void Update(IMockDriver* Driver, const float DeltaTimeSeconds, const FMockInputCmd& InputCmd, const FMockSyncState& InputState, FMockSyncState& OutputState, const FMockAuxState& AuxState);
+	static void Update(IMockDriver* Driver, const float DeltaTimeSeconds, const FMockInputCmd& InputCmd, const FMockSyncState& InputState, FMockSyncState& OutputState, const FMockAuxState& AuxState, const TLazyStateAccessor<FMockAuxState>& OutAuxStateAccessor);
 
 	/** Tick group the simulation maps to */
 	static const FName GroupName;
@@ -173,8 +173,11 @@ public:
 	virtual UObject* GetVLogOwner() const override final;
 
 	void InitSyncState(FMockSyncState& OutSyncState) const override;
-	void FinalizeFrame(const FMockSyncState& SyncState) override;
+	void InitAuxState(FMockAuxState& OutAuxState) const override;
 	void ProduceInput(const FNetworkSimTime SimTime, FMockInputCmd& Cmd);
+	void FinalizeFrame(const FMockSyncState& SyncState, const FMockAuxState& AuxState) override;
+	
+
 	virtual UWorld* GetDriverWorld() const override final { return GetWorld(); }
 	virtual FTransform GetDebugWorldTransform() const override final;
 
diff --git a/Engine/Plugins/Runtime/NetworkPredictionExtras/Source/NetworkPredictionExtras/Public/NetworkPredictionExtrasFlyingPawn.h b/Engine/Plugins/Runtime/NetworkPredictionExtras/Source/NetworkPredictionExtras/Public/NetworkPredictionExtrasFlyingPawn.h
--- a/Engine/Plugins/Runtime/NetworkPredictionExtras/Source/NetworkPredictionExtras/Public/NetworkPredictionExtrasFlyingPawn.h
+++ b/Engine/Plugins/Runtime/NetworkPredictionExtras/Source/NetworkPredictionExtras/Public/NetworkPredictionExtrasFlyingPawn.h
@@ -31,6 +31,14 @@ class UFlyingMovementComponent;
 //
 // -------------------------------------------------------------------------------------------------------------------------------
 
+UENUM()
+enum class ENetworkPredictionExtrasFlyingInputPreset: uint8
+{
+	/** No input */
+	None,
+	/** Just moves forward */
+	Forward
+};
 
 /** Sample pawn that uses UFlyingMovementComponent. The main thing this provides is actually producing user input for the component/simulation to consume. */
 UCLASS(config = Game)
@@ -38,10 +46,28 @@ class NETWORKPREDICTIONEXTRAS_API ANetworkPredictionExtrasFlyingPawn : public AP
 {
 	GENERATED_BODY()
 
+public:
+
 	ANetworkPredictionExtrasFlyingPawn();
 
+	virtual void BeginPlay() override;
 	virtual void SetupPlayerInputComponent(UInputComponent* PlayerInputComponent) override;
 	virtual void Tick( float DeltaSeconds) override;
+	virtual UNetConnection* GetNetConnection() const override;
+
+
+	UPROPERTY(EditAnywhere, BlueprintReadWrite, Category="Automation")
+	ENetworkPredictionExtrasFlyingInputPreset InputPreset;
+
+	/** Actor will behave like autonomous proxy even though not posessed by an APlayercontroller. To be used in conjuction with InputPreset. */
+	UPROPERTY(EditAnywhere, BlueprintReadOnly, Category="Automation")
+	bool bFakeAutonomousProxy = false;
+
+	UFUNCTION(BlueprintCallable, Category="Debug")
+	void PrintDebug();
+
+	UFUNCTION(BlueprintCallable, Category="Gameplay")
+	void SetMaxMoveSpeed(float NewMaxMoveSpeed);
 
 private:
 
