Network Prediction Plugin
-Cleanup how we advance the simulation and unified which keyframes are accessed across the three main buffers
-LastProcessedInputKeyframe -> PendingKeyframe. Code is simplified dealing with the "next" keyframe rather than making assumptions based on "last" keyframe.
-Input/Sync/Aux states [at] PendingKeyframe are what gets used as inputs to T::Update. This is more consistent that how it worked previously, where it was essentially (Input+1/Sync) -> (Sync+1).
-This gets rid of awkward hacks like empty input cmd [at] keyframe = 0.
-Also got rid of InitSync/InitAuxState interface functions. The initial states can now be provided when instantiating the network sim model.

#rnx

diff --git a/Engine/Plugins/Runtime/NetworkPrediction/Source/NetworkPrediction/Private/Movement/FlyingMovementComponent.cpp b/Engine/Plugins/Runtime/NetworkPrediction/Source/NetworkPrediction/Private/Movement/FlyingMovementComponent.cpp
--- a/Engine/Plugins/Runtime/NetworkPrediction/Source/NetworkPrediction/Private/Movement/FlyingMovementComponent.cpp
+++ b/Engine/Plugins/Runtime/NetworkPrediction/Source/NetworkPrediction/Private/Movement/FlyingMovementComponent.cpp
@@ -69,7 +69,16 @@ UFlyingMovementComponent::UFlyingMovementComponent()
 
 INetworkSimulationModel* UFlyingMovementComponent::InstantiateNetworkSimulation()
 {
-	auto NewSim = new FlyingMovement::FMovementSystem<0>(this);
+	check(UpdatedComponent);
+	
+	FlyingMovement::FMoveState InitialSyncState;
+	InitialSyncState.Location = UpdatedComponent->GetComponentLocation();
+	InitialSyncState.Rotation = UpdatedComponent->GetComponentQuat().Rotator();	
+
+	FlyingMovement::FAuxState InitialAuxState;
+	InitialAuxState.MaxSpeed = FlyingMovementCVars::MaxSpeed;
+
+	auto NewSim = new FlyingMovement::FMovementSystem<0>(this, InitialSyncState, InitialAuxState);
 	DO_NETSIM_MODEL_DEBUG(FNetworkSimulationModelDebuggerManager::Get().RegisterNetworkSimulationModel(NewSim, GetOwner()));
 	MovementSyncState = &NewSim->SyncAccessor;
 	MovementAuxState = &NewSim->AuxAccessor;
@@ -127,18 +136,6 @@ void UFlyingMovementComponent::TickComponent(float DeltaTime, enum ELevelTick Ti
 //
 // ----------------------------------------------------------------------------------------------------------
 
-void UFlyingMovementComponent::InitSyncState(FlyingMovement::FMoveState& OutSyncState) const
-{
-	OutSyncState.Location = UpdatedComponent->GetComponentLocation();
-	OutSyncState.Rotation = UpdatedComponent->GetComponentQuat().Rotator();	
-}
-
-void UFlyingMovementComponent::InitAuxState(FlyingMovement::FAuxState& OutAuxState) const
-{
-	OutAuxState = FlyingMovement::FAuxState();
-	OutAuxState.MaxSpeed = FlyingMovementCVars::MaxSpeed;
-}
-
 void UFlyingMovementComponent::PreSimSync(const FlyingMovement::FMoveState& SyncState)
 {
 	// Does checking equality make any sense here? This is unfortunate
diff --git a/Engine/Plugins/Runtime/NetworkPrediction/Source/NetworkPrediction/Private/Movement/ParametricMovement.cpp b/Engine/Plugins/Runtime/NetworkPrediction/Source/NetworkPrediction/Private/Movement/ParametricMovement.cpp
--- a/Engine/Plugins/Runtime/NetworkPrediction/Source/NetworkPrediction/Private/Movement/ParametricMovement.cpp
+++ b/Engine/Plugins/Runtime/NetworkPrediction/Source/NetworkPrediction/Private/Movement/ParametricMovement.cpp
@@ -145,17 +145,6 @@ void UParametricMovementComponent::BeginPlay()
 	CachedStartingTransform = UpdatedComponent->GetComponentToWorld();
 }
 
-void UParametricMovementComponent::InitSyncState(ParametricMovement::FMoveState& OutSyncState) const
-{
-	// In this case, we just default to the 0 position. Maybe this could be a starting variable set on the component.
-	OutSyncState.Position = 0.f;
-	OutSyncState.PlayRate = 0.f;
-}
-
-void UParametricMovementComponent::InitAuxState(ParametricMovement::FAuxState& OutAuxState) const
-{
-}
-
 void UParametricMovementComponent::FinalizeFrame(const ParametricMovement::FMoveState& SyncState, const ParametricMovement::FAuxState& AuxState)
 {
 	FTransform NewTransform;
diff --git a/Engine/Plugins/Runtime/NetworkPrediction/Source/NetworkPrediction/Public/Movement/FlyingMovement.h b/Engine/Plugins/Runtime/NetworkPrediction/Source/NetworkPrediction/Public/Movement/FlyingMovement.h
--- a/Engine/Plugins/Runtime/NetworkPrediction/Source/NetworkPrediction/Public/Movement/FlyingMovement.h
+++ b/Engine/Plugins/Runtime/NetworkPrediction/Source/NetworkPrediction/Public/Movement/FlyingMovement.h
@@ -185,9 +185,7 @@ public:
 
 	FString GetDebugName() const override;
 	const UObject* GetVLogOwner() const override;
-
-	void InitSyncState(FlyingMovement::FMoveState& OutSyncState) const override;
-	void InitAuxState(FlyingMovement::FAuxState& OutAuxState) const override;
+	
 	void ProduceInput(const FNetworkSimTime SimTime, FlyingMovement::FInputCmd& Cmd) override;
 	void FinalizeFrame(const FlyingMovement::FMoveState& SyncState, const FlyingMovement::FAuxState& AuxState) override;
 
diff --git a/Engine/Plugins/Runtime/NetworkPrediction/Source/NetworkPrediction/Public/Movement/ParametricMovement.h b/Engine/Plugins/Runtime/NetworkPrediction/Source/NetworkPrediction/Public/Movement/ParametricMovement.h
--- a/Engine/Plugins/Runtime/NetworkPrediction/Source/NetworkPrediction/Public/Movement/ParametricMovement.h
+++ b/Engine/Plugins/Runtime/NetworkPrediction/Source/NetworkPrediction/Public/Movement/ParametricMovement.h
@@ -153,8 +153,6 @@ class NETWORKPREDICTION_API UParametricMovementComponent : public UBaseMovementC
 	// Base TNetworkModelSimulation driver
 	FString GetDebugName() const override;
 	const UObject* GetVLogOwner() const override;
-	void InitSyncState(ParametricMovement::FMoveState& OutSyncState) const override;
-	void InitAuxState(ParametricMovement::FAuxState& OutAuxState) const override;
 	void ProduceInput(const FNetworkSimTime SimTime, ParametricMovement::FInputCmd& Cmd);
 	void FinalizeFrame(const ParametricMovement::FMoveState& SyncState, const ParametricMovement::FAuxState& AuxState) override;	
 
diff --git a/Engine/Plugins/Runtime/NetworkPrediction/Source/NetworkPrediction/Public/NetworkPredictionTypes.h b/Engine/Plugins/Runtime/NetworkPrediction/Source/NetworkPrediction/Public/NetworkPredictionTypes.h
--- a/Engine/Plugins/Runtime/NetworkPrediction/Source/NetworkPrediction/Public/NetworkPredictionTypes.h
+++ b/Engine/Plugins/Runtime/NetworkPrediction/Source/NetworkPrediction/Public/NetworkPredictionTypes.h
@@ -55,7 +55,7 @@ struct FNetSimProcessedFrameDebugInfo
 	int32 LocalGFrameNumber; // Local GFrame number
 	float LocalDeltaTimeSeconds; // Local frame time
 	TArray<int32> ProcessedKeyframes; // Which keyframes were processed this frame
-	int32 LastProcessedKeyframe; // What LastProcessedKeyframe was at the end of the frame. Does NOT mean we processed it this frame!
+	int32 PendingKeyframe; // What PendingKeyframe was at the end of the frame. Does NOT mean we processed it this frame!
 	int32 HeadKeyframe; // Head keyframe of the inputbuffer when the frame ended
 
 	float RemainingAllowedSimulationTimeSeconds;
@@ -67,7 +67,7 @@ struct FNetSimProcessedFrameDebugInfo
 	{	
 		P.Ar << LocalGFrameNumber;
 		P.Ar << LocalDeltaTimeSeconds;
-		P.Ar << LastProcessedKeyframe;
+		P.Ar << PendingKeyframe;
 		P.Ar << HeadKeyframe;
 		P.Ar << RemainingAllowedSimulationTimeSeconds;
 
diff --git a/Engine/Plugins/Runtime/NetworkPrediction/Source/NetworkPrediction/Public/NetworkSimulationModel.h b/Engine/Plugins/Runtime/NetworkPrediction/Source/NetworkPrediction/Public/NetworkSimulationModel.h
--- a/Engine/Plugins/Runtime/NetworkPrediction/Source/NetworkPrediction/Public/NetworkSimulationModel.h
+++ b/Engine/Plugins/Runtime/NetworkPrediction/Source/NetworkPrediction/Public/NetworkSimulationModel.h
@@ -52,11 +52,14 @@ public:
 	using TSimTime = FNetworkSimTime;
 	using TRealTime = FNetworkSimTime::FRealTime;
 	
-	TNetworkedSimulationModel(TDriver* InDriver)
-		: SyncAccessor(Buffers.SyncMods, &TickInfo.LastProcessedInputKeyframe)
-		, AuxAccessor(Buffers.AuxMods, &TickInfo.LastProcessedInputKeyframe)
+	TNetworkedSimulationModel(TDriver* InDriver, const TSyncState& InitialSyncState = TSyncState(), const TAuxState& InitialAuxState = TAuxState())
+		: SyncAccessor(Buffers.SyncMods, &TickInfo.PendingKeyframe)
+		, AuxAccessor(Buffers.AuxMods, &TickInfo.PendingKeyframe)
 	{
 		Driver = InDriver;
+		*Buffers.Sync.WriteKeyframe(0) = InitialSyncState;
+		*Buffers.Aux.WriteKeyframe(0) = InitialAuxState;
+		TickInfo.SetTotalProcessedSimulationTime(FNetworkSimTime(), 0);
 	}
 	virtual ~TNetworkedSimulationModel()
 	{
@@ -125,90 +128,52 @@ public:
 		// -------------------------------------------------------------------------------------------------------------------------------------------------
 		//												Input Processing & Simulation Update
 		// -------------------------------------------------------------------------------------------------------------------------------------------------
-		if (Buffers.Input.HeadKeyframe() > Buffers.Sync.HeadKeyframe())
+		while (TickInfo.PendingKeyframe <= TickInfo.MaxAllowedKeyframe)
 		{
-			// -------------------------------------------------------------------------------------------------------------
-			//	
-			//	The SyncedState buffer needs to be in sync here:
-			//		-We want it to have a SyncedState, but it may not on the first frame through (thats ok).
-			//		-Its HeadKeyframe should be one behind the Keyframe we are about to process.
-			//
-			//	Note, InputCmds start @ Keyframe=1. The first SyncedState that Update produces will go in KeyFrame=1.
-			//	(E.g, InputCmd @ keyframe=X is used to generate SyncState @ keyframe=X)
-			//	This means that SyncedState @ keyframe=0 is always created here via InitSyncState.
-			//	This also means that we never actually process InputCmd @ keyframe=0. Which is why LastProcessedInputKeyframe is initialized to 0 ("already processed")
-			//	and the buffer has an empty element inserted in InitializeForNetworkRole.
-			// -------------------------------------------------------------------------------------------------------------
-
-			if (Buffers.Sync.HeadKeyframe() != TickInfo.LastProcessedInputKeyframe)
-			{
-				if (TickInfo.LastProcessedInputKeyframe != 0)
-				{
-					// This shouldn't happen, but is not fatal. We are reseting the sync state buffer.
-					UE_LOG(LogNetworkSim, Warning, TEXT("%s. Break in SyncState continuity. LastProcessedInputKeyframe: %d. SyncBuffer.HeadKeyframe(): %d."), *Driver->GetDebugName() , TickInfo.LastProcessedInputKeyframe, Buffers.Sync.HeadKeyframe());
-				}
+			const int32 InputKeyframe = TickInfo.PendingKeyframe;
+			const int32 OutputKeyframe = TickInfo.PendingKeyframe + 1;
 
-				// We need an initial/current state. Get this from the sim driver
-				TSyncState* StartingState = Buffers.Sync.WriteKeyframe(TickInfo.LastProcessedInputKeyframe);
-				Driver->InitSyncState(*StartingState);
-
-				// Reset time tracking buffer too
-				TickInfo.SetTotalProcessedSimulationTime(TickInfo.GetTotalProcessedSimulationTime(), Buffers.Sync.HeadKeyframe());
-			}
-
-			if (Buffers.Aux.Num() == 0)
+			const TInputCmd* InputCmd = Buffers.Input[InputKeyframe];
+			if (!ensureMsgf(InputCmd, TEXT("No InputCmd available for Keyframe %d. PendingKeyframe: %d. MaxAllowedKeyframe: %d."), InputKeyframe, TickInfo.PendingKeyframe, TickInfo.MaxAllowedKeyframe))
 			{
-				TAuxState* StartingAuxState = Buffers.Aux.WriteKeyframe(TickInfo.LastProcessedInputKeyframe);
-				Driver->InitAuxState(*StartingAuxState);
+				break;
 			}
-		
-			// -------------------------------------------------------------------------------------------------------------
-			// Process Input
-			// -------------------------------------------------------------------------------------------------------------
-			while(true)
+
+			// We have an unprocessed command, do we have enough allotted simulation time to process it?
+			if (TickInfo.GetRemainingAllowedSimulationTime() >= InputCmd->GetFrameDeltaTime())
 			{
-				const int32 Keyframe = TickInfo.LastProcessedInputKeyframe+1;
-				if (Keyframe > TickInfo.MaxAllowedInputKeyframe)
+				TSyncState* InSyncState = Buffers.Sync[InputKeyframe];
+				if (InSyncState == nullptr)
 				{
-					break;
+					// We don't have valid sync state for this frame. This could mean we skipped ahead somehow, such as heavy packet loss where the server doesn't receive all input cmds
+					// We will use the last known sync state, but first write it out to the current InputKeyframe
+					UE_LOG(LogNetworkSim, Log, TEXT("%s: No sync state found @ keyframe %d. Using previous head element @ keyframe %d."), *Driver->GetDebugName(), InputKeyframe, Buffers.Sync.HeadKeyframe());
+					
+					TSyncState LastSyncState = *Buffers.Sync.HeadElement();
+					InSyncState = Buffers.Sync.WriteKeyframe(InputKeyframe);
+					*InSyncState = LastSyncState;
 				}
 
-				if (TInputCmd* NextCmd = Buffers.Input[Keyframe])
-				{
-					// We have an unprocessed command, do we have enough allotted simulation time to process it?
-					if (TickInfo.GetRemainingAllowedSimulationTime() >= NextCmd->GetFrameDeltaTime())
-					{
-						// -------------------------------------------------------------------------------------------------
-						//	The core process input command and call ::Update block!
-						// -------------------------------------------------------------------------------------------------
-						TSyncState* PrevSyncState = Buffers.Sync[TickInfo.LastProcessedInputKeyframe];
-						TSyncState* NextSyncState = Buffers.Sync.WriteKeyframe(Keyframe);
-						TAuxState* AuxState = Buffers.Aux[TickInfo.LastProcessedInputKeyframe];
-
-						check(PrevSyncState != nullptr);
-						check(NextSyncState != nullptr);
-						check(Buffers.Sync.HeadKeyframe() == Keyframe);
-				
-						if (DebugState)
-						{
-							DebugState->ProcessedKeyframes.Add(Keyframe);
-						}
-						
-						TSimulation::Update(Driver, NextCmd->GetFrameDeltaTime().ToRealTimeSeconds(), *NextCmd, *PrevSyncState, *NextSyncState, *AuxState, Buffers.Aux.WriteKeyframeFunc(Keyframe+1));
+				const TAuxState* InAuxState = Buffers.Aux[InputKeyframe];
+				checkf(InAuxState, TEXT("No AuxState available for Keyframe %d. PendingKeyframe: %d. MaxAllowedKeyframe: %d."), InputKeyframe, TickInfo.PendingKeyframe, TickInfo.MaxAllowedKeyframe);
+
+				TSyncState* OutSyncState = Buffers.Sync.WriteKeyframe(OutputKeyframe);
+				check(OutSyncState);
+
+				TSimulation::Update(Driver, InputCmd->GetFrameDeltaTime().ToRealTimeSeconds(), *InputCmd, *InSyncState, *OutSyncState, *InAuxState, Buffers.Aux.WriteKeyframeFunc(OutputKeyframe));
 					
-						TickInfo.IncrementTotalProcessedSimulationTime(NextCmd->GetFrameDeltaTime(), Keyframe);
-						TickInfo.LastProcessedInputKeyframe = Keyframe;
-					}
-					else
-					{
-						break;
-					}
-				}
-				else
+				TickInfo.IncrementTotalProcessedSimulationTime(InputCmd->GetFrameDeltaTime(), OutputKeyframe);
+				TickInfo.PendingKeyframe = OutputKeyframe;
+
+				if (DebugState)
 				{
-					break;
+					DebugState->ProcessedKeyframes.Add(InputKeyframe);
 				}
 			}
+			else
+			{
+				break;
+			}
 		}
 
 		// -------------------------------------------------------------------------------------------------------------------------------------------------
@@ -237,7 +202,7 @@ public:
 		// Finish debug state buffer recording (what the server processed each frame)
 		if (DebugState)
 		{
-			DebugState->LastProcessedKeyframe = TickInfo.LastProcessedInputKeyframe;
+			DebugState->PendingKeyframe = TickInfo.PendingKeyframe;
 			DebugState->HeadKeyframe = Buffers.Input.HeadKeyframe();
 			DebugState->RemainingAllowedSimulationTimeSeconds = (float)TickInfo.GetRemainingAllowedSimulationTime().ToRealTimeSeconds();
 		}
@@ -298,9 +263,6 @@ public:
 
 		Buffers.SyncMods.SetIsAuthority(Role == ROLE_Authority);
 		Buffers.AuxMods.SetIsAuthority(Role == ROLE_Authority);
-
-		// We want to start with an empty command in the input buffer. The sync buffer will be populated @ frame 0 with the "current" state when we actually sim. This keeps them in sync
-		*Buffers.Input.WriteKeyframe(0) = TInputCmd();
 	}
 
 	void NetSerializeProxy(EReplicationProxyTarget Target, const FNetSerializeParams& Params) final override
diff --git a/Engine/Plugins/Runtime/NetworkPrediction/Source/NetworkPrediction/Public/NetworkSimulationModelBuffer.h b/Engine/Plugins/Runtime/NetworkPrediction/Source/NetworkPrediction/Public/NetworkSimulationModelBuffer.h
--- a/Engine/Plugins/Runtime/NetworkPrediction/Source/NetworkPrediction/Public/NetworkSimulationModelBuffer.h
+++ b/Engine/Plugins/Runtime/NetworkPrediction/Source/NetworkPrediction/Public/NetworkSimulationModelBuffer.h
@@ -10,7 +10,7 @@ struct TNetworkSimBufferIterator
 	ElementType* Element() const { return Buffer.IterElement(CurrentPos); }
 
 	void operator++() { CurrentPos++; }
-	operator bool()	{ return CurrentPos <= Buffer.IterEndPos(); }
+	operator bool()	{ return CurrentPos != INDEX_NONE && CurrentPos <= Buffer.IterEndPos(); }
 
 private:
 	TBuffer& Buffer;
diff --git a/Engine/Plugins/Runtime/NetworkPrediction/Source/NetworkPrediction/Public/NetworkSimulationModelDebugger.h b/Engine/Plugins/Runtime/NetworkPrediction/Source/NetworkPrediction/Public/NetworkSimulationModelDebugger.h
--- a/Engine/Plugins/Runtime/NetworkPrediction/Source/NetworkPrediction/Public/NetworkSimulationModelDebugger.h
+++ b/Engine/Plugins/Runtime/NetworkPrediction/Source/NetworkPrediction/Public/NetworkSimulationModelDebugger.h
@@ -433,7 +433,7 @@ struct TNetworkSimulationModelDebugger : public INetworkSimulationModelDebugger
 				}
 
 				// Unprocessed InputCmds (above processed)				
-				for (int32 Keyframe = DebugState->LastProcessedKeyframe+1; Keyframe <= DebugState->HeadKeyframe; ++Keyframe)
+				for (int32 Keyframe = DebugState->PendingKeyframe; Keyframe <= DebugState->HeadKeyframe; ++Keyframe)
 				{
 					if (auto* Cmd = InputBuffer[Keyframe])
 					{
@@ -508,7 +508,7 @@ struct TNetworkSimulationModelDebugger : public INetworkSimulationModelDebugger
 					ColumnTime += (float)Cmd->GetFrameDeltaTime().ToRealTimeSeconds();
 				}
 			}
-			for (int32 Keyframe = DebugState->LastProcessedKeyframe+1; Keyframe <= DebugState->HeadKeyframe; ++Keyframe)
+			for (int32 Keyframe = DebugState->PendingKeyframe; Keyframe <= DebugState->HeadKeyframe; ++Keyframe)
 			{
 				if (auto* Cmd = InputBuffer[Keyframe])
 				{
@@ -534,7 +534,7 @@ struct TNetworkSimulationModelDebugger : public INetworkSimulationModelDebugger
 		// ------------------------------------------------------------------------------------------------------------------------------------------------
 
 		Out.Emit(FString::Printf(TEXT("%s - %s"), *Owner->GetName(), *UEnum::GetValueAsString(TEXT("Engine.ENetRole"), Owner->GetLocalRole())), FColor::Yellow);
-		Out.Emit(FString::Printf(TEXT("LastProcessedInputKeyframe: %d (%d Buffered)"), NetworkSim->TickInfo.LastProcessedInputKeyframe, NetworkSim->Buffers.Input.HeadKeyframe() - NetworkSim->TickInfo.LastProcessedInputKeyframe));
+		Out.Emit(FString::Printf(TEXT("PendingKeyframe: %d (%d Buffered)"), NetworkSim->TickInfo.PendingKeyframe, NetworkSim->Buffers.Input.HeadKeyframe() - NetworkSim->TickInfo.PendingKeyframe));
 				
 		if (Owner->GetLocalRole() == ROLE_AutonomousProxy)
 		{			
@@ -567,7 +567,7 @@ struct TNetworkSimulationModelDebugger : public INetworkSimulationModelDebugger
 				*NetworkSim->RepProxy_Autonomous.GetLastSerializedSimTime().ToString(), *(NetworkSim->TickInfo.GetTotalProcessedSimulationTime() - NetworkSim->RepProxy_Autonomous.GetLastSerializedSimTime()).ToString());
 			Out.Emit(*SimulationTimeString, Color);
 
-			FString AllowedSimulationTimeString = FString::Printf(TEXT("Allowed Simulation Time: %s. Keyframe: %d/%d/%d"), *NetworkSim->TickInfo.GetRemainingAllowedSimulationTime().ToString(), NetworkSim->TickInfo.MaxAllowedInputKeyframe, NetworkSim->TickInfo.LastProcessedInputKeyframe, NetworkSim->Buffers.Input.HeadKeyframe());
+			FString AllowedSimulationTimeString = FString::Printf(TEXT("Allowed Simulation Time: %s. Keyframe: %d/%d/%d"), *NetworkSim->TickInfo.GetRemainingAllowedSimulationTime().ToString(), NetworkSim->TickInfo.MaxAllowedKeyframe, NetworkSim->TickInfo.PendingKeyframe, NetworkSim->Buffers.Input.HeadKeyframe());
 			Out.Emit(*AllowedSimulationTimeString, Color);
 		}
 		else if (Owner->GetLocalRole() == ROLE_SimulatedProxy)
diff --git a/Engine/Plugins/Runtime/NetworkPrediction/Source/NetworkPrediction/Public/NetworkSimulationModelReplicators.h b/Engine/Plugins/Runtime/NetworkPrediction/Source/NetworkPrediction/Public/NetworkSimulationModelReplicators.h
--- a/Engine/Plugins/Runtime/NetworkPrediction/Source/NetworkPrediction/Public/NetworkSimulationModelReplicators.h
+++ b/Engine/Plugins/Runtime/NetworkPrediction/Source/NetworkPrediction/Public/NetworkSimulationModelReplicators.h
@@ -227,6 +227,31 @@ struct TReplicator_Single : public TBase
 	}
 };
 
+/** Helper to generate a local input cmd if necessary and advance the simulation's MaxAllowedKeyframe so that it can be processed */
+template<typename TDriver, typename TBufferTypes, typename TTickSettings>
+void GenerateLocalInputCmd(TDriver* Driver, TNetworkSimBufferContainer<TBufferTypes>& Buffers, TSimulationTickState<TTickSettings>& TickInfo)
+{
+	using TInputCmd = typename TBufferTypes::TInputCmd;
+
+	FNetworkSimTime DeltaSimTime = TickInfo.GetRemainingAllowedSimulationTime();
+	if (DeltaSimTime.IsPositive())
+	{
+		if (TInputCmd* InputCmd = Buffers.Input[TickInfo.PendingKeyframe])
+		{
+			InputCmd->SetFrameDeltaTime(DeltaSimTime);
+		}
+		else
+		{
+			InputCmd = Buffers.Input.WriteKeyframe(TickInfo.PendingKeyframe);
+			*InputCmd = TInputCmd();
+			InputCmd->SetFrameDeltaTime(DeltaSimTime);
+			Driver->ProduceInput(DeltaSimTime, *InputCmd);
+		}
+		
+		TickInfo.MaxAllowedKeyframe = TickInfo.PendingKeyframe;
+	}
+}
+
 // -------------------------------------------------------------------------------------------------------
 //	Role based Replicators: these replicators are meant to serve specific roles wrt how the simulation evolves
 // -------------------------------------------------------------------------------------------------------
@@ -242,14 +267,14 @@ struct TReplicator_Server : public TBase
 	{
 		// After receiving input, server may process up to the latest received frames.
 		// (If we needed to buffer input server side for whatever reason, we would do it here)
-		// (also note that we will implicitly guard against speed hacks int he core update loop by not processing cmds past what we have been "allowed")
-		TickInfo.MaxAllowedInputKeyframe = Buffers.Input.HeadKeyframe();
+		// (also note that we will implicitly guard against speed hacks in the core update loop by not processing cmds past what we have been "allowed")
+		TickInfo.MaxAllowedKeyframe = Buffers.Input.HeadKeyframe();
 
 		// Check for gaps in commands
-		if ( TickInfo.LastProcessedInputKeyframe+1 < Buffers.Input.TailKeyframe() )
+		if ( TickInfo.PendingKeyframe < Buffers.Input.TailKeyframe() )
 		{
-			UE_LOG(LogNetworkSim, Warning, TEXT("TReplicator_Server::Reconcile missing inputcmds. LastProcessedInputKeyframe: %d. %s"), TickInfo.LastProcessedInputKeyframe, *Buffers.Input.GetBasicDebugStr());
-			TickInfo.LastProcessedInputKeyframe = Buffers.Input.TailKeyframe()+1;
+			UE_LOG(LogNetworkSim, Warning, TEXT("TReplicator_Server::Reconcile missing inputcmds. PendingKeyframe: %d. %s. This can happen via packet loss"), TickInfo.PendingKeyframe, *Buffers.Input.GetBasicDebugStr());
+			TickInfo.PendingKeyframe = Buffers.Input.TailKeyframe();
 		}
 	}
 
@@ -257,20 +282,9 @@ struct TReplicator_Server : public TBase
 	void PreSimTick(TDriver* Driver, TNetworkSimBufferContainer<TBufferTypes>& Buffers, TSimulationTickState<TTickSettings>& TickInfo, const FNetSimTickParameters& TickParameters)
 	{
 		TickInfo.GiveSimulationTime(TickParameters.LocalDeltaTimeSeconds);
-
 		if (TickParameters.bGenerateLocalInputCmds)
 		{
-			FNetworkSimTime DeltaSimTime = TickInfo.GetRemainingAllowedSimulationTime();
-			if (DeltaSimTime.IsPositive())
-			{
-				if (TInputCmd* InputCmd = Buffers.Input.WriteKeyframe(Buffers.Input.HeadKeyframe() + 1))
-				{
-					*InputCmd = TInputCmd();
-					InputCmd->SetFrameDeltaTime(DeltaSimTime);
-					Driver->ProduceInput(DeltaSimTime, *InputCmd);
-					TickInfo.MaxAllowedInputKeyframe = Buffers.Input.HeadKeyframe();
-				}
-			}
+			GenerateLocalInputCmd(Driver, Buffers, TickInfo);
 		}
 	}
 };
@@ -393,8 +407,8 @@ struct TReplicator_Simulated : public TBase
 
 			check(TickInfo.GetTotalProcessedSimulationTime() <= TickInfo.GetTotalAllowedSimulationTime());
 
-			TickInfo.LastProcessedInputKeyframe = DestinationKeyframe;
-			TickInfo.MaxAllowedInputKeyframe = DestinationKeyframe;	
+			TickInfo.PendingKeyframe = DestinationKeyframe;	// We are about to serialize state to DestinationKeyframe which will be "unprocessed" (has not been used to generate a new frame)
+			TickInfo.MaxAllowedKeyframe = DestinationKeyframe-1; // Do not process PendingKeyframe on our account. ::PreSimTick will advance this based on our interpolation/extrapolation settings
 		}
 
 		check(SyncState && AuxState);
@@ -456,19 +470,7 @@ struct TReplicator_Simulated : public TBase
 
 			if (TickParameters.bGenerateLocalInputCmds)
 			{
-				FNetworkSimTime DeltaSimTime = TickInfo.GetRemainingAllowedSimulationTime();
-				if (DeltaSimTime.IsPositive())
-				{
-					if (TInputCmd* InputCmd = Buffers.Input.WriteKeyframe(Buffers.Input.HeadKeyframe() + 1))
-					{
-						//UE_LOG(LogNetworkSim, Warning, TEXT("   Extrapolating %s"), *DeltaSimTime.ToString());
-
-						*InputCmd = TInputCmd();
-						InputCmd->SetFrameDeltaTime(DeltaSimTime);
-						Driver->ProduceInput(DeltaSimTime, *InputCmd);
-						TickInfo.MaxAllowedInputKeyframe = Buffers.Input.HeadKeyframe();
-					}
-				}
+				GenerateLocalInputCmd(Driver, Buffers, TickInfo);
 			}
 		}
 	}
@@ -496,12 +498,12 @@ struct TReplicator_Simulated : public TBase
 	{
 		// For now, we make the assumption that our last serialized state and time match the parent simulation.
 		// This would not always be the case with low frequency simulated proxies. But could be handled by replicating the simulations together (at the replication level)
-		int32 NewHeadKeyframe = Buffers.Sync.HeadKeyframe()+1;
+		const int32 NewHeadKeyframe = Buffers.Sync.HeadKeyframe()+1;
 		
 		TickInfo.SetTotalProcessedSimulationTime(LastSerializedSimulationTime, NewHeadKeyframe);
 		TickInfo.SetTotalAllowedSimulationTime(LastSerializedSimulationTime);
-		TickInfo.LastProcessedInputKeyframe = NewHeadKeyframe;
-		TickInfo.MaxAllowedInputKeyframe = NewHeadKeyframe;
+		TickInfo.PendingKeyframe = NewHeadKeyframe;
+		TickInfo.MaxAllowedKeyframe = NewHeadKeyframe;
 
 		*Buffers.Sync.WriteKeyframe(NewHeadKeyframe) = LastSerializedSyncState;
 		*Buffers.Input.WriteKeyframe(NewHeadKeyframe) = TInputCmd();
@@ -533,27 +535,28 @@ private:
 
 		ensure(Buffers.Input.HeadKeyframe() == Buffers.Sync.HeadKeyframe());
 
+		const int32 InputKeyframe = Buffers.Input.HeadKeyframe();
+		const int32 OutputKeyframe = InputKeyframe + 1;
+
 		TInputCmd* LastCmd = Buffers.Input.HeadElement();
-		const int32 LastKeyframe = Buffers.Input.HeadKeyframe();
-		const int32 Keyframe = LastKeyframe + 1;
 
 		// Create fake cmd				
-		TInputCmd* NewCmd = Buffers.Input.WriteKeyframe(Keyframe);
+		TInputCmd* NewCmd = Buffers.Input.WriteKeyframe(OutputKeyframe);
 		*NewCmd = LastCmd ? *LastCmd : TInputCmd();
 		NewCmd->SetFrameDeltaTime(DeltaSimTime);	
 		
 		// Create new sync state to write to
-		TSyncState* PrevSyncState = Buffers.Sync.HeadElement();
-		TSyncState* NextSyncState = Buffers.Sync.WriteKeyframe(Keyframe);
-		TAuxState* AuxState = Buffers.Aux.HeadElement();
+		TSyncState* PrevSyncState = Buffers.Sync[InputKeyframe];
+		TSyncState* NextSyncState = Buffers.Sync.WriteKeyframe(OutputKeyframe);
+		TAuxState* AuxState = Buffers.Aux[InputKeyframe];
 
 		// Do the actual update
-		T::Update(Driver, NewCmd->GetFrameDeltaTime().ToRealTimeSeconds(), *NewCmd, *PrevSyncState, *NextSyncState, *AuxState, Buffers.Aux.WriteKeyframeFunc(Keyframe+1));
-		TickInfo.IncrementTotalProcessedSimulationTime(NewCmd->GetFrameDeltaTime(), Buffers.Sync.HeadKeyframe());
+		T::Update(Driver, NewCmd->GetFrameDeltaTime().ToRealTimeSeconds(), *NewCmd, *PrevSyncState, *NextSyncState, *AuxState, Buffers.Aux.WriteKeyframeFunc(OutputKeyframe));
+		TickInfo.IncrementTotalProcessedSimulationTime(NewCmd->GetFrameDeltaTime(), OutputKeyframe);
 
-		// Set our LastProcessedInputKeyframe to fake that we handled it
-		TickInfo.LastProcessedInputKeyframe = Keyframe;
-		TickInfo.MaxAllowedInputKeyframe = Keyframe;
+		// Set our PendingKeyframe to fake that we handled it
+		TickInfo.PendingKeyframe = OutputKeyframe;
+		TickInfo.MaxAllowedKeyframe = OutputKeyframe;
 	}
 	
 	FNetworkSimTime ReconcileSimulationTime;
@@ -589,7 +592,7 @@ struct TReplicator_Autonomous : public TBase
 	void NetSerialize(const FNetSerializeParams& P, TNetworkSimBufferContainer<TBufferTypes>& Buffers, TSimulationTickState<TTickSettings>& TickInfo)
 	{
 		FArchive& Ar = P.Ar;
-
+		
 		SerializedKeyframe = FNetworkSimulationSerialization::SerializeKeyframe(Ar, Buffers.Sync.HeadKeyframe());
 		
 		SerializedTime = TickInfo.GetTotalProcessedSimulationTime();
@@ -619,7 +622,7 @@ struct TReplicator_Autonomous : public TBase
 				if (ClientExistingState->ShouldReconcile(SerializedSyncState) || (NetworkSimulationModelCVars::ForceReconcile() > 0) || (NetworkSimulationModelCVars::ForceReconcileSingle() > 0))
 				{
 					NetworkSimulationModelCVars::SetForceReconcileSingle(0);
-					UE_CLOG(!Buffers.Input.IsValidKeyframe(SerializedKeyframe), LogNetworkSim, Error, TEXT("::NetSerialize: Client InputBuffer does not contain data for frame %d. {%s} {%s}"), SerializedKeyframe, *Buffers.Input.GetBasicDebugStr(), *Buffers.Sync.GetBasicDebugStr());
+					UE_CLOG(!Buffers.Input.IsValidKeyframe(SerializedKeyframe-1), LogNetworkSim, Error, TEXT("::NetSerialize: Client InputBuffer does not contain data for frame %d. {%s} {%s}"), SerializedKeyframe, *Buffers.Input.GetBasicDebugStr(), *Buffers.Sync.GetBasicDebugStr());
 					bPendingReconciliation =  true;
 				}
 			}
@@ -646,21 +649,13 @@ struct TReplicator_Autonomous : public TBase
 	template<typename T, typename TDriver>
 	void Reconcile(TDriver* Driver, TNetworkSimBufferContainer<TBufferTypes>& Buffers, TSimulationTickState<TTickSettings>& TickInfo)
 	{
-		if (bPendingReconciliation == false && bDependentSimulationNeedsReconcile == false)
+		if (bPendingReconciliation == false && (bDependentSimulationNeedsReconcile == false || SerializedKeyframe == INDEX_NONE))
 		{
 			return;
 		}
 		bPendingReconciliation = false;
 		bDependentSimulationNeedsReconcile = false;
 
-		TInputCmd* ClientInputCmd = Buffers.Input[SerializedKeyframe];
-		if (ClientInputCmd == nullptr)
-		{
-			// Fault: no historic command for this frame to reconcile with.
-			UE_LOG(LogNetworkSim, Error, TEXT("Client InputBuffer does not contain data for frame %d. {%s} {%s}"), SerializedKeyframe, *Buffers.Input.GetBasicDebugStr(), *Buffers.Sync.GetBasicDebugStr());
-			return;
-		}
-
 		TSyncState* ClientSyncState = Buffers.Sync[SerializedKeyframe];
 		const bool bDoVisualLog = NetworkSimulationModelCVars::EnableLocalPrediction() > 0; // don't visual log if we have prediction disabled
 		
@@ -694,8 +689,8 @@ struct TReplicator_Autonomous : public TBase
 
 		// Set the canonical simulation time to what we received (we will advance it as we resimulate)
 		TickInfo.SetTotalProcessedSimulationTime(SerializedTime, SerializedKeyframe);
-		TickInfo.LastProcessedInputKeyframe = SerializedKeyframe;
-		TickInfo.MaxAllowedInputKeyframe = FMath::Max(TickInfo.MaxAllowedInputKeyframe, TickInfo.LastProcessedInputKeyframe); // Make sure this doesn't lag behind. This is the only place we should need to do this.
+		TickInfo.PendingKeyframe = SerializedKeyframe;
+		//TickInfo.MaxAllowedKeyframe = FMath::Max(TickInfo.MaxAllowedKeyframe, TickInfo.PendingKeyframe); // Make sure this doesn't lag behind. This is the only place we should need to do this.
 
 		if (NetworkSimulationModelCVars::EnableLocalPrediction() == 0)
 		{
@@ -710,30 +705,31 @@ struct TReplicator_Autonomous : public TBase
 		}
 		
 		// Resimulate all user commands 
-		const int32 LastKeyframeToProcess = TickInfo.MaxAllowedInputKeyframe;
-		for (int32 Keyframe = SerializedKeyframe+1; Keyframe <= LastKeyframeToProcess; ++Keyframe)
+		const int32 LastKeyframeToProcess = TickInfo.MaxAllowedKeyframe;
+		for (int32 Keyframe = SerializedKeyframe; Keyframe <= LastKeyframeToProcess; ++Keyframe)
 		{
+			const int32 OutputKeyframe = Keyframe+1;
+
 			// Keyframe is the frame we are resimulating right now.
 			TInputCmd* ResimulateCmd  = Buffers.Input[Keyframe];
-			TSyncState* PrevSyncState = Buffers.Sync[Keyframe - 1];
-			TSyncState* NextSyncState = Buffers.Sync.WriteKeyframe(Keyframe);
+			TAuxState* AuxState = Buffers.Aux[Keyframe];
+			TSyncState* PrevSyncState = Buffers.Sync[Keyframe];
+			TSyncState* NextSyncState = Buffers.Sync.WriteKeyframe(OutputKeyframe);
 			
 			check(ResimulateCmd);
 			check(PrevSyncState);
 			check(NextSyncState);
-			
-			TAuxState* AuxState = Buffers.Aux[Keyframe];
 			check(AuxState);
 
 			// Log out the Mispredicted state that we are about to overwrite.
 			NextSyncState->VisualLog( FVisualLoggingParameters(Keyframe == LastKeyframeToProcess ? EVisualLoggingContext::LastMispredicted : EVisualLoggingContext::OtherMispredicted, Keyframe, EVisualLoggingLifetime::Persistent), Driver, Driver);
 
 			// Do the actual update
-			T::Update(Driver, ResimulateCmd->GetFrameDeltaTime().ToRealTimeSeconds(), *ResimulateCmd, *PrevSyncState, *NextSyncState, *AuxState, Buffers.Aux.WriteKeyframeFunc(Keyframe+1));
+			T::Update(Driver, ResimulateCmd->GetFrameDeltaTime().ToRealTimeSeconds(), *ResimulateCmd, *PrevSyncState, *NextSyncState, *AuxState, Buffers.Aux.WriteKeyframeFunc(OutputKeyframe));
 			
 			// Update TickInfo
-			TickInfo.IncrementTotalProcessedSimulationTime(ResimulateCmd->GetFrameDeltaTime(), Keyframe);
-			TickInfo.LastProcessedInputKeyframe = Keyframe;
+			TickInfo.IncrementTotalProcessedSimulationTime(ResimulateCmd->GetFrameDeltaTime(), OutputKeyframe);
+			TickInfo.PendingKeyframe = OutputKeyframe;
 
 			// Log out the newly predicted state that we got.
 			NextSyncState->VisualLog( FVisualLoggingParameters(Keyframe == LastKeyframeToProcess ? EVisualLoggingContext::LastPredicted : EVisualLoggingContext::OtherPredicted, Keyframe, EVisualLoggingLifetime::Persistent), Driver, Driver);
@@ -764,17 +760,8 @@ struct TReplicator_Autonomous : public TBase
 			{
 				// Prediction: add simulation time and generate new commands
 				TickInfo.GiveSimulationTime(TickParameters.LocalDeltaTimeSeconds);
-				const FNetworkSimTime DeltaSimTime = TickInfo.GetRemainingAllowedSimulationTime();
-				if (DeltaSimTime.IsPositive())
-				{
-					if (TInputCmd* InputCmd = Buffers.Input.WriteKeyframe(Buffers.Input.HeadKeyframe() + 1))
-					{
-						*InputCmd = TInputCmd();
-						InputCmd->SetFrameDeltaTime(DeltaSimTime);
-						Driver->ProduceInput(DeltaSimTime, *InputCmd);
-						TickInfo.MaxAllowedInputKeyframe = Buffers.Input.HeadKeyframe(); // Allow the new command to be processed by the local simulation
-					}
-				}
+
+				GenerateLocalInputCmd(Driver, Buffers, TickInfo);
 			}
 			else
 			{
diff --git a/Engine/Plugins/Runtime/NetworkPrediction/Source/NetworkPrediction/Public/NetworkSimulationModelTypes.h b/Engine/Plugins/Runtime/NetworkPrediction/Source/NetworkPrediction/Public/NetworkSimulationModelTypes.h
--- a/Engine/Plugins/Runtime/NetworkPrediction/Source/NetworkPrediction/Public/NetworkSimulationModelTypes.h
+++ b/Engine/Plugins/Runtime/NetworkPrediction/Source/NetworkPrediction/Public/NetworkSimulationModelTypes.h
@@ -206,8 +206,8 @@ private:
 template<typename T>
 struct TPredictedStateAccessor
 {
-	TPredictedStateAccessor(TBaseStateAccessor<T>& InModBuffer, int32* InLastProcessedKeyframe)
-		: ModBuffer(InModBuffer), LastProcessedKeyframe(InLastProcessedKeyframe) { }
+	TPredictedStateAccessor(TBaseStateAccessor<T>& InModBuffer, int32* InPendingKeyframe)
+		: ModBuffer(InModBuffer), PendingKeyframe(InPendingKeyframe) { }
 
 	// Returns the current state, cannot be modified directly
 	const T* Get() const
@@ -218,13 +218,13 @@ struct TPredictedStateAccessor
 	// Applies a mod to the current state. If in prediction mode, will save the mod to be reapplied during a rollback
 	void Modify(TUniqueFunction<void(T&)>&& Func)
 	{
-		ModBuffer.Modify(MoveTemp(Func), *LastProcessedKeyframe);
+		ModBuffer.Modify(MoveTemp(Func), *PendingKeyframe);
 	}
 
 private:
 
 	TBaseStateAccessor<T>& ModBuffer;
-	int32* LastProcessedKeyframe = nullptr;
+	int32* PendingKeyframe = nullptr;
 };
 
 // Struct that encapsulates writing a new element to a buffer. This is used to allow a new Aux state to be created in the ::Update loop.
@@ -443,17 +443,27 @@ private:
 };
 
 // ----------------------------------------------------------------------------------------------------------------------------------------------
+//	TSimulationTickState: Holds active state for simulation ticking. We track two things: keyframes and time.
 //
+//	PendingKeyframe is the next keyframe we will process: the input/sync/aux state @ PendingKeyframe will be run through T::Update and produce the 
+//	next frame's (PendingKeyframe+1) Sync and possibly Aux state (if it changes). "Out of band" modifications to the sync/aux state should happen
+//	to PendingKeyframe (e.g, before it is processed. Once a frame has been processed, we won't run it through T::Update again!).
+//
+//	MaxAllowedKeyframe is a keyframe based limiter on simulation updates. This must be incremented to allow the simulation to advance.
+//
+//	Time is also tracked. We keep running total for how much the sim has advanced and how much it is allowed to advance. There is also a historic buffer of
+//	simulation time in SimulationTimeBuffer.
+//
+//	Consider that Keyframes are essentially client dependent and gaps can happen due to packet loss, etc. Time will always be continuous though.
+//	
 // ----------------------------------------------------------------------------------------------------------------------------------------------
-
-// Holds active state for simulation ticking: what inputs we have processed, how time we have simulated, how much we are allowed to, etc.
 template<typename TickSettings=TNetworkSimTickSettings<>>
 struct TSimulationTickState
 {
 	using TSettings = TickSettings;
 
-	int32 LastProcessedInputKeyframe = 0;	// The last input keyframe that we processed
-	int32 MaxAllowedInputKeyframe = 0;		// The max input keyframe that we are allowed to process (e.g, don't process input past this keyframe yet)
+	int32 PendingKeyframe = 0;
+	int32 MaxAllowedKeyframe = -1;
 	
 	FNetworkSimTime GetTotalProcessedSimulationTime() const 
 	{ 
@@ -557,12 +567,6 @@ class TNetworkSimDriverInterfaceBase
 public:
 	virtual FString GetDebugName() const = 0; // Used for debugging. Recommended to emit the simulation name and the actor name/role.
 	virtual const UObject* GetVLogOwner() const = 0; // Owning object for Visual Logs
-
-	// Called to create initial value of the sync state.
-	virtual void InitSyncState(typename TBufferTypes::TSyncState& OutSyncState) const = 0;
-
-	// Called to create initial value of the aux state.
-	virtual void InitAuxState(typename TBufferTypes::TAuxState& OutAuxState) const = 0;
 	
 	// Called whenever the sim is ready to process new local input.
 	virtual void ProduceInput(const FNetworkSimTime SimTime, typename TBufferTypes::TInputCmd&) = 0;
diff --git a/Engine/Plugins/Runtime/NetworkPrediction/readme.txt b/Engine/Plugins/Runtime/NetworkPrediction/readme.txt
--- a/Engine/Plugins/Runtime/NetworkPrediction/readme.txt
+++ b/Engine/Plugins/Runtime/NetworkPrediction/readme.txt
@@ -105,6 +105,14 @@ Smoothing: Taking the output of the simulation and applying an additional layer
 // Release notes
 // ----------------------------------------------------------------------------------------------------------
 
+
+Update (10-16-19)
+-Cleanup how we advance the simulation and unified which keyframes are accessed across the three main buffers
+-LastProcessedInputKeyframe -> PendingKeyframe. Code is simplified dealing with the "next" keyframe rather than making assumptions based on "last" keyframe.
+-Input/Sync/Aux states @ PendingKeyframe are what gets used as inputs to T::Update. This is more consistent that how it worked previously, where it was essentially (Input+1/Sync) -> (Sync+1).
+-This gets rid of awkward hacks like empty input cmd @ keyframe = 0.
+-Also got rid of InitSync/InitAuxState interface functions. The initial states can now be provided when instantiating the network sim model.
+
 Update (10-15-19)
 -Aux buffer hooked up. Still some decisions to make here, expect a few more refactors.
 -TReplicationBuffer replaced with simpler API.
diff --git a/Engine/Plugins/Runtime/NetworkPredictionExtras/Source/NetworkPredictionExtras/Private/MockNetworkSimulation.cpp b/Engine/Plugins/Runtime/NetworkPredictionExtras/Source/NetworkPredictionExtras/Private/MockNetworkSimulation.cpp
--- a/Engine/Plugins/Runtime/NetworkPredictionExtras/Source/NetworkPredictionExtras/Private/MockNetworkSimulation.cpp
+++ b/Engine/Plugins/Runtime/NetworkPredictionExtras/Source/NetworkPredictionExtras/Private/MockNetworkSimulation.cpp
@@ -134,6 +134,9 @@ UMockNetworkSimulationComponent::UMockNetworkSimulationComponent()
 
 INetworkSimulationModel* UMockNetworkSimulationComponent::InstantiateNetworkSimulation()
 {
+	FMockSyncState InitialSyncState;
+	InitialSyncState.Total = MockValue;
+
 	auto* NewSim = new FMockNetworkModel(this);
 	DO_NETSIM_MODEL_DEBUG(FNetworkSimulationModelDebuggerManager::Get().RegisterNetworkSimulationModel(NewSim, GetOwner()));
 	return NewSim;
@@ -166,16 +169,6 @@ void UMockNetworkSimulationComponent::ProduceInput(const FNetworkSimTime SimTime
 	}
 }
 
-void UMockNetworkSimulationComponent::InitSyncState(FMockSyncState& OutSyncState) const
-{
-	OutSyncState.Total = MockValue;
-}
-
-void UMockNetworkSimulationComponent::InitAuxState(FMockAuxState& OutAuxState) const
-{
-
-}
-
 void UMockNetworkSimulationComponent::FinalizeFrame(const FMockSyncState& SyncState, const FMockAuxState& AuxState)
 {
 	MockValue = SyncState.Total;
diff --git a/Engine/Plugins/Runtime/NetworkPredictionExtras/Source/NetworkPredictionExtras/Public/MockNetworkSimulation.h b/Engine/Plugins/Runtime/NetworkPredictionExtras/Source/NetworkPredictionExtras/Public/MockNetworkSimulation.h
--- a/Engine/Plugins/Runtime/NetworkPredictionExtras/Source/NetworkPredictionExtras/Public/MockNetworkSimulation.h
+++ b/Engine/Plugins/Runtime/NetworkPredictionExtras/Source/NetworkPredictionExtras/Public/MockNetworkSimulation.h
@@ -171,9 +171,7 @@ public:
 
 	FString GetDebugName() const override;
 	virtual UObject* GetVLogOwner() const override final;
-
-	void InitSyncState(FMockSyncState& OutSyncState) const override;
-	void InitAuxState(FMockAuxState& OutAuxState) const override;
+	
 	void ProduceInput(const FNetworkSimTime SimTime, FMockInputCmd& Cmd);
 	void FinalizeFrame(const FMockSyncState& SyncState, const FMockAuxState& AuxState) override;
 	
