Network Prediction Plugin:
-Hook up aux buffer to debug hud
-Fix some uninitalized variables
-Test map improvements

#rnx

diff --git a/Engine/Plugins/Runtime/NetworkPrediction/Source/NetworkPrediction/Public/Movement/FlyingMovement.h b/Engine/Plugins/Runtime/NetworkPrediction/Source/NetworkPrediction/Public/Movement/FlyingMovement.h
--- a/Engine/Plugins/Runtime/NetworkPrediction/Source/NetworkPrediction/Public/Movement/FlyingMovement.h
+++ b/Engine/Plugins/Runtime/NetworkPrediction/Source/NetworkPrediction/Public/Movement/FlyingMovement.h
@@ -36,6 +36,11 @@ namespace FlyingMovement
 		FRotator RotationInput;
 		FVector MovementInput;
 
+		FInputCmd()
+			: RotationInput(ForceInitToZero)
+			, MovementInput(ForceInitToZero)
+		{ }
+
 		void NetSerialize(const FNetSerializeParams& P)
 		{
 			P.Ar << RotationInput;
@@ -63,6 +68,12 @@ namespace FlyingMovement
 		FVector Velocity;
 		FRotator Rotation;
 
+		FMoveState()
+			: Location(ForceInitToZero)
+			, Velocity(ForceInitToZero)
+			, Rotation(ForceInitToZero)
+		{ }
+
 		void NetSerialize(const FNetSerializeParams& P)
 		{
 			P.Ar << Location;
@@ -117,6 +128,21 @@ namespace FlyingMovement
 			P.Ar << Deceleration;
 			P.Ar << Acceleration;
 		}
+
+		void Log(FStandardLoggingParameters& P) const
+		{
+			if (P.Context == EStandardLoggingContext::HeaderOnly)
+			{
+				P.Ar->Logf(TEXT(" %d "), P.Keyframe);
+			}
+			else if (P.Context == EStandardLoggingContext::Full)
+			{
+				P.Ar->Logf(TEXT("MaxSpeed: %.2f"), MaxSpeed);
+				P.Ar->Logf(TEXT("TurningBoost: %.2f"), TurningBoost);
+				P.Ar->Logf(TEXT("Deceleration: %.2f"), Deceleration);
+				P.Ar->Logf(TEXT("Acceleration: %.2f"), Acceleration);
+			}
+		}
 	};
 
 	using TMovementBufferTypes = TNetworkSimBufferTypes<FInputCmd, FMoveState, FAuxState>;
@@ -154,9 +180,6 @@ namespace FlyingMovement
 	template<int32 InFixedStepMS=0>
 	using FMovementSystem = TNetworkedSimulationModel<FMovementSimulation, IMovementDriver, TMovementBufferTypes, TNetworkSimTickSettings<InFixedStepMS> >;
 
-	// Simulation time used by the movement system
-	//using TSimTime = FMovementSystem::TSimTime;
-
 	// general tolerance value for rotation checks
 	static const float ROTATOR_TOLERANCE = (1e-3);
 
diff --git a/Engine/Plugins/Runtime/NetworkPrediction/Source/NetworkPrediction/Public/Movement/ParametricMovement.h b/Engine/Plugins/Runtime/NetworkPrediction/Source/NetworkPrediction/Public/Movement/ParametricMovement.h
--- a/Engine/Plugins/Runtime/NetworkPrediction/Source/NetworkPrediction/Public/Movement/ParametricMovement.h
+++ b/Engine/Plugins/Runtime/NetworkPrediction/Source/NetworkPrediction/Public/Movement/ParametricMovement.h
@@ -91,6 +91,18 @@ namespace ParametricMovement
 		{
 			P.Ar << Multiplier;
 		}
+
+		void Log(FStandardLoggingParameters& Params) const
+		{
+			if (Params.Context == EStandardLoggingContext::HeaderOnly)
+			{
+				Params.Ar->Logf(TEXT(" %d "), Params.Keyframe);
+			}
+			else if (Params.Context == EStandardLoggingContext::Full)
+			{
+				Params.Ar->Logf(TEXT("Multiplier: %f"), Multiplier);
+			}
+		}
 	};
 
 	using TMovementBufferTypes = TNetworkSimBufferTypes<FInputCmd, FMoveState, FAuxState>;
diff --git a/Engine/Plugins/Runtime/NetworkPrediction/Source/NetworkPrediction/Public/NetworkSimulationModelDebugger.h b/Engine/Plugins/Runtime/NetworkPrediction/Source/NetworkPrediction/Public/NetworkSimulationModelDebugger.h
--- a/Engine/Plugins/Runtime/NetworkPrediction/Source/NetworkPrediction/Public/NetworkSimulationModelDebugger.h
+++ b/Engine/Plugins/Runtime/NetworkPrediction/Source/NetworkPrediction/Public/NetworkSimulationModelDebugger.h
@@ -588,6 +588,7 @@ struct TNetworkSimulationModelDebugger : public INetworkSimulationModelDebugger
 
 		EmitBuffer(TEXT("InputBuffer"), NetworkSim->Buffers.Input);
 		EmitBuffer(TEXT("SyncBuffer"), NetworkSim->Buffers.Sync);
+		EmitBuffer(TEXT("AuxBuffer"), NetworkSim->Buffers.Aux);
 
 		// ------------------------------------------------------------------------------------------------------------------------------------------------
 		//	Canvas
diff --git a/Engine/Plugins/Runtime/NetworkPrediction/Source/NetworkPrediction/Public/NetworkSimulationModelTypes.h b/Engine/Plugins/Runtime/NetworkPrediction/Source/NetworkPrediction/Public/NetworkSimulationModelTypes.h
--- a/Engine/Plugins/Runtime/NetworkPrediction/Source/NetworkPrediction/Public/NetworkSimulationModelTypes.h
+++ b/Engine/Plugins/Runtime/NetworkPrediction/Source/NetworkPrediction/Public/NetworkSimulationModelTypes.h
@@ -167,7 +167,7 @@ struct TPredictedModBuffer : public TBaseStateAccessor<typename TBuffer::Element
 
 	void Modify(TUniqueFunction<void(ElementType&)>&& Func, int32 Keyframe) final override
 	{
-		ElementType* State = Buffer[Keyframe];
+		ElementType* State = Buffer.WriteKeyframe(Keyframe);
 		check(State);
 		Func(*State);
 
diff --git a/Engine/Plugins/Runtime/NetworkPredictionExtras/Source/NetworkPredictionExtras/Public/MockNetworkSimulation.h b/Engine/Plugins/Runtime/NetworkPredictionExtras/Source/NetworkPredictionExtras/Public/MockNetworkSimulation.h
--- a/Engine/Plugins/Runtime/NetworkPredictionExtras/Source/NetworkPredictionExtras/Public/MockNetworkSimulation.h
+++ b/Engine/Plugins/Runtime/NetworkPredictionExtras/Source/NetworkPredictionExtras/Public/MockNetworkSimulation.h
@@ -117,6 +117,18 @@ struct FMockAuxState
 	{
 		P.Ar << Multiplier;
 	}
+
+	void Log(FStandardLoggingParameters& Params) const
+	{
+		if (Params.Context == EStandardLoggingContext::HeaderOnly)
+		{
+			Params.Ar->Logf(TEXT(" %d "), Params.Keyframe);
+		}
+		else if (Params.Context == EStandardLoggingContext::Full)
+		{
+			Params.Ar->Logf(TEXT("Multiplier: %f"), Multiplier);
+		}
+	}
 };
 
 using TMockNetworkSimulationBufferTypes = TNetworkSimBufferTypes<FMockInputCmd, FMockSyncState, FMockAuxState>;
