Network Prediction Plugin
Continued WIP on game code interface to sync/aux states

#rnx

diff --git a/Engine/Plugins/Runtime/NetworkPrediction/Source/NetworkPrediction/Private/Movement/FlyingMovementComponent.cpp b/Engine/Plugins/Runtime/NetworkPrediction/Source/NetworkPrediction/Private/Movement/FlyingMovementComponent.cpp
--- a/Engine/Plugins/Runtime/NetworkPrediction/Source/NetworkPrediction/Private/Movement/FlyingMovementComponent.cpp
+++ b/Engine/Plugins/Runtime/NetworkPrediction/Source/NetworkPrediction/Private/Movement/FlyingMovementComponent.cpp
@@ -80,8 +80,8 @@ INetworkSimulationModel* UFlyingMovementComponent::InstantiateNetworkSimulation(
 
 	auto NewSim = new FlyingMovement::FMovementSystem<0>(this, InitialSyncState, InitialAuxState);
 	DO_NETSIM_MODEL_DEBUG(FNetworkSimulationModelDebuggerManager::Get().RegisterNetworkSimulationModel(NewSim, GetOwner()));
-	MovementSyncState = &NewSim->SyncAccessor;
-	MovementAuxState = &NewSim->AuxAccessor;
+	MovementSyncState.Init(NewSim);
+	MovementAuxState.Init(NewSim);
 	return NewSim;
 }
 
diff --git a/Engine/Plugins/Runtime/NetworkPrediction/Source/NetworkPrediction/Private/Movement/ParametricMovement.cpp b/Engine/Plugins/Runtime/NetworkPrediction/Source/NetworkPrediction/Private/Movement/ParametricMovement.cpp
--- a/Engine/Plugins/Runtime/NetworkPrediction/Source/NetworkPrediction/Private/Movement/ParametricMovement.cpp
+++ b/Engine/Plugins/Runtime/NetworkPrediction/Source/NetworkPrediction/Private/Movement/ParametricMovement.cpp
@@ -114,16 +114,16 @@ INetworkSimulationModel* UParametricMovementComponent::InstantiateNetworkSimulat
 		auto *NewSim = new ParametricMovement::FMovementSystem<16>(this);
 		NewSim->RepProxy_Simulated.bAllowSimulatedExtrapolation = !bEnableInterpolation;
 		DO_NETSIM_MODEL_DEBUG(FNetworkSimulationModelDebuggerManager::Get().RegisterNetworkSimulationModel(NewSim, GetOwner()));
-		MovementSyncState = &NewSim->SyncAccessor;
-		MovementAuxState = &NewSim->AuxAccessor;
+		//MovementSyncState = &NewSim->SyncAccessor;
+		//MovementAuxState = &NewSim->AuxAccessor;
 		return NewSim;
 	}
 	
 	auto *NewSim = new ParametricMovement::FMovementSystem<>(this);
 	NewSim->RepProxy_Simulated.bAllowSimulatedExtrapolation = !bEnableInterpolation;
 	DO_NETSIM_MODEL_DEBUG(FNetworkSimulationModelDebuggerManager::Get().RegisterNetworkSimulationModel(NewSim, GetOwner()));
-	MovementSyncState = &NewSim->SyncAccessor;
-	MovementAuxState = &NewSim->AuxAccessor;
+	//MovementSyncState = &NewSim->SyncAccessor;
+	//MovementAuxState = &NewSim->AuxAccessor;
 	return NewSim;
 }
 
diff --git a/Engine/Plugins/Runtime/NetworkPrediction/Source/NetworkPrediction/Private/NetworkPredictionTypes.cpp b/Engine/Plugins/Runtime/NetworkPrediction/Source/NetworkPrediction/Private/NetworkPredictionTypes.cpp
--- a/Engine/Plugins/Runtime/NetworkPrediction/Source/NetworkPrediction/Private/NetworkPredictionTypes.cpp
+++ b/Engine/Plugins/Runtime/NetworkPrediction/Source/NetworkPrediction/Private/NetworkPredictionTypes.cpp
@@ -134,7 +134,7 @@ void FNetSimTickParameters::InitFromActor(AActor* Actor)
 	Role = Actor->GetLocalRole();
 	const bool bHasNetConnection = Actor->GetNetConnection() != nullptr;
 
-	// The only default case we really don;t want generating local input is on the server when there is a net owning player (auto proxy)
+	// The only default case we really don't want generating local input is on the server when there is a net owning player (auto proxy)
 	// This implies that we will (sim) extrapolate by default. To opt out of extrapolation, set to false or you can just not tick the net sim
 	bGenerateLocalInputCmds = !(Role == ROLE_Authority && bHasNetConnection);
 }
\ No newline at end of file
diff --git a/Engine/Plugins/Runtime/NetworkPrediction/Source/NetworkPrediction/Public/Movement/FlyingMovement.h b/Engine/Plugins/Runtime/NetworkPrediction/Source/NetworkPrediction/Public/Movement/FlyingMovement.h
--- a/Engine/Plugins/Runtime/NetworkPrediction/Source/NetworkPrediction/Public/Movement/FlyingMovement.h
+++ b/Engine/Plugins/Runtime/NetworkPrediction/Source/NetworkPrediction/Public/Movement/FlyingMovement.h
@@ -217,8 +217,8 @@ public:
 	DECLARE_DELEGATE_TwoParams(FProduceFlyingInput, const FNetworkSimTime /*SimTime*/, FlyingMovement::FInputCmd& /*Cmd*/)
 	FProduceFlyingInput ProduceInputDelegate;
 
-	TPredictedStateAccessor<FlyingMovement::FMoveState>* MovementSyncState = nullptr;
-	TPredictedStateAccessor<FlyingMovement::FAuxState>*	 MovementAuxState = nullptr;
+	TNetworkSimStateAccessor<FlyingMovement::FMoveState> MovementSyncState;
+	TNetworkSimStateAccessor<FlyingMovement::FAuxState> MovementAuxState;
 
 protected:
 
diff --git a/Engine/Plugins/Runtime/NetworkPrediction/Source/NetworkPrediction/Public/Movement/ParametricMovement.h b/Engine/Plugins/Runtime/NetworkPrediction/Source/NetworkPrediction/Public/Movement/ParametricMovement.h
--- a/Engine/Plugins/Runtime/NetworkPrediction/Source/NetworkPrediction/Public/Movement/ParametricMovement.h
+++ b/Engine/Plugins/Runtime/NetworkPrediction/Source/NetworkPrediction/Public/Movement/ParametricMovement.h
@@ -175,10 +175,6 @@ class NETWORKPREDICTION_API UParametricMovementComponent : public UBaseMovementC
 	virtual void AdvanceParametricTime(const float InPosition, const float InPlayRate, float &OutPosition, float& OutPlayRate, const float DeltaTimeSeconds) const override;
 	virtual void MapTimeToTransform(const float InPosition, FTransform& OutTransform) const override;
 
-
-	TPredictedStateAccessor<ParametricMovement::FMoveState>* MovementSyncState = nullptr;
-	TPredictedStateAccessor<ParametricMovement::FAuxState>*	 MovementAuxState = nullptr;
-
 protected:
 
 	virtual INetworkSimulationModel* InstantiateNetworkSimulation() override;
diff --git a/Engine/Plugins/Runtime/NetworkPrediction/Source/NetworkPrediction/Public/NetworkSimulationModel.h b/Engine/Plugins/Runtime/NetworkPrediction/Source/NetworkPrediction/Public/NetworkSimulationModel.h
--- a/Engine/Plugins/Runtime/NetworkPrediction/Source/NetworkPrediction/Public/NetworkSimulationModel.h
+++ b/Engine/Plugins/Runtime/NetworkPrediction/Source/NetworkPrediction/Public/NetworkSimulationModel.h
@@ -18,6 +18,13 @@
 //		-Aux: State that is also an input into the simulation but does not intrinsically evolve from to frame. Changes to this state can be trapped/tracked/predicted.
 //		-Debug: Replicated buffer from server->client with server-frame centered debug information. Compiled out of shipping builds.
 //
+//	* How other code interacts with this:
+//		-Network updates will come in through UE4 networking -> FReplicationProxy (on actor/component) -> TNetworkedSimulationModel::RepProxy_* -> Buffers.*
+//		-UNetworkSimulationGlobalManager: responsible for ticking simulation (after recv net traffic, prior to UE4 actor ticking)
+//		-External game code can interact with the system:
+//			-The TNetworkedSimulationModel is mostly public and exposed. It is not recommend to publicly expose to your "user" code (high level scripting, designers, etc).
+//			-TNetworkSimStateAccessor is a helper for safely reading/writing to state within the system.
+//
 // ----------------------------------------------------------------------------------------------------------------------------------------------------------------------
 
 template <
@@ -53,8 +60,6 @@ public:
 	using TRealTime = FNetworkSimTime::FRealTime;
 	
 	TNetworkedSimulationModel(TDriver* InDriver, const TSyncState& InitialSyncState = TSyncState(), const TAuxState& InitialAuxState = TAuxState())
-		: SyncAccessor(Buffers.SyncMods, &TickInfo.PendingKeyframe)
-		, AuxAccessor(Buffers.AuxMods, &TickInfo.PendingKeyframe)
 	{
 		Driver = InDriver;
 		*Buffers.Sync.WriteKeyframe(0) = InitialSyncState;
@@ -148,10 +153,7 @@ public:
 					// We don't have valid sync state for this frame. This could mean we skipped ahead somehow, such as heavy packet loss where the server doesn't receive all input cmds
 					// We will use the last known sync state, but first write it out to the current InputKeyframe
 					UE_LOG(LogNetworkSim, Log, TEXT("%s: No sync state found @ keyframe %d. Using previous head element @ keyframe %d."), *Driver->GetDebugName(), InputKeyframe, Buffers.Sync.HeadKeyframe());
-					
-					TSyncState LastSyncState = *Buffers.Sync.HeadElement();
-					InSyncState = Buffers.Sync.WriteKeyframe(InputKeyframe);
-					*InSyncState = LastSyncState;
+					InSyncState = Buffers.Sync.WriteKeyframeInitializedFromHead(InputKeyframe);
 				}
 
 				const TAuxState* InAuxState = Buffers.Aux[InputKeyframe];
@@ -160,10 +162,10 @@ public:
 				TSyncState* OutSyncState = Buffers.Sync.WriteKeyframe(OutputKeyframe);
 				check(OutSyncState);
 
-				TSimulation::Update(Driver, InputCmd->GetFrameDeltaTime().ToRealTimeSeconds(), *InputCmd, *InSyncState, *OutSyncState, *InAuxState, Buffers.Aux.WriteKeyframeFunc(OutputKeyframe));
-					
-				TickInfo.IncrementTotalProcessedSimulationTime(InputCmd->GetFrameDeltaTime(), OutputKeyframe);
-				TickInfo.PendingKeyframe = OutputKeyframe;
+				{
+					TScopedSimulationTick UpdateScope(TickInfo, OutputKeyframe, InputCmd->GetFrameDeltaTime());
+					TSimulation::Update(Driver, InputCmd->GetFrameDeltaTime().ToRealTimeSeconds(), *InputCmd, *InSyncState, *OutSyncState, *InAuxState, Buffers.Aux.WriteKeyframeFunc(OutputKeyframe));
+				}
 
 				if (DebugState)
 				{
@@ -260,9 +262,6 @@ public:
 		}
 
 		//TickInfo.InitSimulationTimeBuffer(Parameters.SyncedBufferSize);
-
-		Buffers.SyncMods.SetIsAuthority(Role == ROLE_Authority);
-		Buffers.AuxMods.SetIsAuthority(Role == ROLE_Authority);
 	}
 
 	void NetSerializeProxy(EReplicationProxyTarget Target, const FNetSerializeParams& Params) final override
@@ -382,9 +381,6 @@ public:
 
 	TNetworkSimBufferContainer<TBufferTypes> Buffers;
 
-	TPredictedStateAccessor<TSyncState> SyncAccessor;
-	TPredictedStateAccessor<TAuxState> AuxAccessor;
-
 	TRepProxyServerRPC RepProxy_ServerRPC;
 	TRepProxyAutonomous RepProxy_Autonomous;
 	TRepProxySimulated RepProxy_Simulated;
diff --git a/Engine/Plugins/Runtime/NetworkPrediction/Source/NetworkPrediction/Public/NetworkSimulationModelBuffer.h b/Engine/Plugins/Runtime/NetworkPrediction/Source/NetworkPrediction/Public/NetworkSimulationModelBuffer.h
--- a/Engine/Plugins/Runtime/NetworkPrediction/Source/NetworkPrediction/Public/NetworkSimulationModelBuffer.h
+++ b/Engine/Plugins/Runtime/NetworkPrediction/Source/NetworkPrediction/Public/NetworkSimulationModelBuffer.h
@@ -39,6 +39,31 @@ struct TNetworkSimBufferBase
 			*((T*)this)->WriteKeyframe( It.Keyframe() ) = *It.Element();
 		}
 	}
+protected:
+
+	// Creates a new keyframe, but
+	//	-If keyframe already exists, returns existing
+	//	-If keyframe > head, contents of head are copied into new frame
+	//	-If keyframe < head, contents are cleared to default value
+	template<typename ElementType>
+	ElementType* WriteKeyframeInitializedFromHeadImpl(int32 Keyframe)
+	{
+		if (((T*)this)->HeadKeyframe() > Keyframe)
+		{
+			ElementType* HeadElementPtr = ((T*)this)->HeadElement();
+			ElementType* NewElement = ((T*)this)->WriteKeyframe(Keyframe);
+			*NewElement = *HeadElementPtr;
+			return NewElement;
+		}
+		else if (Keyframe < ((T*)this)->TailKeyframe())
+		{
+			ElementType* NewElement = ((T*)this)->WriteKeyframe(Keyframe);
+			*NewElement = ElementType();
+			return NewElement;
+		}
+		
+		return ((T*)this)->WriteKeyframe(Keyframe);
+	}
 };
 
 template<typename T, int32 NumElements=32>
@@ -89,6 +114,11 @@ struct TNetworkSimContiguousBuffer : public TNetworkSimBufferBase<TNetworkSimCon
 		++DirtyCount;
 		return &Data[Keyframe % Data.Num()];
 	}
+	
+	T* WriteKeyframeInitializedFromHead(int32 Keyframe)
+	{
+		return this->template WriteKeyframeInitializedFromHeadImpl<ElementType>(Keyframe);
+	}
 
 	TUniqueFunction<T*()> WriteKeyframeFunc(int32 Keyframe)
 	{
@@ -213,6 +243,11 @@ struct TNetworkSimSparseBuffer : public TNetworkSimBufferBase<TNetworkSimSparseB
 		return &WriteData.Element;
 	}
 
+	ElementType* WriteKeyframeInitializedFromHead(int32 Keyframe)
+	{
+		return this->template WriteKeyframeInitializedFromHeadImpl<ElementType>(Keyframe);
+	}
+
 	TUniqueFunction<T*()> WriteKeyframeFunc(int32 Keyframe)
 	{
 		return [this, Keyframe]() { return WriteKeyframe(Keyframe); };
diff --git a/Engine/Plugins/Runtime/NetworkPrediction/Source/NetworkPrediction/Public/NetworkSimulationModelReplicators.h b/Engine/Plugins/Runtime/NetworkPrediction/Source/NetworkPrediction/Public/NetworkSimulationModelReplicators.h
--- a/Engine/Plugins/Runtime/NetworkPrediction/Source/NetworkPrediction/Public/NetworkSimulationModelReplicators.h
+++ b/Engine/Plugins/Runtime/NetworkPrediction/Source/NetworkPrediction/Public/NetworkSimulationModelReplicators.h
@@ -227,27 +227,34 @@ struct TReplicator_Single : public TBase
 	}
 };
 
-/** Helper to generate a local input cmd if necessary and advance the simulation's MaxAllowedKeyframe so that it can be processed */
+/** Helper that writes a new input cmd to the input buffer, at given Keyframe (usually the sim's PendingKeyframe). If keyframe doesn't exist, ProduceInput is called on the driver. */
+template<typename TDriver, typename TBufferTypes>
+void GenerateLocalInputCmdAtKeyframe(TDriver* Driver, TNetworkSimBufferContainer<TBufferTypes>& Buffers, const FNetworkSimTime& DeltaSimTime, int32 Keyframe)
+{
+	using TInputCmd = typename TBufferTypes::TInputCmd;
+	if (TInputCmd* InputCmd = Buffers.Input[Keyframe])
+	{
+		InputCmd->SetFrameDeltaTime(DeltaSimTime);
+	}
+	else
+	{
+		InputCmd = Buffers.Input.WriteKeyframe(Keyframe);
+		*InputCmd = TInputCmd();
+		InputCmd->SetFrameDeltaTime(DeltaSimTime);
+		Driver->ProduceInput(DeltaSimTime, *InputCmd);
+	}
+}
+
+/** Helper to generate a local input cmd if we have simulation time to spend and advance the simulation's MaxAllowedKeyframe so that it can be processed. */
 template<typename TDriver, typename TBufferTypes, typename TTickSettings>
-void GenerateLocalInputCmd(TDriver* Driver, TNetworkSimBufferContainer<TBufferTypes>& Buffers, TSimulationTickState<TTickSettings>& TickInfo)
+void TryGenerateLocalInput(TDriver* Driver, TNetworkSimBufferContainer<TBufferTypes>& Buffers, TSimulationTickState<TTickSettings>& TickInfo)
 {
 	using TInputCmd = typename TBufferTypes::TInputCmd;
 
 	FNetworkSimTime DeltaSimTime = TickInfo.GetRemainingAllowedSimulationTime();
 	if (DeltaSimTime.IsPositive())
 	{
-		if (TInputCmd* InputCmd = Buffers.Input[TickInfo.PendingKeyframe])
-		{
-			InputCmd->SetFrameDeltaTime(DeltaSimTime);
-		}
-		else
-		{
-			InputCmd = Buffers.Input.WriteKeyframe(TickInfo.PendingKeyframe);
-			*InputCmd = TInputCmd();
-			InputCmd->SetFrameDeltaTime(DeltaSimTime);
-			Driver->ProduceInput(DeltaSimTime, *InputCmd);
-		}
-		
+		GenerateLocalInputCmdAtKeyframe(Driver, Buffers, DeltaSimTime, TickInfo.PendingKeyframe);
 		TickInfo.MaxAllowedKeyframe = TickInfo.PendingKeyframe;
 	}
 }
@@ -284,7 +291,7 @@ struct TReplicator_Server : public TBase
 		TickInfo.GiveSimulationTime(TickParameters.LocalDeltaTimeSeconds);
 		if (TickParameters.bGenerateLocalInputCmds)
 		{
-			GenerateLocalInputCmd(Driver, Buffers, TickInfo);
+			TryGenerateLocalInput(Driver, Buffers, TickInfo);
 		}
 	}
 };
@@ -434,14 +441,11 @@ struct TReplicator_Simulated : public TBase
 
 		if (bAllowSimulatedExtrapolation && ParentSimulation == nullptr && NetworkSimulationModelCVars::EnableSimulatedExtrapolation() && NetworkSimulationModelCVars::EnableSimulatedReconcile())
 		{
-			TInputCmd* LastCmd = Buffers.Input.HeadElement();
-
 			// Simulated Reconcile requires the input buffer to be kept up to date with the Sync buffer
 			// Generate a new, fake, command since we just added a new sync state to head
 			while (Buffers.Input.HeadKeyframe() < Buffers.Sync.HeadKeyframe())
 			{
-				TInputCmd* Next = Buffers.Input.WriteKeyframe(Buffers.Input.HeadKeyframe()+1);
-				*Next = LastCmd ? *LastCmd : TInputCmd();
+				TInputCmd* Next = Buffers.Input.WriteKeyframeInitializedFromHead(Buffers.Input.HeadKeyframe()+1);
 			}
 
 			// Do we have time to make up? We may have extrapolated ahead of the server (totally fine - can happen with small amount of latency variance)
@@ -470,7 +474,7 @@ struct TReplicator_Simulated : public TBase
 
 			if (TickParameters.bGenerateLocalInputCmds)
 			{
-				GenerateLocalInputCmd(Driver, Buffers, TickInfo);
+				TryGenerateLocalInput(Driver, Buffers, TickInfo);
 			}
 		}
 	}
@@ -538,11 +542,8 @@ private:
 		const int32 InputKeyframe = Buffers.Input.HeadKeyframe();
 		const int32 OutputKeyframe = InputKeyframe + 1;
 
-		TInputCmd* LastCmd = Buffers.Input.HeadElement();
-
 		// Create fake cmd				
-		TInputCmd* NewCmd = Buffers.Input.WriteKeyframe(OutputKeyframe);
-		*NewCmd = LastCmd ? *LastCmd : TInputCmd();
+		TInputCmd* NewCmd = Buffers.Input.WriteKeyframeInitializedFromHead(OutputKeyframe);
 		NewCmd->SetFrameDeltaTime(DeltaSimTime);	
 		
 		// Create new sync state to write to
@@ -551,11 +552,11 @@ private:
 		TAuxState* AuxState = Buffers.Aux[InputKeyframe];
 
 		// Do the actual update
-		T::Update(Driver, NewCmd->GetFrameDeltaTime().ToRealTimeSeconds(), *NewCmd, *PrevSyncState, *NextSyncState, *AuxState, Buffers.Aux.WriteKeyframeFunc(OutputKeyframe));
-		TickInfo.IncrementTotalProcessedSimulationTime(NewCmd->GetFrameDeltaTime(), OutputKeyframe);
+		{
+			TScopedSimulationTick UpdateScope(TickInfo, OutputKeyframe, NewCmd->GetFrameDeltaTime());
+			T::Update(Driver, NewCmd->GetFrameDeltaTime().ToRealTimeSeconds(), *NewCmd, *PrevSyncState, *NextSyncState, *AuxState, Buffers.Aux.WriteKeyframeFunc(OutputKeyframe));
+		}
 
-		// Set our PendingKeyframe to fake that we handled it
-		TickInfo.PendingKeyframe = OutputKeyframe;
 		TickInfo.MaxAllowedKeyframe = OutputKeyframe;
 	}
 	
@@ -725,11 +726,10 @@ struct TReplicator_Autonomous : public TBase
 			NextSyncState->VisualLog( FVisualLoggingParameters(Keyframe == LastKeyframeToProcess ? EVisualLoggingContext::LastMispredicted : EVisualLoggingContext::OtherMispredicted, Keyframe, EVisualLoggingLifetime::Persistent), Driver, Driver);
 
 			// Do the actual update
-			T::Update(Driver, ResimulateCmd->GetFrameDeltaTime().ToRealTimeSeconds(), *ResimulateCmd, *PrevSyncState, *NextSyncState, *AuxState, Buffers.Aux.WriteKeyframeFunc(OutputKeyframe));
-			
-			// Update TickInfo
-			TickInfo.IncrementTotalProcessedSimulationTime(ResimulateCmd->GetFrameDeltaTime(), OutputKeyframe);
-			TickInfo.PendingKeyframe = OutputKeyframe;
+			{
+				TScopedSimulationTick UpdateScope(TickInfo, OutputKeyframe, ResimulateCmd->GetFrameDeltaTime());
+				T::Update(Driver, ResimulateCmd->GetFrameDeltaTime().ToRealTimeSeconds(), *ResimulateCmd, *PrevSyncState, *NextSyncState, *AuxState, Buffers.Aux.WriteKeyframeFunc(OutputKeyframe));
+			}
 
 			// Log out the newly predicted state that we got.
 			NextSyncState->VisualLog( FVisualLoggingParameters(Keyframe == LastKeyframeToProcess ? EVisualLoggingContext::LastPredicted : EVisualLoggingContext::OtherPredicted, Keyframe, EVisualLoggingLifetime::Persistent), Driver, Driver);
@@ -760,25 +760,18 @@ struct TReplicator_Autonomous : public TBase
 			{
 				// Prediction: add simulation time and generate new commands
 				TickInfo.GiveSimulationTime(TickParameters.LocalDeltaTimeSeconds);
-
-				GenerateLocalInputCmd(Driver, Buffers, TickInfo);
+				TryGenerateLocalInput(Driver, Buffers, TickInfo);
 			}
 			else
 			{
 				// Local prediction disabled: we must use a separate time accumulator to figure out when we should add more input cmds.
 				// Since we aren't processing the simulation locally, our core simulation time will only advance from network updates.
 				// (still, we need *something* to tell us when to generate a new command and what the delta time should be)
-
 				FNetworkSimTime NonPredictedInputTime;
 				NonPredictedInputTimeAccumulator.Accumulate(NonPredictedInputTime, TickParameters.LocalDeltaTimeSeconds);
 				if (NonPredictedInputTime.IsPositive())
 				{
-					if (TInputCmd* InputCmd = Buffers.Input.WriteKeyframe(Buffers.Input.HeadKeyframe() + 1))
-					{
-						*InputCmd = TInputCmd();
-						InputCmd->SetFrameDeltaTime(NonPredictedInputTime);
-						Driver->ProduceInput(NonPredictedInputTime, *InputCmd);
-					}
+					GenerateLocalInputCmdAtKeyframe(Driver, Buffers, NonPredictedInputTime, Buffers.Input.HeadKeyframe() + 1);
 				}
 			}
 		}
diff --git a/Engine/Plugins/Runtime/NetworkPrediction/Source/NetworkPrediction/Public/NetworkSimulationModelTypes.h b/Engine/Plugins/Runtime/NetworkPrediction/Source/NetworkPrediction/Public/NetworkSimulationModelTypes.h
--- a/Engine/Plugins/Runtime/NetworkPrediction/Source/NetworkPrediction/Public/NetworkSimulationModelTypes.h
+++ b/Engine/Plugins/Runtime/NetworkPrediction/Source/NetworkPrediction/Public/NetworkSimulationModelTypes.h
@@ -9,7 +9,7 @@
 // ---------------------------------------------------------------------------------------------------------------------
 
 // Enum to refer to buffer type. These are used as template arguments to write generic code that can act on any of the buffers.
-enum ENetworkSimBufferTypeId
+enum class ENetworkSimBufferTypeId : uint8
 {
 	Input,
 	Sync,
@@ -21,10 +21,10 @@ inline FString LexToString(ENetworkSimBufferTypeId A)
 {
 	switch(A)
 	{
-		case Input: return TEXT("Input");
-		case Sync: return TEXT("Sync");
-		case Aux: return TEXT("Aux");
-		case Debug: return TEXT("Debug");
+		case ENetworkSimBufferTypeId::Input: return TEXT("Input");
+		case ENetworkSimBufferTypeId::Sync: return TEXT("Sync");
+		case ENetworkSimBufferTypeId::Aux: return TEXT("Aux");
+		case ENetworkSimBufferTypeId::Debug: return TEXT("Debug");
 	};
 	return TEXT("Unknown");
 }
@@ -37,25 +37,25 @@ struct TSelectTypeHelper
 };
 
 template<typename TBufferTypes>
-struct TSelectTypeHelper<TBufferTypes, Input>
+struct TSelectTypeHelper<TBufferTypes, ENetworkSimBufferTypeId::Input>
 {
 	using type = typename TBufferTypes::TInputCmd;
 };
 
 template<typename TBufferTypes>
-struct TSelectTypeHelper<TBufferTypes, Sync>
+struct TSelectTypeHelper<TBufferTypes, ENetworkSimBufferTypeId::Sync>
 {
 	using type = typename TBufferTypes::TSyncState;
 };
 
 template<typename TBufferTypes>
-struct TSelectTypeHelper<TBufferTypes, Aux>
+struct TSelectTypeHelper<TBufferTypes, ENetworkSimBufferTypeId::Aux>
 {
 	using type = typename TBufferTypes::TAuxState;
 };
 
 template<typename TBufferTypes>
-struct TSelectTypeHelper<TBufferTypes, Debug>
+struct TSelectTypeHelper<TBufferTypes, ENetworkSimBufferTypeId::Debug>
 {
 	using type = typename TBufferTypes::TDebugState;
 };
@@ -103,7 +103,7 @@ struct TBufferGetterHelper
 {
 	static typename TContainer::template select_buffer_type<BufferId>::type& Get(TContainer& Container)
 	{
-		static_assert(!BufferId, "Failed to find specialized Get for your BufferId");
+		static_assert(!Container, "Failed to find specialized Get for your BufferId");
 	}
 };
 
@@ -140,93 +140,56 @@ struct TBufferGetterHelper<TContainer, ENetworkSimBufferTypeId::Debug>
 	}
 };
 
-// -------------------------------------------------------------------
-
-template<typename InElementType>
-struct TBaseStateAccessor
-{
-	using ElementType = InElementType;
-	virtual ~TBaseStateAccessor() { }
-	virtual const ElementType* Get() const = 0;
-	virtual void Modify(TUniqueFunction<void(ElementType&)>&& Func, int32 Keyframe) = 0;
-};
+// -------------------------------------------------
 
-// Stored modifications to a NetSim Buffer. Needed for replaying predictive mutations that happen outside of the core Update function
-template<typename TBuffer>
-struct TPredictedModBuffer : public TBaseStateAccessor<typename TBuffer::ElementType>
+// Helper struct for accessing a netsim buffer given the underlying type
+template<typename TContainer, typename TState>
+struct TBufferGetterHelper_ByState
 {
-	using ElementType = typename TBuffer::ElementType;
-
-	TPredictedModBuffer(TBuffer& InBuffer)
-		: Buffer(InBuffer) { }
-
-	const ElementType* Get() const final override
-	{
-		return Buffer.HeadElement();
-	}
-
-	void Modify(TUniqueFunction<void(ElementType&)>&& Func, int32 Keyframe) final override
-	{
-		ElementType* State = Buffer.WriteKeyframe(Keyframe);
-		check(State);
-		Func(*State);
-
-		if (!bAuthority)
-		{
-			PredictedMods.FindOrAdd(Keyframe).Add(MoveTemp(Func));
-		}
-	}
-
-	// -------------------------------------------------------
-
-	void ApplySavedMods(int32 Keyframe)
+	static void GetBuffer(TContainer& Container)
 	{
-
+		static_assert(!Container, "Failed to find specialized Get for your BufferId");
 	}
+};
 
-	void TrimSavedMods(int32 Keyframe)
+template<typename TContainer>
+struct TBufferGetterHelper_ByState<TContainer, typename TContainer::TInputCmd>
+{
+	static typename TContainer::TInputBuffer& GetBuffer(TContainer& Container)
 	{
-
+		return Container.Input;
 	}
+};
 
-	void SetIsAuthority(bool bInAuthority)
+template<typename TContainer>
+struct TBufferGetterHelper_ByState<TContainer, typename TContainer::TSyncState>
+{
+	static typename TContainer::TSyncBuffer& GetBuffer(TContainer& Container)
 	{
-		bAuthority = bInAuthority;
+		return Container.Sync;
 	}
-
-private:
-
-	TSortedMap<int32, TArray<TUniqueFunction<void(ElementType&)>>> PredictedMods;
-	TBuffer& Buffer;
-	bool bAuthority = false; 
 };
 
-// Accessor for predicted state. This is what GameCode outside of the core ::Update function should use to mutate state in the Sync or Aux buffer/
-// This builds off TPredictedModBuffer: it stores a pointer to the simulation's current frame so that game code doesn't need to supply it.
-template<typename T>
-struct TPredictedStateAccessor
+template<typename TContainer>
+struct TBufferGetterHelper_ByState<TContainer, typename TContainer::TAuxState>
 {
-	TPredictedStateAccessor(TBaseStateAccessor<T>& InModBuffer, int32* InPendingKeyframe)
-		: ModBuffer(InModBuffer), PendingKeyframe(InPendingKeyframe) { }
-
-	// Returns the current state, cannot be modified directly
-	const T* Get() const
+	static typename TContainer::TAuxBuffer& GetBuffer(TContainer& Container)
 	{
-		return ModBuffer.Get();
+		return Container.Aux;
 	}
+};
 
-	// Applies a mod to the current state. If in prediction mode, will save the mod to be reapplied during a rollback
-	void Modify(TUniqueFunction<void(T&)>&& Func)
+template<typename TContainer>
+struct TBufferGetterHelper_ByState<TContainer, typename TContainer::TDebugState>
+{
+	static typename TContainer::TDebugBuffer& GetBuffer(TContainer& Container)
 	{
-		ModBuffer.Modify(MoveTemp(Func), *PendingKeyframe);
+		return Container.Debug;
 	}
-
-private:
-
-	TBaseStateAccessor<T>& ModBuffer;
-	int32* PendingKeyframe = nullptr;
 };
 
+// -------------------------------------------------------------------
+
 // Struct that encapsulates writing a new element to a buffer. This is used to allow a new Aux state to be created in the ::Update loop.
 template<typename T>
 struct TLazyStateAccessor
@@ -253,12 +216,14 @@ private:
 template<typename InBufferTypes>
 struct TNetworkSimBufferContainer
 {
-	TNetworkSimBufferContainer()
-		: SyncMods(Sync), AuxMods(Aux) { }
-
 	// Collection of types we were assigned
 	using TBufferTypes = InBufferTypes;
 
+	using TInputCmd = typename TBufferTypes::TInputCmd;
+	using TSyncState = typename TBufferTypes::TSyncState;
+	using TAuxState = typename TBufferTypes::TAuxState;
+	using TDebugState = typename TBufferTypes::TDebugState;
+
 	// helper that returns the buffer type for a given BufferId (not the underlying type: the actual TReplicatedBuffer<TUnderlyingType>
 	template<ENetworkSimBufferTypeId TypeId>
 	struct select_buffer_type
@@ -286,10 +251,7 @@ struct TNetworkSimBufferContainer
 	TInputBuffer Input;
 	TSyncBuffer Sync;
 	TAuxBuffer Aux;
-	TDebugBuffer Debug;
-
-	TPredictedModBuffer<TSyncBuffer> SyncMods;
-	TPredictedModBuffer<TAuxBuffer> AuxMods;
+	TDebugBuffer Debug;	
 
 	// Finally, template accessor for getting buffers based on enum. This is really what all this junk is about.
 	// This allows other templated classes in the system to access a specific buffer from another templated argument
@@ -464,6 +426,7 @@ struct TSimulationTickState
 
 	int32 PendingKeyframe = 0;
 	int32 MaxAllowedKeyframe = -1;
+	bool UpdateInProgress = false;
 	
 	FNetworkSimTime GetTotalProcessedSimulationTime() const 
 	{ 
@@ -517,6 +480,113 @@ private:
 	TRealTimeAccumulator<TSettings>	RealtimeAccumulator;
 };
 
+// Scoped helper to be used right before entering a call to the sim's ::Update function
+struct TScopedSimulationTick
+{
+	template<typename TSimulationTickState>
+	TScopedSimulationTick(TSimulationTickState& TickState, const int32& OutputKeyframe, const FNetworkSimTime& DeltaSimTime)
+	{
+		check(TickState.UpdateInProgress == false);
+		TickState.UpdateInProgress = true;
+		TickState.PendingKeyframe = OutputKeyframe;
+		TickState.IncrementTotalProcessedSimulationTime(DeltaSimTime, OutputKeyframe);
+		UpdateInProgressPtr = &TickState.UpdateInProgress;
+	}
+	~TScopedSimulationTick()
+	{
+		*UpdateInProgressPtr = false;
+	}
+	bool* UpdateInProgressPtr;
+};
+
+// ----------------------------------------------------------------------------------------------------------------------------------------------
+//	Accessors - helper structs that provide safe/cleaner access to the underlying NetSim states/events
+// ----------------------------------------------------------------------------------------------------------------------------------------------
+
+// Accessor conditionally gives access to the current (pending) Sync/Aux state to outside code.
+// Reads are always allowed
+// Writes are conditional. Authority can always write to the pending keyframe. Non authority requires the netsim to be currently processing an ::Update.
+// If you aren't inside an ::Update call, it is really not safe to predict state changes. It is safest and simplest to just not predict these changes.
+//
+// Explanation: During the scope of an ::Update call, we know exactly 'when' we are relative to what the server is processing. If the predicting client wants
+// to predict a change to sync/aux state during an update, the server will do it at the exact same time (assuming not a mis prediction). When a state change
+// happens "out of band" (outside an ::Update call) - we really have no way to correlate when the server will do it. While its tempting to think "we will get
+// a correction anyways, might as well guess at it and maybe get a smaller correction" - but this opens us up to other problems. The server may actually not 
+// change the state at all and you may not get an update that corrects you. You could add a timeout and track the state change somewhere but that really complicates
+// things and could leave you open to "double" problems: if the state change is additive, you may stack the authority change on top of the local predicted change, or
+// you may roll back the predicted change to then later receive the authority change.
+//	
+// What still may make sense to do is allow the "In Update" bool to be temporarily disabled if we enter code that we know is not rollback friendly.
+template<typename TState>
+struct TNetworkSimStateAccessor
+{
+	template<typename TNetworkSimModel>
+	void Init(TNetworkSimModel* NetSimModel)
+	{
+		GetStateFunc = [NetSimModel](bool bWrite, TState*& OutState, bool& OutSafe)
+		{
+			// Gross: find the buffer our TState is in. This allows these accessors to be declared as class member variables without knowing the exact net sim it will
+			// be pulling from. (For example, you may want to templatize your network sim model).
+			auto& Buffer = TBufferGetterHelper_ByState<TNetworkSimBufferContainer<typename TNetworkSimModel::TBufferTypes>, TState>::GetBuffer(NetSimModel->Buffers);
+			OutState = bWrite ? Buffer.WriteKeyframeInitializedFromHead(NetSimModel->TickInfo.PendingKeyframe) : Buffer[NetSimModel->TickInfo.PendingKeyframe];
+			OutSafe = NetSimModel->TickInfo.UpdateInProgress;
+		};
+	}
+
+	void Clear()
+	{
+		GetStateFunc = nullptr;
+	}
+
+	/** Gets the current (PendingKeyframe) state for reading. This is not expected to fail outside of startup/shutdown edge cases */
+	const TState* GetStateRead() const
+	{
+		TState* State = nullptr;
+		bool bSafe = false;
+		if (GetStateFunc)
+		{
+			GetStateFunc(false, State, bSafe);
+		}
+		return State;
+	}
+
+	/** Gets the current (PendingKeyframe) state for writing. This is expected to fail outside of the core update loop when bHasAuthority=false. (E.g, it is not safe to predict writes) */
+	TState* GetStateWrite(bool bHasAuthority) const
+	{
+		TState* State = nullptr;
+		bool bSafe = false;
+		if (GetStateFunc)
+		{
+			GetStateFunc(true, State, bSafe);
+		}
+		return (bHasAuthority || bSafe) ? State : nullptr;
+	}
+
+private:
+
+	TFunction<void(bool bForWrite, TState*& OutState, bool& OutSafe)> GetStateFunc;
+};
+
+//TODO
+struct FNetworkSimEventAccessor
+{
+	int32 GetPendingKeyframe() { return -1; }
+
+	struct FFrameEvents
+	{
+		/** This frame has been received from the authority and will not be rolled back or resimulated again. */
+		DECLARE_MULTICAST_DELEGATE(FOnFrameConfirmed)
+		FOnFrameConfirmed Confirmed;
+
+		/** This frame was previously simulated and the simulation has now been rolled back */
+		DECLARE_MULTICAST_DELEGATE(FOnFrameRolledBack)
+		FOnFrameRolledBack RolledBack;
+
+		DECLARE_MULTICAST_DELEGATE(FOnFrameSimulated)
+		FOnFrameRolledBack Simulated;
+	};
+};
+
 // ----------------------------------------------------------------------------------------------------------------------------------------------
 //
 // ----------------------------------------------------------------------------------------------------------------------------------------------
diff --git a/Engine/Plugins/Runtime/NetworkPredictionExtras/Source/NetworkPredictionExtras/Private/NetworkPredictionExtrasFlyingPawn.cpp b/Engine/Plugins/Runtime/NetworkPredictionExtras/Source/NetworkPredictionExtras/Private/NetworkPredictionExtrasFlyingPawn.cpp
--- a/Engine/Plugins/Runtime/NetworkPredictionExtras/Source/NetworkPredictionExtras/Private/NetworkPredictionExtrasFlyingPawn.cpp
+++ b/Engine/Plugins/Runtime/NetworkPredictionExtras/Source/NetworkPredictionExtras/Private/NetworkPredictionExtrasFlyingPawn.cpp
@@ -202,15 +202,60 @@ void ANetworkPredictionExtrasFlyingPawn::PrintDebug()
 	}
 }
 
+float ANetworkPredictionExtrasFlyingPawn::GetMaxMoveSpeed() const
+{
+	if (const FlyingMovement::FAuxState* AuxState = GetAuxStateRead())
+	{
+		return AuxState->MaxSpeed;
+	}
+	return 0;
+}
+
 void ANetworkPredictionExtrasFlyingPawn::SetMaxMoveSpeed(float NewMaxMoveSpeed)
+{
+	if (FlyingMovement::FAuxState* AuxState = GetAuxStateWrite())
+	{
+		AuxState->MaxSpeed = NewMaxMoveSpeed;
+	}
+}
+
+void ANetworkPredictionExtrasFlyingPawn::AddMaxMoveSpeed(float AdditiveMaxMoveSpeed)
+{
+	if (FlyingMovement::FAuxState* AuxState = GetAuxStateWrite())
+	{
+		//UE_LOG(LogTemp, Warning, TEXT("MaxSpeed: %.2f += %.2f"), AuxState->MaxSpeed, AdditiveMaxMoveSpeed);
+		AuxState->MaxSpeed += AdditiveMaxMoveSpeed;
+	}
+}
+
+const FlyingMovement::FAuxState* ANetworkPredictionExtrasFlyingPawn::GetAuxStateRead() const
 {
 	if (ensure(FlyingMovementComponent))
 	{
-		FlyingMovementComponent->MovementAuxState->Modify([NewMaxMoveSpeed](FlyingMovement::FAuxState& State)
-		{
-			State.MaxSpeed = NewMaxMoveSpeed;
-		});
+		return FlyingMovementComponent->MovementAuxState.GetStateRead();
+	}
+	return nullptr;
+}
+
+FlyingMovement::FAuxState* ANetworkPredictionExtrasFlyingPawn::GetAuxStateWrite()
+{
+	if (ensure(FlyingMovementComponent))
+	{
+		return FlyingMovementComponent->MovementAuxState.GetStateWrite(GetLocalRole() == ROLE_Authority);
+	}
+	return nullptr;
+}
+
+int32 ANetworkPredictionExtrasFlyingPawn::GetPendingKeyframe() const
+{
+	// TODO
+	/*
+	if (ensure(FlyingMovementComponent))
+	{
+		return FlyingMovementComponent->MovementSyncState->GetPendingKeyframe();
 	}
+	*/
+	return 0;
 }
 
 void ANetworkPredictionExtrasFlyingPawn::ProduceInput(const FNetworkSimTime SimTime, FlyingMovement::FInputCmd& Cmd)
diff --git a/Engine/Plugins/Runtime/NetworkPredictionExtras/Source/NetworkPredictionExtras/Public/NetworkPredictionExtrasFlyingPawn.h b/Engine/Plugins/Runtime/NetworkPredictionExtras/Source/NetworkPredictionExtras/Public/NetworkPredictionExtrasFlyingPawn.h
--- a/Engine/Plugins/Runtime/NetworkPredictionExtras/Source/NetworkPredictionExtras/Public/NetworkPredictionExtrasFlyingPawn.h
+++ b/Engine/Plugins/Runtime/NetworkPredictionExtras/Source/NetworkPredictionExtras/Public/NetworkPredictionExtrasFlyingPawn.h
@@ -66,9 +66,24 @@ public:
 	UFUNCTION(BlueprintCallable, Category="Debug")
 	void PrintDebug();
 
+	UFUNCTION(BlueprintCallable, Category="Gameplay")
+	float GetMaxMoveSpeed() const;
+
 	UFUNCTION(BlueprintCallable, Category="Gameplay")
 	void SetMaxMoveSpeed(float NewMaxMoveSpeed);
 
+	UFUNCTION(BlueprintCallable, Category="Gameplay")
+	void AddMaxMoveSpeed(float AdditiveMaxMoveSpeed);
+
+	// Only intended for debugging in test map examples. Not really intended to be useful for general game code use.
+	UFUNCTION(BlueprintCallable, Category="Gameplay")
+	int32 GetPendingKeyframe() const;
+
+protected:
+
+	const FlyingMovement::FAuxState* GetAuxStateRead() const;
+	FlyingMovement::FAuxState* GetAuxStateWrite();
+
 private:
 
 	void ProduceInput(const FNetworkSimTime SimTime, FlyingMovement::FInputCmd& Cmd);
