Large refactor on the core classes:
"Simulation Driver" and a few other concepts have changed:
	-"Simulation" is now implemented by a simulation class. What was previously a static method + a "driver" pointer, is now an instantiated C++ class that user code creates and passes into TNetworkedSimulationModel.
	-The Simulation (previously driver) is now longer implemented by the owning actor/component. This simplifies a lot of things and make extending simulation much less awkward.
	-TNetworkSimDriverInterfaceBase --> TNetworkedSimulationModelDriver. This is now distinct from "what the simulation needs to call internally" (previously defined by the driver now the simulation).
	-TNetworkedSimulationModelDriver ("System Driver") is the set of functions the system internally needs to call.
	-Visual logging has been refactored to be a single function of the system driver that takes all 3 states instead of single methods on the individual system types (Input,Sync,Aux). this simplifies a lot of stuff and allows you to visual sync+aux together.

"MockAbilitySystem" has been implemented.
	-Simple, "hardcoded" abilities that function on top of the flying movement simulation. See notes in MockAbilitySimulation.h
	-This code gives a good glimpse into what GameplayAbilities will look like (again, more hardcoded in this example but same techniques will be applied).

#rnx

diff --git a/Engine/Plugins/Runtime/NetworkPrediction/Source/NetworkPrediction/Private/Movement/BaseMovementComponent.cpp b/Engine/Plugins/Runtime/NetworkPrediction/Source/NetworkPrediction/Private/Movement/BaseMovementComponent.cpp
--- a/Engine/Plugins/Runtime/NetworkPrediction/Source/NetworkPrediction/Private/Movement/BaseMovementComponent.cpp
+++ b/Engine/Plugins/Runtime/NetworkPrediction/Source/NetworkPrediction/Private/Movement/BaseMovementComponent.cpp
@@ -9,13 +9,10 @@
 #include "VisualLogger/VisualLogger.h"
 #include "NetworkPredictionTypes.h"
 
-
 DEFINE_LOG_CATEGORY_STATIC(LogBaseMovement, Log, All);
 
-
 namespace BaseMovementCVars
 {
-
 static float PenetrationPullbackDistance = 0.125f;
 static FAutoConsoleVariableRef CVarPenetrationPullbackDistance(TEXT("bm.PenetrationPullbackDistance"),
 	PenetrationPullbackDistance,
@@ -33,19 +30,6 @@ static FAutoConsoleVariableRef CVarPenetrationOverlapCheckInflation(TEXT("bm.Pen
 static int32 RequestMispredict = 0;
 static FAutoConsoleVariableRef CVarRequestMispredict(TEXT("bm.RequestMispredict"),
 	RequestMispredict, TEXT("Causes a misprediction by inserting random value into stream on authority side"), ECVF_Default);
-
-
-static int32 UseVLogger = 1;
-static FAutoConsoleVariableRef CVarUseVLogger(TEXT("bm.Debug.UseUnrealVLogger"),
-	UseVLogger,	TEXT("Use Unreal Visual Logger\n"),	ECVF_Default);
-
-static int32 UseDrawDebug = 1;
-static FAutoConsoleVariableRef CVarUseDrawDebug(TEXT("bm.Debug.UseDrawDebug"),
-	UseDrawDebug,	TEXT("Use built in DrawDebug* functions for visual logging\n"), ECVF_Default);
-
-static float DrawDebugDefaultLifeTime = 30.f;
-static FAutoConsoleVariableRef CVarDrawDebugDefaultLifeTime(TEXT("bm.Debug.DrawDebugLifetime.DefaultLifeTime"),
-	DrawDebugDefaultLifeTime, TEXT("Use built in DrawDebug* functions for visual logging"), ECVF_Default);
 }
 
 // ----------------------------------------------------------------------------------------------------------
@@ -183,7 +167,7 @@ void UBaseMovementComponent::PhysicsVolumeChanged(APhysicsVolume* NewVolume)
 //
 // ----------------------------------------------------------------------------------------------------------
 
-FVector UBaseMovementComponent::GetPenetrationAdjustment(const FHitResult& Hit)
+FVector FBaseMovementSimulation::GetPenetrationAdjustment(const FHitResult& Hit) const
 {
 	if (!Hit.bStartPenetrating)
 	{
@@ -199,15 +183,15 @@ FVector UBaseMovementComponent::GetPenetrationAdjustment(const FHitResult& Hit)
 	return Result;
 }
 
-bool UBaseMovementComponent::OverlapTest(const FVector& Location, const FQuat& RotationQuat, const ECollisionChannel CollisionChannel, const FCollisionShape& CollisionShape, const AActor* IgnoreActor) const
+bool FBaseMovementSimulation::OverlapTest(const FVector& Location, const FQuat& RotationQuat, const ECollisionChannel CollisionChannel, const FCollisionShape& CollisionShape, const AActor* IgnoreActor) const
 {
 	FCollisionQueryParams QueryParams(SCENE_QUERY_STAT(MovementOverlapTest), false, IgnoreActor);
 	FCollisionResponseParams ResponseParam;
 	InitCollisionParams(QueryParams, ResponseParam);
-	return GetWorld()->OverlapBlockingTestByChannel(Location, RotationQuat, CollisionChannel, CollisionShape, QueryParams, ResponseParam);
+	return UpdatedComponent->GetWorld()->OverlapBlockingTestByChannel(Location, RotationQuat, CollisionChannel, CollisionShape, QueryParams, ResponseParam);
 }
 
-void UBaseMovementComponent::InitCollisionParams(FCollisionQueryParams &OutParams, FCollisionResponseParams& OutResponseParam) const
+void FBaseMovementSimulation::InitCollisionParams(FCollisionQueryParams &OutParams, FCollisionResponseParams& OutResponseParam) const
 {
 	if (UpdatedPrimitive)
 	{
@@ -215,7 +199,7 @@ void UBaseMovementComponent::InitCollisionParams(FCollisionQueryParams &OutParam
 	}
 }
 
-bool UBaseMovementComponent::ResolvePenetration(const FVector& ProposedAdjustment, const FHitResult& Hit, const FQuat& NewRotationQuat) const
+bool FBaseMovementSimulation::ResolvePenetration(const FVector& ProposedAdjustment, const FHitResult& Hit, const FQuat& NewRotationQuat) const
 {
 	// SceneComponent can't be in penetration, so this function really only applies to PrimitiveComponent.
 	const FVector Adjustment = ProposedAdjustment; //ConstrainDirectionToPlane(ProposedAdjustment);
@@ -237,7 +221,7 @@ bool UBaseMovementComponent::ResolvePenetration(const FVector& ProposedAdjustmen
 			   *GetNameSafe(Hit.GetComponent()),
 			   Hit.Component.IsValid() ? *Hit.GetComponent()->GetComponentLocation().ToString() : TEXT("<unknown>"),
 			   Hit.PenetrationDepth,
-			   (uint32)GetNetMode());
+			   (uint32)ActorOwner->GetNetMode());
 
 		// We really want to make sure that precision differences or differences between the overlap test and sweep tests don't put us into another overlap,
 		// so make the overlap test a bit more restrictive.
@@ -292,7 +276,7 @@ bool UBaseMovementComponent::ResolvePenetration(const FVector& ProposedAdjustmen
 	return false;
 }
 
-bool UBaseMovementComponent::SafeMoveUpdatedComponent(const FVector& Delta, const FQuat& NewRotation, bool bSweep, FHitResult& OutHit, ETeleportType Teleport) const
+bool FBaseMovementSimulation::SafeMoveUpdatedComponent(const FVector& Delta, const FQuat& NewRotation, bool bSweep, FHitResult& OutHit, ETeleportType Teleport) const
 {
 	if (UpdatedComponent == NULL)
 	{
@@ -321,7 +305,7 @@ bool UBaseMovementComponent::SafeMoveUpdatedComponent(const FVector& Delta, cons
 	return bMoveResult;
 }
 
-bool UBaseMovementComponent::MoveUpdatedComponent(const FVector& Delta, const FQuat& NewRotation, bool bSweep, FHitResult* OutHit, ETeleportType Teleport) const
+bool FBaseMovementSimulation::MoveUpdatedComponent(const FVector& Delta, const FQuat& NewRotation, bool bSweep, FHitResult* OutHit, ETeleportType Teleport) const
 {
 	if (UpdatedComponent)
 	{
@@ -332,103 +316,12 @@ bool UBaseMovementComponent::MoveUpdatedComponent(const FVector& Delta, const FQ
 	return false;
 }
 
-FTransform UBaseMovementComponent::GetUpdateComponentTransform() const
+
+FTransform FBaseMovementSimulation::GetUpdateComponentTransform() const
 {
 	if (ensure(UpdatedComponent))
 	{
 		return UpdatedComponent->GetComponentTransform();		
 	}
 	return FTransform::Identity;
-}
-
-UObject* UBaseMovementComponent::GetBaseMovementVLogOwner() const
-{
-	return GetOwner();
-}
-
-void UBaseMovementComponent::DrawDebug(const IBaseMovementDriver::FDrawDebugParams& Params) const
-{
-	if (!UpdatedComponent)
-	{
-		return;
-	}
-
-	AActor* Owner = GetOwner();
-
-	const bool PersistentLines = (Params.Lifetime == EVisualLoggingLifetime::Persistent);
-	const float LifetimeSeconds = PersistentLines ? 20.f : -1.f;
-
-	if (Params.DrawType == EVisualLoggingDrawType::Crumb)
-	{
-		if (BaseMovementCVars::UseDrawDebug)
-		{
-			static float PointSize = 3.f;
-			DrawDebugPoint(Params.DebugWorld, Params.Transform.GetLocation(), PointSize, Params.DrawColor, false, LifetimeSeconds);
-		}
-
-		if (BaseMovementCVars::UseVLogger)
-		{
-			static float CrumbRadius = 3.f;
-			UE_VLOG_LOCATION(Params.DebugLogOwner, LogBaseMovement, Log, Params.Transform.GetLocation(), CrumbRadius, Params.DrawColor, TEXT("%s"), *Params.InWorldText);
-		}
-	}
-	else
-	{
-		// Full drawing. Trying to provide a generic implementation that will give the most accurate representation.
-		// Still, subclasses may want to customize this
-		if (UCapsuleComponent* CapsuleComponent = Cast<UCapsuleComponent>(UpdatedComponent))
-		{
-			float Radius=0.f;
-			float HalfHeight=0.f;
-			CapsuleComponent->GetScaledCapsuleSize(Radius, HalfHeight);
-
-			if (BaseMovementCVars::UseDrawDebug)
-			{
-				static const float Thickness = 2.f;
-				DrawDebugCapsule(Params.DebugWorld, Params.Transform.GetLocation(), HalfHeight, Radius, Params.Transform.GetRotation(), Params.DrawColor, false, LifetimeSeconds, 0.f, Thickness);
-			}
-
-			if (BaseMovementCVars::UseVLogger)
-			{
-				FVector VLogPosition = Params.Transform.GetLocation();
-				VLogPosition.Z -= HalfHeight;
-				UE_VLOG_CAPSULE(Params.DebugLogOwner, LogBaseMovement, Log, VLogPosition, HalfHeight, Radius, Params.Transform.GetRotation(), Params.DrawColor, TEXT("%s"), *Params.InWorldText);
-			}
-		}
-		else
-		{
-			// Generic Actor Bounds drawing
-			FBox LocalSpaceBox = Owner->CalculateComponentsBoundingBoxInLocalSpace();
-
-			if (BaseMovementCVars::UseDrawDebug)
-			{
-				static const float Thickness = 2.f;
-
-				FVector ActorOrigin;
-				FVector ActorExtent;
-				LocalSpaceBox.GetCenterAndExtents(ActorOrigin, ActorExtent);
-				ActorExtent *= Params.Transform.GetScale3D();
-				DrawDebugBox(Params.DebugWorld, Params.Transform.GetLocation(), ActorExtent, Params.Transform.GetRotation(), Params.DrawColor, false, LifetimeSeconds, 0, Thickness);
-			}
-
-			if (BaseMovementCVars::UseVLogger)
-			{
-				UE_VLOG_OBOX(Params.DebugLogOwner, LogBaseMovement, Log, LocalSpaceBox, Params.Transform.ToMatrixWithScale(), Params.DrawColor, TEXT("%s"), *Params.InWorldText);
-			}
-		}
-	}
-
-	if (BaseMovementCVars::UseVLogger)
-	{
-		UE_VLOG(Params.DebugLogOwner, LogBaseMovement, Log, TEXT("%s"), *Params.LogText);
-	}
-}
-
-IBaseMovementDriver::FDrawDebugParams::FDrawDebugParams(const FVisualLoggingParameters& Parameters, IBaseMovementDriver* LogDriver)
-{
-	DebugWorld = LogDriver->GetBaseMovementDriverWorld();
-	DebugLogOwner = LogDriver->GetBaseMovementVLogOwner();
-	DrawType = (Parameters.Context == EVisualLoggingContext::OtherMispredicted || Parameters.Context == EVisualLoggingContext::OtherPredicted) ? EVisualLoggingDrawType::Crumb : EVisualLoggingDrawType::Full;
-	Lifetime = Parameters.Lifetime;
-	DrawColor = Parameters.GetDebugColor();	
 }
\ No newline at end of file
diff --git a/Engine/Plugins/Runtime/NetworkPrediction/Source/NetworkPrediction/Private/Movement/FlyingMovementComponent.cpp b/Engine/Plugins/Runtime/NetworkPrediction/Source/NetworkPrediction/Private/Movement/FlyingMovementComponent.cpp
--- a/Engine/Plugins/Runtime/NetworkPrediction/Source/NetworkPrediction/Private/Movement/FlyingMovementComponent.cpp
+++ b/Engine/Plugins/Runtime/NetworkPrediction/Source/NetworkPrediction/Private/Movement/FlyingMovementComponent.cpp
@@ -21,28 +21,12 @@
 #include "Engine/Canvas.h"
 #include "Engine/Engine.h"
 #include "Debug/ReporterGraph.h"
-#include "NetworkSimulationModelDebugger.h"
 
 DEFINE_LOG_CATEGORY_STATIC(LogFlyingMovement, Log, All);
 
 
 namespace FlyingMovementCVars
 {
-
-static float PenetrationPullbackDistance = 0.125f;
-static FAutoConsoleVariableRef CVarPenetrationPullbackDistance(TEXT("fp.PenetrationPullbackDistance"),
-	PenetrationPullbackDistance,
-	TEXT("Pull out from penetration of an object by this extra distance.\n")
-	TEXT("Distance added to penetration fix-ups."),
-	ECVF_Default);
-
-static float PenetrationOverlapCheckInflation = 0.100f;
-static FAutoConsoleVariableRef CVarPenetrationOverlapCheckInflation(TEXT("motion.PenetrationOverlapCheckInflation"),
-	PenetrationOverlapCheckInflation,
-	TEXT("Inflation added to object when checking if a location is free of blocking collision.\n")
-	TEXT("Distance added to inflation in penetration overlap check."),
-	ECVF_Default);
-
 static float MaxSpeed = 1200.f;
 static FAutoConsoleVariableRef CVarMaxSpeed(TEXT("motion.MaxSpeed"),
 	MaxSpeed,
@@ -70,7 +54,13 @@ UFlyingMovementComponent::UFlyingMovementComponent()
 INetworkSimulationModel* UFlyingMovementComponent::InstantiateNetworkSimulation()
 {
 	check(UpdatedComponent);
+
+	// Create the simulation
+	MovementSimulation.Reset(new FlyingMovement::FMovementSimulation());
+	MovementSimulation->UpdatedComponent = UpdatedComponent;
+	MovementSimulation->UpdatedPrimitive = UpdatedPrimitive;
 	
+	// Initial states
 	FlyingMovement::FMoveState InitialSyncState;
 	InitialSyncState.Location = UpdatedComponent->GetComponentLocation();
 	InitialSyncState.Rotation = UpdatedComponent->GetComponentQuat().Rotator();	
@@ -78,11 +68,13 @@ INetworkSimulationModel* UFlyingMovementComponent::InstantiateNetworkSimulation(
 	FlyingMovement::FAuxState InitialAuxState;
 	InitialAuxState.MaxSpeed = FlyingMovementCVars::MaxSpeed;
 
-	auto NewSim = new FlyingMovement::FMovementSystem<0>(this, InitialSyncState, InitialAuxState);
-	DO_NETSIM_MODEL_DEBUG(FNetworkSimulationModelDebuggerManager::Get().RegisterNetworkSimulationModel(NewSim, GetOwner()));
-	MovementSyncState.Init(NewSim);
-	MovementAuxState.Init(NewSim);
-	return NewSim;
+	// The NetworkModel
+	auto NewModel = new FlyingMovement::FMovementSystem<0>(MovementSimulation.Get(), this, InitialSyncState, InitialAuxState);
+
+	// Accessors to the latest sync/aux state
+	MovementSyncState.Init(NewModel);
+	MovementAuxState.Init(NewModel);
+	return NewModel;
 }
 
 void UFlyingMovementComponent::TickComponent(float DeltaTime, enum ELevelTick TickType, FActorComponentTickFunction* ThisTickFunction)
@@ -136,7 +128,13 @@ void UFlyingMovementComponent::TickComponent(float DeltaTime, enum ELevelTick Ti
 //
 // ----------------------------------------------------------------------------------------------------------
 
-void UFlyingMovementComponent::PreSimSync(const FlyingMovement::FMoveState& SyncState)
+void UFlyingMovementComponent::ProduceInput(const FNetworkSimTime SimTime, FlyingMovement::FInputCmd& Cmd)
+{
+	// This isn't ideal. It probably makes sense for the component to do all the input binding rather.
+	ProduceInputDelegate.ExecuteIfBound(SimTime, Cmd);
+}
+
+void UFlyingMovementComponent::FinalizeFrame(const FlyingMovement::FMoveState& SyncState, const FlyingMovement::FAuxState& AuxState)
 {
 	// Does checking equality make any sense here? This is unfortunate
 	if (UpdatedComponent->GetComponentLocation().Equals(SyncState.Location) == false || UpdatedComponent->GetComponentQuat().Rotator().Equals(SyncState.Rotation, FlyingMovement::ROTATOR_TOLERANCE) == false)
@@ -148,23 +146,18 @@ void UFlyingMovementComponent::PreSimSync(const FlyingMovement::FMoveState& Sync
 	}
 }
 
-void UFlyingMovementComponent::ProduceInput(const FNetworkSimTime SimTime, FlyingMovement::FInputCmd& Cmd)
-{
-	// This isn't ideal. It probably makes sense for the component to do all the input binding rather.
-	ProduceInputDelegate.ExecuteIfBound(SimTime, Cmd);
-}
-
-void UFlyingMovementComponent::FinalizeFrame(const FlyingMovement::FMoveState& SyncState, const FlyingMovement::FAuxState& AuxState)
-{
-	PreSimSync(SyncState);
-}
-
 FString UFlyingMovementComponent::GetDebugName() const
 {
 	return FString::Printf(TEXT("FlyingMovement. %s. %s"), *UEnum::GetValueAsString(TEXT("Engine.ENetRole"), GetOwnerRole()), *GetName());
 }
 
-const UObject* UFlyingMovementComponent::GetVLogOwner() const
+const AActor* UFlyingMovementComponent::GetVLogOwner() const
 {
 	return GetOwner();
+}
+
+void UFlyingMovementComponent::VisualLog(const FlyingMovement::FInputCmd* Input, const FlyingMovement::FMoveState* Sync, const FlyingMovement::FAuxState* Aux, const FVisualLoggingParameters& SystemParameters) const
+{
+	FTransform Transform(Sync->Rotation, Sync->Location);
+	FVisualLoggingHelpers::VisualLogActor(GetOwner(), Transform, SystemParameters);	
 }
\ No newline at end of file
diff --git a/Engine/Plugins/Runtime/NetworkPrediction/Source/NetworkPrediction/Private/Movement/FlyingMovementSimulation.cpp b/Engine/Plugins/Runtime/NetworkPrediction/Source/NetworkPrediction/Private/Movement/FlyingMovementSimulation.cpp
--- a/Engine/Plugins/Runtime/NetworkPrediction/Source/NetworkPrediction/Private/Movement/FlyingMovementSimulation.cpp
+++ b/Engine/Plugins/Runtime/NetworkPrediction/Source/NetworkPrediction/Private/Movement/FlyingMovementSimulation.cpp
@@ -22,7 +22,6 @@
 #include "Engine/Canvas.h"
 #include "Engine/Engine.h"
 #include "Debug/ReporterGraph.h"
-#include "NetworkSimulationModelDebugger.h"
 
 DEFINE_LOG_CATEGORY_STATIC(LogFlyingPawnSimulation, Log, All);
 
@@ -39,7 +38,6 @@ static FAutoConsoleVariableRef CVarUseDrawDebug(TEXT("fp.Debug.UseDrawDebug"),
 static float DrawDebugDefaultLifeTime = 30.f;
 static FAutoConsoleVariableRef CVarDrawDebugDefaultLifeTime(TEXT("fp.Debug.UseDrawDebug.DefaultLifeTime"),
 	DrawDebugDefaultLifeTime, TEXT("Use built in DrawDebug* functions for visual logging"), ECVF_Default);
-
 }
 
 // ============================================================================================
@@ -101,7 +99,7 @@ namespace FlyingMovement
 		OutDelta = Delta;
 	}
 
-	float SlideAlongSurface(IMovementDriver* Driver, const FVector& Delta, float Time, const FQuat Rotation, const FVector& Normal, FHitResult& Hit, bool bHandleImpact)
+	float FMovementSimulation::SlideAlongSurface(const FVector& Delta, float Time, const FQuat Rotation, const FVector& Normal, FHitResult& Hit, bool bHandleImpact)
 	{
 		if (!Hit.bBlockingHit)
 		{
@@ -115,9 +113,7 @@ namespace FlyingMovement
 
 		if ((SlideDelta | Delta) > 0.f)
 		{
-			IBaseMovementDriver& BaseMovementDriver = Driver->GetBaseMovementDriver();
-
-			BaseMovementDriver.SafeMoveUpdatedComponent(SlideDelta, Rotation, true, Hit, ETeleportType::None);
+			SafeMoveUpdatedComponent(SlideDelta, Rotation, true, Hit, ETeleportType::None);
 
 			const float FirstHitPercent = Hit.Time;
 			PercentTimeApplied = FirstHitPercent;
@@ -136,7 +132,7 @@ namespace FlyingMovement
 				if (!SlideDelta.IsNearlyZero(1e-3f) && (SlideDelta | Delta) > 0.f)
 				{
 					// Perform second move
-					BaseMovementDriver.SafeMoveUpdatedComponent(SlideDelta, Rotation, true, Hit, ETeleportType::None);
+					SafeMoveUpdatedComponent(SlideDelta, Rotation, true, Hit, ETeleportType::None);
 					const float SecondHitPercent = Hit.Time * (1.f - FirstHitPercent);
 					PercentTimeApplied += SecondHitPercent;
 
@@ -155,10 +151,8 @@ namespace FlyingMovement
 	}
 
 
-	void FMovementSimulation::Update(IMovementDriver* Driver, const float DeltaSeconds, const FInputCmd& InputCmd, const FMoveState& InputState, FMoveState& OutputState, const FAuxState& AuxState, const TLazyStateAccessor<FAuxState>& OutAuxStateAccessor)
+	void FMovementSimulation::Update(const float DeltaSeconds, const FInputCmd& InputCmd, const FMoveState& InputState, FMoveState& OutputState, const FAuxState& AuxState, const TNetSimLazyWriter<FAuxState>& OutAuxStateAccessor)
 	{
-		IBaseMovementDriver& BaseMovementDriver = Driver->GetBaseMovementDriver();
-
 		OutputState = InputState;
 
 		// --------------------------------------------------------------
@@ -175,7 +169,7 @@ namespace FlyingMovement
 
 		// After the rotation is known, we need to sync our driver to this state. This is unfortunate and probably not best for perf, but since the Driver owned
 		// primitive component is ultimately what does our scene queries, we have no choice: we must get the primitive component in the state we say it should be in.
-		Driver->PreSimSync(OutputState);
+		PreSimSync(OutputState);
 
 		const FVector LocalSpaceMovementInput = OutputState.Rotation.RotateVector( InputCmd.MovementInput );
 	   	
@@ -242,12 +236,12 @@ namespace FlyingMovement
 		if (!Delta.IsNearlyZero(1e-6f))
 		{
 			FHitResult Hit(1.f);
-			BaseMovementDriver.SafeMoveUpdatedComponent(Delta, OutputQuat, true, Hit, ETeleportType::None);
+			SafeMoveUpdatedComponent(Delta, OutputQuat, true, Hit, ETeleportType::None);
 
 			if (Hit.IsValidBlockingHit())
 			{
 				// Try to slide the remaining distance along the surface.
-				SlideAlongSurface(Driver, Delta, 1.f-Hit.Time, OutputQuat, Hit.Normal, Hit, true);
+				SlideAlongSurface(Delta, 1.f-Hit.Time, OutputQuat, Hit.Normal, Hit, true);
 			}
 
 			// Update velocity
@@ -261,24 +255,23 @@ namespace FlyingMovement
 		}
 
 		// Finalize. This is unfortunate. The component mirrors our internal motion state and since we call into it to update, at this point, it has the real position.
-		const FTransform UpdateComponentTransform = BaseMovementDriver.GetUpdateComponentTransform();
+		const FTransform UpdateComponentTransform = GetUpdateComponentTransform();
 		OutputState.Location = UpdateComponentTransform.GetLocation();
 
 		// Note that we don't pull the rotation out of the final update transform. Converting back from a quat will lead to a different FRotator than what we are storing
 		// here in the simulation layer. This may not be the best choice for all movement simulations, but is ok for this one.
 	}
 
-	void FMoveState::VisualLog(const FVisualLoggingParameters& Parameters, IMovementDriver* Driver, IMovementDriver* LogDriver) const
+	void FMovementSimulation::PreSimSync(const FMoveState& SyncState)
 	{
-		IBaseMovementDriver& BaseMovementDriver = Driver->GetBaseMovementDriver();		
-		FTransform Transform(Rotation, Location);		
-
-		IBaseMovementDriver::FDrawDebugParams DrawParams(Parameters, &LogDriver->GetBaseMovementDriver());		
-		DrawParams.Transform = Transform;		
-		DrawParams.InWorldText = LexToString(Parameters.Keyframe);
-		DrawParams.LogText = FString::Printf(TEXT("[%d] %s. Location: %s. Rotation: %s"), Parameters.Keyframe, *LexToString(Parameters.Context), *Transform.GetLocation().ToString(), *Transform.GetRotation().Rotator().ToString());
+		// Does checking equality make any sense here? This is unfortunate
+		if (UpdatedComponent->GetComponentLocation().Equals(SyncState.Location) == false || UpdatedComponent->GetComponentQuat().Rotator().Equals(SyncState.Rotation, ROTATOR_TOLERANCE) == false)
+		{
+			FTransform Transform(SyncState.Rotation.Quaternion(), SyncState.Location, UpdatedComponent->GetComponentTransform().GetScale3D() );
+			UpdatedComponent->SetWorldTransform(Transform, false, nullptr, ETeleportType::TeleportPhysics);
 
-		BaseMovementDriver.DrawDebug(DrawParams);
+			UpdatedComponent->ComponentVelocity = SyncState.Velocity;
+		}
 	}
 
 } // end namespace
\ No newline at end of file
diff --git a/Engine/Plugins/Runtime/NetworkPrediction/Source/NetworkPrediction/Private/Movement/ParametricMovement.cpp b/Engine/Plugins/Runtime/NetworkPrediction/Source/NetworkPrediction/Private/Movement/ParametricMovement.cpp
--- a/Engine/Plugins/Runtime/NetworkPrediction/Source/NetworkPrediction/Private/Movement/ParametricMovement.cpp
+++ b/Engine/Plugins/Runtime/NetworkPrediction/Source/NetworkPrediction/Private/Movement/ParametricMovement.cpp
@@ -1,7 +1,6 @@
 // Copyright 1998-2019 Epic Games, Inc. All Rights Reserved.
 
 #include "Movement/ParametricMovement.h"
-#include "NetworkSimulationModelDebugger.h"
 #include "HAL/IConsoleManager.h"
 #include "NetworkSimulationGlobalManager.h"
 #include "Engine/LocalPlayer.h"
@@ -42,14 +41,13 @@ namespace ParametricMovement
 {
 	const FName FMovementSimulation::GroupName(TEXT("Parametric"));
 
-	void FMovementSimulation::Update(IMovementDriver* Driver, const float DeltaSeconds, const FInputCmd& InputCmd, const FMoveState& InputState, FMoveState& OutputState, const FAuxState& AuxState, const TLazyStateAccessor<FAuxState>& OutAuxStateAccessor)
+	void FMovementSimulation::Update(const float DeltaSeconds, const FInputCmd& InputCmd, const FMoveState& InputState, FMoveState& OutputState, const FAuxState& AuxState, const TNetSimLazyWriter<FAuxState>& OutAuxStateAccessor)
 	{
-		IBaseMovementDriver& BaseMovementDriver = Driver->GetBaseMovementDriver();
-		//const float DeltaSeconds = SimTimeDeltaMS.ToRealTimeSeconds();
+		check(Motion != nullptr); // Must set motion mapping prior to running the simulation
 
 		// Advance parametric time. This won't always be linear: we could loop, rewind/bounce, etc
 		const float InputPlayRate = InputCmd.PlayRate.Get(InputState.PlayRate); // Returns InputCmds playrate if set, else returns previous state's playrate		
-		Driver->AdvanceParametricTime(InputState.Position, InputPlayRate, OutputState.Position, OutputState.PlayRate, DeltaSeconds);
+		Motion->AdvanceParametricTime(InputState.Position, InputPlayRate, OutputState.Position, OutputState.PlayRate, DeltaSeconds);
 
 		// We have our time that we should be at. We just need to move primitive component to that position.
 		// Again, note that we expect this cannot fail. We move like this so that it can push things, but we don't expect failure.
@@ -57,39 +55,23 @@ namespace ParametricMovement
 		// the output Position in the case where a Move is blocked (E.g, you move 50% towards the desired location)
 
 		FTransform StartingTransform;
-		Driver->MapTimeToTransform(InputState.Position, StartingTransform);
+		Motion->MapTimeToTransform(InputState.Position, StartingTransform);
 
 		FTransform NewTransform;
-		Driver->MapTimeToTransform(OutputState.Position, NewTransform);
+		Motion->MapTimeToTransform(OutputState.Position, NewTransform);
 
 		FHitResult Hit(1.f);
 
 		FVector Delta = NewTransform.GetLocation() - StartingTransform.GetLocation();
-		BaseMovementDriver.MoveUpdatedComponent(Delta, NewTransform.GetRotation(), true, &Hit, ETeleportType::None);
+		MoveUpdatedComponent(Delta, NewTransform.GetRotation(), true, &Hit, ETeleportType::None);
 
 		if (Hit.IsValidBlockingHit())
 		{
-			FTransform ActualTransform = BaseMovementDriver.GetUpdateComponentTransform();
+			FTransform ActualTransform = GetUpdateComponentTransform();
 			FVector ActualDelta = NewTransform.GetLocation() - ActualTransform.GetLocation();
 			UE_LOG(LogParametricMovement, Warning, TEXT("Blocking hit occurred when trying to move parametric mover. ActualDelta: %s"), *ActualDelta.ToString());
 		}
 	}
-
-	void FMoveState::VisualLog(const FVisualLoggingParameters& Parameters, IMovementDriver* Driver, IMovementDriver* LogDriver) const
-	{
-		IBaseMovementDriver& BaseMovementDriver = Driver->GetBaseMovementDriver();
-		IBaseMovementDriver& LogMovementDriver = LogDriver->GetBaseMovementDriver();
-
-		FTransform Transform;
-		Driver->MapTimeToTransform(Position, Transform);
-
-		IBaseMovementDriver::FDrawDebugParams DrawParams(Parameters, &LogMovementDriver);
-		DrawParams.Transform = Transform;
-		DrawParams.InWorldText = LexToString(Parameters.Keyframe);
-		DrawParams.LogText = FString::Printf(TEXT("[%d] %s. Position: %.4f. Location: %s. Rotation: %s"), Parameters.Keyframe, *LexToString(Parameters.Context), Position, *Transform.GetLocation().ToString(), *Transform.GetRotation().Rotator().ToString());
-
-		BaseMovementDriver.DrawDebug(DrawParams);
-	}
 };
 
 
@@ -109,22 +91,23 @@ INetworkSimulationModel* UParametricMovementComponent::InstantiateNetworkSimulat
 		return nullptr;
 	}
 
+	ParametricMovementSimulation.Reset(new FParametricMovementSimulation());
+	ParametricMovementSimulation->Motion = &ParametricMotion;
+
 	if (ParametricMoverCVars::FixStep > 0)
 	{
-		auto *NewSim = new ParametricMovement::FMovementSystem<16>(this);
-		NewSim->RepProxy_Simulated.bAllowSimulatedExtrapolation = !bEnableInterpolation;
-		DO_NETSIM_MODEL_DEBUG(FNetworkSimulationModelDebuggerManager::Get().RegisterNetworkSimulationModel(NewSim, GetOwner()));
-		//MovementSyncState = &NewSim->SyncAccessor;
-		//MovementAuxState = &NewSim->AuxAccessor;
-		return NewSim;
+		auto *NewModel = new ParametricMovement::FMovementSystem<32>(ParametricMovementSimulation.Get(), this);
+		NewModel->RepProxy_Simulated.bAllowSimulatedExtrapolation = !bEnableInterpolation;
+		MovementSyncState.Init(NewModel);
+		MovementAuxState.Init(NewModel);
+		return NewModel;
 	}
 	
-	auto *NewSim = new ParametricMovement::FMovementSystem<>(this);
-	NewSim->RepProxy_Simulated.bAllowSimulatedExtrapolation = !bEnableInterpolation;
-	DO_NETSIM_MODEL_DEBUG(FNetworkSimulationModelDebuggerManager::Get().RegisterNetworkSimulationModel(NewSim, GetOwner()));
-	//MovementSyncState = &NewSim->SyncAccessor;
-	//MovementAuxState = &NewSim->AuxAccessor;
-	return NewSim;
+	auto *NewModel = new ParametricMovement::FMovementSystem<>(ParametricMovementSimulation.Get(), this);
+	NewModel->RepProxy_Simulated.bAllowSimulatedExtrapolation = !bEnableInterpolation;
+	MovementSyncState.Init(NewModel);
+	MovementAuxState.Init(NewModel);
+	return NewModel;
 }
 
 FNetworkSimulationModelInitParameters UParametricMovementComponent::GetSimulationInitParameters(ENetRole Role)
@@ -142,13 +125,13 @@ FNetworkSimulationModelInitParameters UParametricMovementComponent::GetSimulatio
 void UParametricMovementComponent::BeginPlay()
 {
 	Super::BeginPlay();
-	CachedStartingTransform = UpdatedComponent->GetComponentToWorld();
+	ParametricMotion.CachedStartingTransform = UpdatedComponent->GetComponentToWorld();
 }
 
 void UParametricMovementComponent::FinalizeFrame(const ParametricMovement::FMoveState& SyncState, const ParametricMovement::FAuxState& AuxState)
 {
 	FTransform NewTransform;
-	MapTimeToTransform(SyncState.Position, NewTransform);
+	ParametricMotion.MapTimeToTransform(SyncState.Position, NewTransform);
 
 	check(UpdatedComponent);
 	UpdatedComponent->SetWorldTransform(NewTransform, false, nullptr, ETeleportType::TeleportPhysics);
@@ -158,41 +141,19 @@ FString UParametricMovementComponent::GetDebugName() const
 {
 	return FString::Printf(TEXT("ParametricMovement. %s. %s"), *UEnum::GetValueAsString(TEXT("Engine.ENetRole"), GetOwnerRole()), *GetName());
 }
-const UObject* UParametricMovementComponent::GetVLogOwner() const
+const AActor* UParametricMovementComponent::GetVLogOwner() const
 {
 	return GetOwner();
 }
 
-void UParametricMovementComponent::AdvanceParametricTime(const float InPosition, const float InPlayRate, float &OutPosition, float& OutPlayRate, const float DeltaTimeSeconds) const
+void UParametricMovementComponent::VisualLog(const ParametricMovement::FInputCmd* Input, const ParametricMovement::FMoveState* Sync, const ParametricMovement::FAuxState* Aux, const FVisualLoggingParameters& SystemParameters) const
 {
-	// Real simple oscillation for now
-	OutPosition = InPosition + (InPlayRate * DeltaTimeSeconds);
-	OutPlayRate = InPlayRate;
-
-	const float DeltaMax = OutPosition - MaxTime;
-	if (DeltaMax > SMALL_NUMBER)
-	{
-		OutPosition = MaxTime - DeltaMax;
-		OutPlayRate *= -1.f;
-	}
-	else
-	{
-		const float DeltaMin = OutPosition - MinTime;
-		if (DeltaMin < SMALL_NUMBER)
-		{
-			OutPosition = MinTime - DeltaMin;
-			OutPlayRate *= -1.f;
-		}
-	}
+	FTransform Transform;
+	ParametricMotion.MapTimeToTransform(Sync->Position, Transform);
+	FVisualLoggingHelpers::VisualLogActor(GetOwner(), Transform, SystemParameters);
+	//DrawParams.LogText = FString::Printf(TEXT("[%d] %s. Position: %.4f. Location: %s. Rotation: %s"), Parameters.Keyframe, *LexToString(Parameters.Context), Position, *Transform.GetLocation().ToString(), *Transform.GetRotation().Rotator().ToString());
 }
 
-void UParametricMovementComponent::MapTimeToTransform(const float InPosition, FTransform& OutTransform) const
-{
-	const FVector Delta = ParametricDelta * InPosition;
-
-	OutTransform = CachedStartingTransform;
-	OutTransform.AddToTranslation(Delta);
-}
 
 void UParametricMovementComponent::TickComponent(float DeltaTime, enum ELevelTick TickType, FActorComponentTickFunction* ThisTickFunction)
 {
@@ -235,4 +196,38 @@ void UParametricMovementComponent::ProduceInput(const FNetworkSimTime DeltaTimeS
 {
 	Cmd.PlayRate = PendingPlayRate;
 	PendingPlayRate.Reset();
-}
\ No newline at end of file
+}
+
+// -------------------------------------------------------------------------------------------
+// -------------------------------------------------------------------------------------------
+
+void FSimpleParametricMotion::MapTimeToTransform(const float InPosition, FTransform& OutTransform) const
+{
+	const FVector Delta = ParametricDelta * InPosition;
+
+	OutTransform = CachedStartingTransform;
+	OutTransform.AddToTranslation(Delta);
+}
+
+void FSimpleParametricMotion::AdvanceParametricTime(const float InPosition, const float InPlayRate, float &OutPosition, float& OutPlayRate, const float DeltaTimeSeconds) const
+{
+	// Real simple oscillation for now
+	OutPosition = InPosition + (InPlayRate * DeltaTimeSeconds);
+	OutPlayRate = InPlayRate;
+
+	const float DeltaMax = OutPosition - MaxTime;
+	if (DeltaMax > SMALL_NUMBER)
+	{
+		OutPosition = MaxTime - DeltaMax;
+		OutPlayRate *= -1.f;
+	}
+	else
+	{
+		const float DeltaMin = OutPosition - MinTime;
+		if (DeltaMin < SMALL_NUMBER)
+		{
+			OutPosition = MinTime - DeltaMin;
+			OutPlayRate *= -1.f;
+		}
+	}
+}
diff --git a/Engine/Plugins/Runtime/NetworkPrediction/Source/NetworkPrediction/Private/NetworkPredictionTypes.cpp b/Engine/Plugins/Runtime/NetworkPrediction/Source/NetworkPrediction/Private/NetworkPredictionTypes.cpp
--- a/Engine/Plugins/Runtime/NetworkPrediction/Source/NetworkPrediction/Private/NetworkPredictionTypes.cpp
+++ b/Engine/Plugins/Runtime/NetworkPrediction/Source/NetworkPrediction/Private/NetworkPredictionTypes.cpp
@@ -2,9 +2,31 @@
 
 #include "NetworkPredictionTypes.h"
 #include "Engine/NetConnection.h"
+#include "GameFramework/Actor.h"
+#include "Engine/World.h"
+#include "Components/PrimitiveComponent.h"
+#include "Components/CapsuleComponent.h"
+#include "DrawDebugHelpers.h"
+#include "VisualLogger/VisualLogger.h"
 
 DEFINE_LOG_CATEGORY(LogNetworkSim);
 
+namespace NetSimVLogCVars
+{
+
+static int32 UseVLogger = 1;
+static FAutoConsoleVariableRef CVarUseVLogger(TEXT("ns.Debug.UseUnrealVLogger"),
+	UseVLogger,	TEXT("Use Unreal Visual Logger\n"),	ECVF_Default);
+
+static int32 UseDrawDebug = 1;
+static FAutoConsoleVariableRef CVarUseDrawDebug(TEXT("ns.Debug.UseDrawDebug"),
+	UseDrawDebug,	TEXT("Use built in DrawDebug* functions for visual logging\n"), ECVF_Default);
+
+static float DrawDebugDefaultLifeTime = 20.f;
+static FAutoConsoleVariableRef CVarDrawDebugDefaultLifeTime(TEXT("ns.Debug.DrawDebugLifetime.DefaultLifeTime"),
+	DrawDebugDefaultLifeTime, TEXT("Use built in DrawDebug* functions for visual logging"), ECVF_Default);
+};
+
 FColor FVisualLoggingParameters::DebugColors[(int32)EVisualLoggingContext::MAX] =
 {
 	FColor::Cyan,		// LastPredicted
@@ -25,6 +47,81 @@ FColor FVisualLoggingParameters::DebugColors[(int32)EVisualLoggingContext::MAX]
 	FColor(0, 0, 255),	// InterpolationSpeedUp
 };
 
+void FVisualLoggingHelpers::VisualLogActor(AActor* Owner, FTransform& Transform, const FVisualLoggingParameters& Params)
+{
+	const bool PersistentLines = (Params.Lifetime == EVisualLoggingLifetime::Persistent);
+	const float LifetimeSeconds = PersistentLines ? NetSimVLogCVars::DrawDebugDefaultLifeTime : -1.f;
+	const bool bDrawCumb = (Params.Context == EVisualLoggingContext::OtherMispredicted || Params.Context == EVisualLoggingContext::OtherPredicted);
+	const FColor DrawColor = Params.GetDebugColor();
+
+	UWorld* World = Owner->GetWorld();
+
+	if (bDrawCumb)
+	{
+		if (NetSimVLogCVars::UseDrawDebug)
+		{
+			static float PointSize = 3.f;
+			DrawDebugPoint(World, Transform.GetLocation(), PointSize, DrawColor, false, LifetimeSeconds);
+		}
+
+		if (NetSimVLogCVars::UseVLogger)
+		{
+			static float CrumbRadius = 3.f;
+			UE_VLOG_LOCATION(Owner, LogNetworkSim, Log, Transform.GetLocation(), CrumbRadius, DrawColor, TEXT("%s"), *Params.DebugString);
+		}
+	}
+	else
+	{
+		// Full drawing. Trying to provide a generic implementation that will give the most accurate representation.
+		// Still, subclasses may want to customize this
+		if (UCapsuleComponent* CapsuleComponent = Owner->FindComponentByClass<UCapsuleComponent>())
+		{
+			float Radius=0.f;
+			float HalfHeight=0.f;
+			CapsuleComponent->GetScaledCapsuleSize(Radius, HalfHeight);
+
+			if (NetSimVLogCVars::UseDrawDebug)
+			{
+				static const float Thickness = 2.f;
+				DrawDebugCapsule(World, Transform.GetLocation(), HalfHeight, Radius, Transform.GetRotation(), DrawColor, false, LifetimeSeconds, 0.f, Thickness);
+			}
+
+			if (NetSimVLogCVars::UseVLogger)
+			{
+				FVector VLogPosition = Transform.GetLocation();
+				VLogPosition.Z -= HalfHeight;
+				UE_VLOG_CAPSULE(Owner, LogNetworkSim, Log, VLogPosition, HalfHeight, Radius, Transform.GetRotation(), DrawColor, TEXT("%s"), *Params.DebugString);
+			}
+		}
+		else
+		{
+			// Generic Actor Bounds drawing
+			FBox LocalSpaceBox = Owner->CalculateComponentsBoundingBoxInLocalSpace();
+
+			if (NetSimVLogCVars::UseDrawDebug)
+			{
+				static const float Thickness = 2.f;
+
+				FVector ActorOrigin;
+				FVector ActorExtent;
+				LocalSpaceBox.GetCenterAndExtents(ActorOrigin, ActorExtent);
+				ActorExtent *= Transform.GetScale3D();
+				DrawDebugBox(World, Transform.GetLocation(), ActorExtent, Transform.GetRotation(), DrawColor, false, LifetimeSeconds, 0, Thickness);
+			}
+
+			if (NetSimVLogCVars::UseVLogger)
+			{
+				UE_VLOG_OBOX(Owner, LogNetworkSim, Log, LocalSpaceBox, Transform.ToMatrixWithScale(), DrawColor, TEXT("%s"), *Params.DebugString);
+			}
+		}
+	}
+
+	if (NetSimVLogCVars::UseVLogger)
+	{
+		UE_VLOG(Owner, LogNetworkSim, Log, TEXT("%s"), *Params.DebugString);
+	}
+}
+
 // -------------------------------------------------------------------------------------------------------------------------------
 //	FReplicationProxy
 // -------------------------------------------------------------------------------------------------------------------------------
diff --git a/Engine/Plugins/Runtime/NetworkPrediction/Source/NetworkPrediction/Private/NetworkSimulationModelDebugger.cpp b/Engine/Plugins/Runtime/NetworkPrediction/Source/NetworkPrediction/Private/NetworkSimulationModelDebugger.cpp
--- a/Engine/Plugins/Runtime/NetworkPrediction/Source/NetworkPrediction/Private/NetworkSimulationModelDebugger.cpp
+++ b/Engine/Plugins/Runtime/NetworkPrediction/Source/NetworkPrediction/Private/NetworkSimulationModelDebugger.cpp
@@ -18,7 +18,7 @@ FNetworkSimulationModelDebuggerManager& FNetworkSimulationModelDebuggerManager::
 	return Manager;
 }
 
-UObject* FindReplicatedObjectOnPIEServer(UObject* ClientObject)
+UObject* FindReplicatedObjectOnPIEServer(const UObject* ClientObject)
 {
 	if (ClientObject == nullptr)
 	{
diff --git a/Engine/Plugins/Runtime/NetworkPrediction/Source/NetworkPrediction/Public/Movement/BaseMovementComponent.h b/Engine/Plugins/Runtime/NetworkPrediction/Source/NetworkPrediction/Public/Movement/BaseMovementComponent.h
--- a/Engine/Plugins/Runtime/NetworkPrediction/Source/NetworkPrediction/Public/Movement/BaseMovementComponent.h
+++ b/Engine/Plugins/Runtime/NetworkPrediction/Source/NetworkPrediction/Public/Movement/BaseMovementComponent.h
@@ -10,49 +10,34 @@
 
 #include "BaseMovementComponent.generated.h"
 
-// -------------------------------------------------------------------------------------------------------------------------------
-//	Base component for movement. This essentially has the generic glue for selecting an UpdatedComponent and moving it along the world
-//	It is abstract in that you still need to define which simulation the component runs (via ::InstantiateNetworkSimulation)
-// -------------------------------------------------------------------------------------------------------------------------------
-
-// This is the base driver, to be used by the backing movement TNetworkedSimulationModel. 
-// It is the basic interface for moving and updating a primitive component that is owned by the UNetworkPredictionComponent component.
-// It is unfortunate and not ideal for perf that we have to go through a virtual function call to do collision queries. 
-// It also creates an awkward "dual hierarchy" when trying to share code like this.
-// A world where the simulation can do all the collision work and just push the final result back to the component layer would be much better.
-class IBaseMovementDriver
+// This provides a base for movement related simulations: moving an "UpdatedComponent" around the world.
+// There is no actual Update function provided here, it will be implemented by subclasses.
+class NETWORKPREDICTION_API FBaseMovementSimulation
 {
 public:
-	virtual bool SafeMoveUpdatedComponent(const FVector& Delta, const FQuat& NewRotation, bool bSweep, FHitResult& OutHit, ETeleportType Teleport) const = 0;
-	virtual bool MoveUpdatedComponent(const FVector& Delta, const FQuat& NewRotation, bool bSweep, FHitResult* OutHit, ETeleportType Teleport) const = 0;
-	virtual FTransform GetUpdateComponentTransform() const = 0;
-
-	virtual UWorld* GetBaseMovementDriverWorld() const = 0;
-	virtual UObject* GetBaseMovementVLogOwner() const = 0;
-
-	struct FDrawDebugParams
-	{
-		FDrawDebugParams() { }
-		FDrawDebugParams(const struct FVisualLoggingParameters& Parameters, IBaseMovementDriver* LogDriver);
-
-		UWorld* DebugWorld;		// World to draw stuff to (can be different than this->GetWorld() in PIE case)
-		UObject* DebugLogOwner; // Object that should own any VLog entries
-		
-		EVisualLoggingDrawType DrawType;
-		EVisualLoggingLifetime Lifetime;
-		FColor DrawColor;
-
-		FTransform Transform;	// Where to draw (important to understand: you don't draw your current state, you draw where the simulation says to draw)
-		FString InWorldText;
-		FString LogText;
-	};
+	bool SafeMoveUpdatedComponent(const FVector& Delta, const FQuat& NewRotation, bool bSweep, FHitResult& OutHit, ETeleportType Teleport) const;
+	bool MoveUpdatedComponent(const FVector& Delta, const FQuat& NewRotation, bool bSweep, FHitResult* OutHit, ETeleportType Teleport) const;
+	FTransform GetUpdateComponentTransform() const;
+
+	FVector GetPenetrationAdjustment(const FHitResult& Hit) const;
+
+	bool OverlapTest(const FVector& Location, const FQuat& RotationQuat, const ECollisionChannel CollisionChannel, const FCollisionShape& CollisionShape, const AActor* IgnoreActor) const;
+	void InitCollisionParams(FCollisionQueryParams &OutParams, FCollisionResponseParams& OutResponseParam) const;
+	bool ResolvePenetration(const FVector& ProposedAdjustment, const FHitResult& Hit, const FQuat& NewRotationQuat) const;
 	
-	// This is the final debug drawing function that the driver can implement. UBaseMovementComponent implements a generic one that will draw capsules or actor bounds.
-	virtual void DrawDebug(const FDrawDebugParams& Params) const = 0;
+	USceneComponent* UpdatedComponent = nullptr;
+	UPrimitiveComponent* UpdatedPrimitive = nullptr;
+
+	/**  Flags that control the behavior of calls to MoveComponent() on our UpdatedComponent. */
+	mutable EMoveComponentFlags MoveComponentFlags = MOVECOMP_NoFlags; // Mutable because we sometimes need to disable these flags ::ResolvePenetration. Better way may be possible
 };
 
+// -------------------------------------------------------------------------------------------------------------------------------
+//	Base component for movement. This essentially has the generic glue for selecting an UpdatedComponent and moving it along the world
+//	It is abstract in that you still need to define which simulation the component runs (via ::InstantiateNetworkSimulation)
+// -------------------------------------------------------------------------------------------------------------------------------
 UCLASS(Abstract)
-class NETWORKPREDICTION_API UBaseMovementComponent : public UNetworkPredictionComponent, public IBaseMovementDriver
+class NETWORKPREDICTION_API UBaseMovementComponent : public UNetworkPredictionComponent
 {
 	GENERATED_BODY()
 
@@ -64,16 +49,6 @@ public:
 	virtual void OnRegister() override;
 	virtual void RegisterComponentTickFunctions(bool bRegister) override;
 
-	// IBaseMovementDriver
-	bool SafeMoveUpdatedComponent(const FVector& Delta, const FQuat& NewRotation, bool bSweep, FHitResult& OutHit, ETeleportType Teleport) const override final;
-	bool MoveUpdatedComponent(const FVector& Delta, const FQuat& NewRotation, bool bSweep, FHitResult* OutHit, ETeleportType Teleport) const override final;
-	FTransform GetUpdateComponentTransform() const override final;
-	
-	UWorld* GetBaseMovementDriverWorld() const override final { return GetWorld(); }
-	UObject* GetBaseMovementVLogOwner() const override final;	
-
-	virtual void DrawDebug(const IBaseMovementDriver::FDrawDebugParams& Params) const override;
-
 protected:
 
 	// Basic "Update Component/Ticking"
@@ -81,7 +56,7 @@ protected:
 	virtual void UpdateTickRegistration();
 
 	UFUNCTION()
-	virtual void PhysicsVolumeChanged(class APhysicsVolume* NewVolume);
+	virtual void PhysicsVolumeChanged(class APhysicsVolume* NewVolume);	
 
 	UPROPERTY()
 	USceneComponent* UpdatedComponent = nullptr;
@@ -96,13 +71,4 @@ private:
 	
 	/** Transient flag indicating whether we are executing InitializeComponent(). */
 	bool bInInitializeComponent = false;
-
-	static FVector GetPenetrationAdjustment(const FHitResult& Hit);
-
-	bool OverlapTest(const FVector& Location, const FQuat& RotationQuat, const ECollisionChannel CollisionChannel, const FCollisionShape& CollisionShape, const AActor* IgnoreActor) const;
-	void InitCollisionParams(FCollisionQueryParams &OutParams, FCollisionResponseParams& OutResponseParam) const;
-	bool ResolvePenetration(const FVector& ProposedAdjustment, const FHitResult& Hit, const FQuat& NewRotationQuat) const;
-
-	/**  Flags that control the behavior of calls to MoveComponent() on our UpdatedComponent. */
-	mutable EMoveComponentFlags MoveComponentFlags = MOVECOMP_NoFlags; // Mutable because we sometimes need to disable these flags ::ResolvePenetration. Better way may be possible
 };
\ No newline at end of file
diff --git a/Engine/Plugins/Runtime/NetworkPrediction/Source/NetworkPrediction/Public/Movement/FlyingMovement.h b/Engine/Plugins/Runtime/NetworkPrediction/Source/NetworkPrediction/Public/Movement/FlyingMovement.h
--- a/Engine/Plugins/Runtime/NetworkPrediction/Source/NetworkPrediction/Public/Movement/FlyingMovement.h
+++ b/Engine/Plugins/Runtime/NetworkPrediction/Source/NetworkPrediction/Public/Movement/FlyingMovement.h
@@ -22,13 +22,10 @@
 // FlyingMovement: simple flying movement that was based on UE4's FloatingPawnMovement
 // -------------------------------------------------------------------------------------------------------------------------------
 
-class IFlyingMovementDriver;
 class USceneComponent;
 
 namespace FlyingMovement
 {
-	class IMovementDriver;
-
 	// State the client generates
 	struct FInputCmd
 	{
@@ -103,8 +100,6 @@ namespace FlyingMovement
 			}
 		}
 
-		void VisualLog(const FVisualLoggingParameters& Parameters, IMovementDriver* Driver, IMovementDriver* LogDriver) const;
-
 		static void Interpolate(const FMoveState& From, const FMoveState& To, const float PCT, FMoveState& OutDest)
 		{
 			OutDest.Location = From.Location + ((To.Location - From.Location) * PCT);
@@ -144,41 +139,37 @@ namespace FlyingMovement
 			}
 		}
 	};
-
+	
 	using TMovementBufferTypes = TNetworkSimBufferTypes<FInputCmd, FMoveState, FAuxState>;
 
 	static FName SimulationGroupName("FlyingMovement");
 
-	// Interface between the simulation and owning component driving it. Functions added here are available in ::Update and must be implemented by UMockNetworkSimulationComponent.
-	class IMovementDriver : public TNetworkSimDriverInterfaceBase<TMovementBufferTypes>
+	class FMovementSimulation : public FBaseMovementSimulation
 	{
 	public:
 
-		// Interface for moving the collision component around
-		virtual IBaseMovementDriver& GetBaseMovementDriver() = 0;
+		/** Tick group the simulation maps to */
+		static const FName GroupName;
+
+		/** Dev tool to force simple mispredict */
+		static bool ForceMispredict;
+
+		/** Main update function */
+		NETWORKPREDICTION_API void Update(const float DeltaSeconds, const FInputCmd& InputCmd, const FMoveState& InputState, FMoveState& OutputState, const FAuxState& AuxState, const TNetSimLazyWriter<FAuxState>& OutAuxStateAccessor);
 
 		// Called prior to running the sim to make sure to make sure the collision component is in the right place. 
 		// This is unfortunate and not good, but is needed to ensure our collision and world position have not been moved out from under us.
 		// Refactoring primitive component movement to allow the sim to do all collision queries outside of the component code would be ideal.
-		virtual void PreSimSync(const FMoveState& SyncState) = 0;
-	};
-
-	class FMovementSimulation
-	{
-	public:
-		/** Main update function */
-		static void Update(IMovementDriver* Driver, const float DeltaSeconds, const FInputCmd& InputCmd, const FMoveState& InputState, FMoveState& OutputState, const FAuxState& AuxState, const TLazyStateAccessor<FAuxState>& OutAuxStateAccessor);
+		void PreSimSync(const FMoveState& SyncState);
 
-		/** Tick group the simulation maps to */
-		static const FName GroupName;
+	protected:
 
-		/** Dev tool to force simple mispredict */
-		static bool ForceMispredict;
+		float SlideAlongSurface(const FVector& Delta, float Time, const FQuat Rotation, const FVector& Normal, FHitResult& Hit, bool bHandleImpact);
 	};
 
 	// Actual definition of our network simulation.
 	template<int32 InFixedStepMS=0>
-	using FMovementSystem = TNetworkedSimulationModel<FMovementSimulation, IMovementDriver, TMovementBufferTypes, TNetworkSimTickSettings<InFixedStepMS> >;
+	using FMovementSystem = TNetworkedSimulationModel<FMovementSimulation, TMovementBufferTypes, TNetworkSimTickSettings<InFixedStepMS>>;
 
 	// general tolerance value for rotation checks
 	static const float ROTATOR_TOLERANCE = (1e-3);
@@ -187,14 +178,14 @@ namespace FlyingMovement
 
 // Needed to trick UHT into letting UMockNetworkSimulationComponent implement. UHT cannot parse the ::
 // Also needed for forward declaring. Can't just be a typedef/using =
-class IFlyingMovementDriver : public FlyingMovement::IMovementDriver { };
+class IFlyingMovementSystemDriver : public TNetworkedSimulationModelDriver<FlyingMovement::TMovementBufferTypes> { };
 
 // -------------------------------------------------------------------------------------------------------------------------------
 // ActorComponent for running FlyingMovement 
 // -------------------------------------------------------------------------------------------------------------------------------
 
 UCLASS(BlueprintType, meta=(BlueprintSpawnableComponent))
-class NETWORKPREDICTION_API UFlyingMovementComponent : public UBaseMovementComponent, public IFlyingMovementDriver
+class NETWORKPREDICTION_API UFlyingMovementComponent : public UBaseMovementComponent, public IFlyingMovementSystemDriver
 {
 	GENERATED_BODY()
 
@@ -204,24 +195,23 @@ public:
 	
 	virtual void TickComponent(float DeltaTime, enum ELevelTick TickType, FActorComponentTickFunction *ThisTickFunction) override;
 
-	IBaseMovementDriver& GetBaseMovementDriver() override final { return *static_cast<IBaseMovementDriver*>(this); }
-
-	FString GetDebugName() const override;
-	const UObject* GetVLogOwner() const override;
-	
-	void ProduceInput(const FNetworkSimTime SimTime, FlyingMovement::FInputCmd& Cmd) override;
-	void FinalizeFrame(const FlyingMovement::FMoveState& SyncState, const FlyingMovement::FAuxState& AuxState) override;
-
-	void PreSimSync(const FlyingMovement::FMoveState& SyncState) override;
-
 	DECLARE_DELEGATE_TwoParams(FProduceFlyingInput, const FNetworkSimTime /*SimTime*/, FlyingMovement::FInputCmd& /*Cmd*/)
 	FProduceFlyingInput ProduceInputDelegate;
 
 	TNetworkSimStateAccessor<FlyingMovement::FMoveState> MovementSyncState;
 	TNetworkSimStateAccessor<FlyingMovement::FAuxState> MovementAuxState;
 
+	// IFlyingMovementSystemDriver
+	FString GetDebugName() const override;
+	const AActor* GetVLogOwner() const override;
+	void VisualLog(const FlyingMovement::FInputCmd* Input, const FlyingMovement::FMoveState* Sync, const FlyingMovement::FAuxState* Aux, const FVisualLoggingParameters& SystemParameters) const override;
+	
+	void ProduceInput(const FNetworkSimTime SimTime, FlyingMovement::FInputCmd& Cmd) override;
+	void FinalizeFrame(const FlyingMovement::FMoveState& SyncState, const FlyingMovement::FAuxState& AuxState) override;
+
 protected:
 
 	// Network Prediction
 	virtual INetworkSimulationModel* InstantiateNetworkSimulation() override;
+	TUniquePtr<FlyingMovement::FMovementSimulation> MovementSimulation;
 };
\ No newline at end of file
diff --git a/Engine/Plugins/Runtime/NetworkPrediction/Source/NetworkPrediction/Public/Movement/ParametricMovement.h b/Engine/Plugins/Runtime/NetworkPrediction/Source/NetworkPrediction/Public/Movement/ParametricMovement.h
--- a/Engine/Plugins/Runtime/NetworkPrediction/Source/NetworkPrediction/Public/Movement/ParametricMovement.h
+++ b/Engine/Plugins/Runtime/NetworkPrediction/Source/NetworkPrediction/Public/Movement/ParametricMovement.h
@@ -7,7 +7,7 @@
 
 #include "ParametricMovement.generated.h"
 
-class IParametricMovementDriver;
+struct FSimpleParametricMotion;
 
 namespace ParametricMovement
 {
@@ -106,42 +106,53 @@ namespace ParametricMovement
 	};
 
 	using TMovementBufferTypes = TNetworkSimBufferTypes<FInputCmd, FMoveState, FAuxState>;
-
-	// Interface between the simulation and owning component driving it. Functions added here are available in ::Update
-	class IMovementDriver : public TNetworkSimDriverInterfaceBase<TMovementBufferTypes>
-	{
-	public:
-
-		// BaseMovement driver API (functions for moving around a primitive component)
-		virtual IBaseMovementDriver& GetBaseMovementDriver() = 0;
-
-		// Advance parametric time. This is meant to do simple things like looping/reversing etc.
-		// Note how this should be STATIC and not rely on state outside of what is passed in (such thing would need to be done inside the simulation, not through the driver!)
-		virtual void AdvanceParametricTime(const float InPosition, const float InPlayRate, float &OutPosition, float& OutPlayRate, const float DeltaTimeSeconds) const = 0;
-
-		// Actually turn the given position into a transform. Again, should be static and not conditional on changing state outside of the network sim
-		virtual void MapTimeToTransform(const float InPosition, FTransform& OutTransform) const = 0;
-	};
 	
-	class FMovementSimulation
+	class FMovementSimulation : public FBaseMovementSimulation
 	{
 	public:
-		static void Update(IMovementDriver* Driver, const float DeltaSeconds, const FInputCmd& InputCmd, const FMoveState& InputState, FMoveState& OutputState, const FAuxState& AuxState, const TLazyStateAccessor<FAuxState>& OutAuxStateAccessor);
 		static const FName GroupName;
+
+		void Update(const float DeltaSeconds, const FInputCmd& InputCmd, const FMoveState& InputState, FMoveState& OutputState, const FAuxState& AuxState, const TNetSimLazyWriter<FAuxState>& OutAuxStateAccessor);
+
+		// Pointer to our static mapping of time->position
+		const FSimpleParametricMotion* Motion = nullptr;
 	};
 
 
 	// Actual definition of our network simulation.
 	template<int32 FixedStepMS=0>
-	using FMovementSystem = TNetworkedSimulationModel<FMovementSimulation, IMovementDriver, TMovementBufferTypes, TNetworkSimTickSettings<FixedStepMS>>;
-
-	//using TSimTime = FMovementSystem::TSimTime;
+	using FMovementSystem = TNetworkedSimulationModel<FMovementSimulation, TMovementBufferTypes, TNetworkSimTickSettings<FixedStepMS>>;
 
 } // End namespace
 
 // Needed to trick UHT into letting UMockNetworkSimulationComponent implement. UHT cannot parse the ::
 // Also needed for forward declaring. Can't just be a typedef/using =
-class IParametricMovementDriver : public ParametricMovement::IMovementDriver { };
+class IParametricMovementDriver : public TNetworkedSimulationModelDriver<ParametricMovement::TMovementBufferTypes> { };
+class FParametricMovementSimulation : public ParametricMovement::FMovementSimulation { };
+
+// Extremely simple struct for defining parametric motion. This is editable in UParametricMovementComponent's defaults, and also used by the simulation code above. 
+USTRUCT(BlueprintType)
+struct FSimpleParametricMotion
+{
+	GENERATED_BODY()
+
+	// Actually turn the given position into a transform. Again, should be static and not conditional on changing state outside of the network sim
+	void MapTimeToTransform(const float InPosition, FTransform& OutTransform) const;
+
+	// Advance parametric time. This is meant to do simple things like looping/reversing etc.
+	void AdvanceParametricTime(const float InPosition, const float InPlayRate, float &OutPosition, float& OutPlayRate, const float DeltaTimeSeconds) const;
+
+	UPROPERTY(EditAnywhere, BlueprintReadOnly, Category=ParametricMovement)
+	FVector ParametricDelta = FVector(0.f, 0.f, 500.f);
+
+	UPROPERTY(EditAnywhere, BlueprintReadOnly, Category=ParametricMovement)
+	float MinTime = -1.f;
+
+	UPROPERTY(EditAnywhere, BlueprintReadOnly, Category=ParametricMovement)
+	float MaxTime = 1.f;
+
+	FTransform CachedStartingTransform;
+};
 
 // -------------------------------------------------------------------------------------------------------------------------------
 //	ActorComponent for running basic Parametric movement. 
@@ -164,22 +175,22 @@ class NETWORKPREDICTION_API UParametricMovementComponent : public UBaseMovementC
 
 	// Base TNetworkModelSimulation driver
 	FString GetDebugName() const override;
-	const UObject* GetVLogOwner() const override;
-	void ProduceInput(const FNetworkSimTime SimTime, ParametricMovement::FInputCmd& Cmd);
-	void FinalizeFrame(const ParametricMovement::FMoveState& SyncState, const ParametricMovement::FAuxState& AuxState) override;	
-
-	// Base Movement Driver
-	IBaseMovementDriver& GetBaseMovementDriver() override final { return *static_cast<IBaseMovementDriver*>(this); }
+	const AActor* GetVLogOwner() const override;
+	void VisualLog(const ParametricMovement::FInputCmd* Input, const ParametricMovement::FMoveState* Sync, const ParametricMovement::FAuxState* Aux, const FVisualLoggingParameters& SystemParameters) const override;
 
-	// Parametric Movement Driver
-	virtual void AdvanceParametricTime(const float InPosition, const float InPlayRate, float &OutPosition, float& OutPlayRate, const float DeltaTimeSeconds) const override;
-	virtual void MapTimeToTransform(const float InPosition, FTransform& OutTransform) const override;
+	void ProduceInput(const FNetworkSimTime SimTime, ParametricMovement::FInputCmd& Cmd);
+	void FinalizeFrame(const ParametricMovement::FMoveState& SyncState, const ParametricMovement::FAuxState& AuxState) override;
 
 protected:
 
+	TNetworkSimStateAccessor<ParametricMovement::FMoveState> MovementSyncState;
+	TNetworkSimStateAccessor<ParametricMovement::FAuxState> MovementAuxState;
+
 	virtual INetworkSimulationModel* InstantiateNetworkSimulation() override;
 	FNetworkSimulationModelInitParameters GetSimulationInitParameters(ENetRole Role) override;
 
+	TUniquePtr<FParametricMovementSimulation> ParametricMovementSimulation;
+
 	// ------------------------------------------------------------------------
 	// Temp Parametric movement example
 	//	The essence of this movement simulation is to map some Time value to a transform. That is it.
@@ -191,15 +202,9 @@ protected:
 	/** Disables starting the simulation. For development/testing ease of use */
 	UPROPERTY(EditAnywhere, BlueprintReadOnly, Category=ParametricMovement)
 	bool bDisableParametricMovementSimulation = false;
-	
-	UPROPERTY(EditAnywhere, BlueprintReadOnly, Category=ParametricMovement)
-	FVector ParametricDelta = FVector(0.f, 0.f, 500.f);
-
-	UPROPERTY(EditAnywhere, BlueprintReadOnly, Category=ParametricMovement)
-	float MinTime = -1.f;
 
 	UPROPERTY(EditAnywhere, BlueprintReadOnly, Category=ParametricMovement)
-	float MaxTime = 1.f;
+	FSimpleParametricMotion ParametricMotion;
 
 	UPROPERTY(EditAnywhere, BlueprintReadOnly, Category=ParametricMovementNetworking)
 	bool bEnableDependentSimulation = false;
@@ -215,8 +220,6 @@ protected:
 	UPROPERTY(EditAnywhere, BlueprintReadOnly, Category=ParametricMovementNetworking)
 	float ParentNetUpdateFrequency = 0.f;
 
-	FTransform CachedStartingTransform;
-
 	TOptional<float> PendingPlayRate = 1.f;
 };
 
diff --git a/Engine/Plugins/Runtime/NetworkPrediction/Source/NetworkPrediction/Public/NetworkPredictionComponent.h b/Engine/Plugins/Runtime/NetworkPrediction/Source/NetworkPrediction/Public/NetworkPredictionComponent.h
--- a/Engine/Plugins/Runtime/NetworkPrediction/Source/NetworkPrediction/Public/NetworkPredictionComponent.h
+++ b/Engine/Plugins/Runtime/NetworkPrediction/Source/NetworkPrediction/Public/NetworkPredictionComponent.h
@@ -12,7 +12,7 @@
 //	initialized and plugged into the UE4 replication system.
 //
 //	This is an abstract component and cannot function on its own. It must be subclassed and InstantiateNetworkSimulation must be implemented.
-//	Ticking and RPC sending will be handled automatically, but the subclass will need to supply input to the simulation (via TNetworkSimDriverInterfaceBase::ProduceInput).
+//	Ticking and RPC sending will be handled automatically, but the subclass will need to supply input to the simulation (via TNetworkedSimulationModelDriver::ProduceInput).
 //
 //	Its also worth pointing out that nothing about being a UActorComponent is essential here. All that this component does could be done within an AActor itself.
 //	An actor component makes sense for flexible/reusable code provided by the plugin. But there is nothing stopping you from copying this directly into an actor if you had reason to.
diff --git a/Engine/Plugins/Runtime/NetworkPrediction/Source/NetworkPrediction/Public/NetworkPredictionTypes.h b/Engine/Plugins/Runtime/NetworkPrediction/Source/NetworkPrediction/Public/NetworkPredictionTypes.h
--- a/Engine/Plugins/Runtime/NetworkPrediction/Source/NetworkPrediction/Public/NetworkPredictionTypes.h
+++ b/Engine/Plugins/Runtime/NetworkPrediction/Source/NetworkPrediction/Public/NetworkPredictionTypes.h
@@ -41,6 +41,21 @@ struct FNetSerializeParams
 	FArchive& Ar;
 };
 
+enum class EStandardLoggingContext: uint8
+{
+	HeaderOnly,		// Minimal logging
+	Full			// "Logs everything"
+};
+
+struct FStandardLoggingParameters
+{
+	FStandardLoggingParameters(FOutputDevice* InAr, EStandardLoggingContext InContext, int32 InKeyframe)
+		: Ar(InAr), Context(InContext), Keyframe(InKeyframe) { }
+	FOutputDevice* Ar;
+	EStandardLoggingContext Context;
+	int32 Keyframe;
+};
+
 struct FNetworkSimulationModelInitParameters
 {
 	int32 InputBufferSize=0;
@@ -76,21 +91,26 @@ struct FNetSimProcessedFrameDebugInfo
 
 		SafeNetSerializeTArray_Default<31>(P.Ar, ProcessedKeyframes);
 	}
-};
-
-enum class EStandardLoggingContext: uint8
-{
-	HeaderOnly,		// Minimal logging
-	Full			// "Logs everything"
-};
-
-struct FStandardLoggingParameters
-{
-	FStandardLoggingParameters(FOutputDevice* InAr, EStandardLoggingContext InContext, int32 InKeyframe)
-		: Ar(InAr), Context(InContext), Keyframe(InKeyframe) { }
-	FOutputDevice* Ar;
-	EStandardLoggingContext Context;
-	int32 Keyframe;
+	void Log(FStandardLoggingParameters& P) const
+	{
+		P.Ar->Logf(TEXT("LocalGFrameNumber: %d"), LocalGFrameNumber);
+		P.Ar->Logf(TEXT("LocalDeltaTimeSeconds: %.4f"), LocalDeltaTimeSeconds);
+		P.Ar->Logf(TEXT("PendingKeyframe: %d"), PendingKeyframe);
+		P.Ar->Logf(TEXT("HeadKeyframe: %d"), HeadKeyframe);
+		P.Ar->Logf(TEXT("RemainingAllowedSimulationTimeSeconds: %.4f"), RemainingAllowedSimulationTimeSeconds);
+
+		P.Ar->Logf(TEXT("LastSentInputKeyframe: %d"), LastSentInputKeyframe);
+		P.Ar->Logf(TEXT("LastReceivedInputKeyframe: %d"), LastReceivedInputKeyframe);
+
+		FString ProcessedKeyframesStr;
+		for (int32 k : ProcessedKeyframes)
+		{
+			ProcessedKeyframesStr += LexToString(k);
+			ProcessedKeyframesStr += TEXT(" ");
+		}
+		ProcessedKeyframesStr.TrimEndInline();
+		P.Ar->Logf(TEXT("ProcessedKeyframes: [%s]"), *ProcessedKeyframesStr);
+	}
 };
 
 UENUM()
@@ -125,12 +145,6 @@ inline FString LexToString(EVisualLoggingContext A)
 	return *UEnum::GetValueAsString(TEXT("NetworkPrediction.EVisualLoggingContext"), A);
 }
 
-enum class EVisualLoggingDrawType : uint8
-{
-	Full,				// Draw "the full thing" (maybe a collision capsule for example)
-	Crumb,				// Draw a small/minimal representation (E.g, a point or small axis)
-};
-
 enum class EVisualLoggingLifetime : uint8
 {
 	Transient,			// This logging is transient and will (probably) be done every frame. Don't persist (in contexts where that would matter).
@@ -139,17 +153,24 @@ enum class EVisualLoggingLifetime : uint8
 
 struct FVisualLoggingParameters
 {
-	FVisualLoggingParameters(EVisualLoggingContext InContext, int32 InKeyframe, EVisualLoggingLifetime InLifetime) :
-		 Context(InContext), Keyframe(InKeyframe), Lifetime(InLifetime) { }
+	FVisualLoggingParameters(EVisualLoggingContext InContext, int32 InKeyframe, EVisualLoggingLifetime InLifetime, const FString& Str=FString()) :
+		 Context(InContext), Keyframe(InKeyframe), Lifetime(InLifetime), DebugString(Str) { }
 
 	EVisualLoggingContext Context;
 	int32 Keyframe;
 	EVisualLoggingLifetime Lifetime;
+	FString DebugString;
 
 	FColor GetDebugColor() const { return DebugColors[(int32)Context]; }
 	static NETWORKPREDICTION_API FColor DebugColors[(int32)EVisualLoggingContext::MAX];
 };
 
+struct FVisualLoggingHelpers
+{
+	// Very basic function to draw an outline of the actor at the given transform
+	static NETWORKPREDICTION_API void VisualLogActor(AActor* Owner, FTransform& Transform, const FVisualLoggingParameters& Params);
+};
+
 UENUM()
 enum class ESimulatedUpdateMode : uint8
 {
@@ -221,7 +242,7 @@ public:
 	virtual void SetParentSimulation(INetworkSimulationModel* Simulation) = 0;
 	virtual INetworkSimulationModel* GetParentSimulation() const = 0;
 	
-	virtual void AddDepdentSimulation(INetworkSimulationModel* Simulation) = 0;
+	virtual void AddDependentSimulation(INetworkSimulationModel* Simulation) = 0;
 	virtual void RemoveDependentSimulation(INetworkSimulationModel* Simulation) = 0;
 	
 	// Tell parent sim that a dependent sim needs to reconcile (parent sim drives this)
diff --git a/Engine/Plugins/Runtime/NetworkPrediction/Source/NetworkPrediction/Public/NetworkSimulationModel.h b/Engine/Plugins/Runtime/NetworkPrediction/Source/NetworkPrediction/Public/NetworkSimulationModel.h
--- a/Engine/Plugins/Runtime/NetworkPrediction/Source/NetworkPrediction/Public/NetworkSimulationModel.h
+++ b/Engine/Plugins/Runtime/NetworkPrediction/Source/NetworkPrediction/Public/NetworkSimulationModel.h
@@ -5,6 +5,7 @@
 #include "NetworkSimulationModelBuffer.h"
 #include "NetworkSimulationModelTypes.h"
 #include "NetworkSimulationModelReplicators.h"
+#include "NetworkSimulationModelDebugger.h"
 
 // ----------------------------------------------------------------------------------------------------------------------------------------------------------------------
 //	TNetworkedSimulationModel
@@ -28,8 +29,7 @@
 // ----------------------------------------------------------------------------------------------------------------------------------------------------------------------
 
 template <
-	typename TSimulation,								// Final Simulation class. Used to call static T::Update function.
-	typename TInDriver,
+	typename InTSimulation,								// Final Simulation class
 	typename TUserBufferTypes,							// The user types (input, sync, aux, debug). Note this gets wrapped in TInternalBufferTypes internally.
 	typename InTTickSettings=TNetworkSimTickSettings<>, // Defines global rules about time keeping and ticking
 
@@ -46,7 +46,8 @@ class TNetworkedSimulationModel : public INetworkSimulationModel
 {
 public:
 
-	using TDriver = TInDriver;
+	using TSimulation = InTSimulation;
+	using TDriver = TNetworkedSimulationModelDriver<TUserBufferTypes>;
 
 	using TBufferTypes = TInternalBufferTypes<TUserBufferTypes, InTTickSettings>;
 	using TTickSettings = InTTickSettings;
@@ -59,13 +60,17 @@ public:
 	using TSimTime = FNetworkSimTime;
 	using TRealTime = FNetworkSimTime::FRealTime;
 	
-	TNetworkedSimulationModel(TDriver* InDriver, const TSyncState& InitialSyncState = TSyncState(), const TAuxState& InitialAuxState = TAuxState())
+	TNetworkedSimulationModel(TSimulation* InSimulation, TDriver* InDriver, const TSyncState& InitialSyncState = TSyncState(), const TAuxState& InitialAuxState = TAuxState())
 	{
+		check(InSimulation && InDriver);
+		Simulation = InSimulation;
 		Driver = InDriver;
 		*Buffers.Sync.WriteKeyframe(0) = InitialSyncState;
 		*Buffers.Aux.WriteKeyframe(0) = InitialAuxState;
 		TickInfo.SetTotalProcessedSimulationTime(FNetworkSimTime(), 0);
+		DO_NETSIM_MODEL_DEBUG(FNetworkSimulationModelDebuggerManager::Get().RegisterNetworkSimulationModel(this, Driver->GetVLogOwner()));
 	}
+
 	virtual ~TNetworkedSimulationModel()
 	{
 		SetParentSimulation(nullptr);
@@ -118,15 +123,15 @@ public:
 		switch (Parameters.Role)
 		{
 			case ROLE_Authority:
-				RepProxy_ServerRPC.template PreSimTick<TSimulation, TDriver>(Driver, Buffers, TickInfo, Parameters);
+				RepProxy_ServerRPC.template PreSimTick<TDriver>(Driver, Buffers, TickInfo, Parameters);
 			break;
 
 			case ROLE_AutonomousProxy:
-				RepProxy_Autonomous.template PreSimTick<TSimulation, TDriver>(Driver, Buffers, TickInfo, Parameters);
+				RepProxy_Autonomous.template PreSimTick<TDriver>(Driver, Buffers, TickInfo, Parameters);
 			break;
 
 			case ROLE_SimulatedProxy:
-				RepProxy_Simulated.template PreSimTick<TSimulation, TDriver>(Driver, Buffers, TickInfo, Parameters);
+				RepProxy_Simulated.template PreSimTick<TDriver>(Driver, Buffers, TickInfo, Parameters);
 			break;
 		}
 
@@ -164,7 +169,7 @@ public:
 
 				{
 					TScopedSimulationTick UpdateScope(TickInfo, OutputKeyframe, InputCmd->GetFrameDeltaTime());
-					TSimulation::Update(Driver, InputCmd->GetFrameDeltaTime().ToRealTimeSeconds(), *InputCmd, *InSyncState, *OutSyncState, *InAuxState, Buffers.Aux.WriteKeyframeFunc(OutputKeyframe));
+					Simulation->Update(InputCmd->GetFrameDeltaTime().ToRealTimeSeconds(), *InputCmd, *InSyncState, *OutSyncState, *InAuxState, Buffers.Aux.LazyWriter(OutputKeyframe));
 				}
 
 				if (DebugState)
@@ -228,15 +233,15 @@ public:
 		switch (Role)
 		{
 			case ROLE_Authority:
-				RepProxy_ServerRPC.template Reconcile<TSimulation, TDriver>(Driver, Buffers, TickInfo);
+				RepProxy_ServerRPC.template Reconcile<TSimulation, TDriver>(Simulation, Driver, Buffers, TickInfo);
 			break;
 
 			case ROLE_AutonomousProxy:
-				RepProxy_Autonomous.template Reconcile<TSimulation, TDriver>(Driver, Buffers, TickInfo);
+				RepProxy_Autonomous.template Reconcile<TSimulation, TDriver>(Simulation, Driver, Buffers, TickInfo);
 			break;
 
 			case ROLE_SimulatedProxy:
-				RepProxy_Simulated.template Reconcile<TSimulation, TDriver>(Driver, Buffers, TickInfo);
+				RepProxy_Simulated.template Reconcile<TSimulation, TDriver>(Simulation, Driver, Buffers, TickInfo);
 			break;
 		}
 	}
@@ -319,17 +324,17 @@ public:
 
 	// ------------------------------------------------------------------------------------------------------
 
-	void SetParentSimulation(INetworkSimulationModel* Simulation) final override
+	void SetParentSimulation(INetworkSimulationModel* ParentSimulation) final override
 	{
 		if (RepProxy_Simulated.ParentSimulation)
 		{
 			RepProxy_Simulated.ParentSimulation->RemoveDependentSimulation(this);
 		}
 		
-		RepProxy_Simulated.ParentSimulation = Simulation;
-		if (Simulation)
+		RepProxy_Simulated.ParentSimulation = ParentSimulation;
+		if (ParentSimulation)
 		{
-			Simulation->AddDepdentSimulation(this);
+			ParentSimulation->AddDependentSimulation(this);
 		}
 	}
 
@@ -338,16 +343,16 @@ public:
 		return RepProxy_Simulated.ParentSimulation;
 	}
 
-	void AddDepdentSimulation(INetworkSimulationModel* Simulation) final override
+	void AddDependentSimulation(INetworkSimulationModel* DependentSimulation) final override
 	{
-		check(RepProxy_Autonomous.DependentSimulations.Contains(Simulation) == false);
-		RepProxy_Autonomous.DependentSimulations.Add(Simulation);
+		check(RepProxy_Autonomous.DependentSimulations.Contains(DependentSimulation) == false);
+		RepProxy_Autonomous.DependentSimulations.Add(DependentSimulation);
 		NotifyDependentSimNeedsReconcile(); // force reconcile on purpose
 	}
 
-	void RemoveDependentSimulation(INetworkSimulationModel* Simulation) final override
+	void RemoveDependentSimulation(INetworkSimulationModel* DependentSimulation) final override
 	{
-		RepProxy_Autonomous.DependentSimulations.Remove(Simulation);
+		RepProxy_Autonomous.DependentSimulations.Remove(DependentSimulation);
 	}
 
 	void NotifyDependentSimNeedsReconcile()
@@ -357,12 +362,12 @@ public:
 
 	void BeginRollback(const FNetworkSimTime& RollbackDeltaTime, const int32 ParentKeyframe) final override
 	{
-		RepProxy_Simulated.template DependentRollbackBegin<TSimulation, TDriver>(Driver, Buffers, TickInfo, RollbackDeltaTime, ParentKeyframe);
+		RepProxy_Simulated.template DependentRollbackBegin<TSimulation, TDriver>(Simulation, Driver, Buffers, TickInfo, RollbackDeltaTime, ParentKeyframe);
 	}
 
 	void StepRollback(const FNetworkSimTime& Step, const int32 ParentKeyframe, const bool bFinalStep) final override
 	{
-		RepProxy_Simulated.template DependentRollbackStep<TSimulation, TDriver>(Driver, Buffers, TickInfo, Step, ParentKeyframe, bFinalStep);
+		RepProxy_Simulated.template DependentRollbackStep<TSimulation, TDriver>(Simulation, Driver, Buffers, TickInfo, Step, ParentKeyframe, bFinalStep);
 	}
 
 	void ClearAllDependentSimulations()
@@ -376,7 +381,9 @@ public:
 
 	// -------------------------------------------------------------------------------------------------------
 
-	TDriver* Driver = nullptr;	
+	TSimulation* Simulation = nullptr;
+	TDriver* Driver = nullptr;
+
 	TSimulationTickState<TTickSettings> TickInfo;	// Manages simulation time and what inputs we are processed
 
 	TNetworkSimBufferContainer<TBufferTypes> Buffers;
@@ -448,4 +455,5 @@ private:
 	TRepProxyDebug RepProxy_Debug;
 	TUniquePtr<TNetworkSimBufferContainer<TBufferTypes>> HistoricBuffers;
 #endif
-};
+
+};
\ No newline at end of file
diff --git a/Engine/Plugins/Runtime/NetworkPrediction/Source/NetworkPrediction/Public/NetworkSimulationModelBuffer.h b/Engine/Plugins/Runtime/NetworkPrediction/Source/NetworkPrediction/Public/NetworkSimulationModelBuffer.h
--- a/Engine/Plugins/Runtime/NetworkPrediction/Source/NetworkPrediction/Public/NetworkSimulationModelBuffer.h
+++ b/Engine/Plugins/Runtime/NetworkPrediction/Source/NetworkPrediction/Public/NetworkSimulationModelBuffer.h
@@ -2,6 +2,7 @@
 
 #pragma once
 
+// Generic iterator for iterating over each keyframe from tail to head
 template<typename TBuffer, typename ElementType>
 struct TNetworkSimBufferIterator
 {
@@ -17,6 +18,42 @@ private:
 	int32 CurrentPos = -1;
 };
 
+// Helper struct to encapsulate optional, delayed writing of new element to buffer
+template<typename ElementType>
+struct TNetSimLazyWriterFunc
+{
+	template<typename TBuffer>
+	TNetSimLazyWriterFunc(TBuffer& Buffer, int32 PendingKeyframe)
+	{
+		GetFunc = [&Buffer, PendingKeyframe]() { return Buffer.WriteKeyframeInitializedFromHead(PendingKeyframe); };
+	}
+
+	ElementType* Get() const
+	{
+		check(GetFunc);
+		return (ElementType*)GetFunc();
+	}
+
+	TFunction<void*()> GetFunc;
+};
+
+// Reference version of LazyWriter. Faster to pass around and do cascading downcasts
+template<typename ElementType>
+struct TNetSimLazyWriter
+{
+	template<typename TLazyWriter>
+	TNetSimLazyWriter(const TLazyWriter& Parent)
+		: GetFunc(Parent.GetFunc) { }
+
+	ElementType* Get() const
+	{
+		return (ElementType*)GetFunc();
+	}
+
+	TFunctionRef<void*()> GetFunc;
+};
+
+// Common implementations shared by TNetworkSimContiguousBuffer/TNetworkSimSparseBuffer
 template<typename T>
 struct TNetworkSimBufferBase
 {
@@ -39,12 +76,14 @@ struct TNetworkSimBufferBase
 			*((T*)this)->WriteKeyframe( It.Keyframe() ) = *It.Element();
 		}
 	}
+
 protected:
 
 	// Creates a new keyframe, but
 	//	-If keyframe already exists, returns existing
 	//	-If keyframe > head, contents of head are copied into new frame
-	//	-If keyframe < head, contents are cleared to default value
+	//	-If keyframe < tail, contents are cleared to default value
+	// (this function is sketchy and poorly named but is proving to be useful in a few places)
 	template<typename ElementType>
 	ElementType* WriteKeyframeInitializedFromHeadImpl(int32 Keyframe)
 	{
@@ -120,9 +159,9 @@ struct TNetworkSimContiguousBuffer : public TNetworkSimBufferBase<TNetworkSimCon
 		return this->template WriteKeyframeInitializedFromHeadImpl<ElementType>(Keyframe);
 	}
 
-	TUniqueFunction<T*()> WriteKeyframeFunc(int32 Keyframe)
+	TNetSimLazyWriterFunc<ElementType> LazyWriter(int32 Keyframe)
 	{
-		return [this, Keyframe]() { return WriteKeyframe(Keyframe); };
+		return TNetSimLazyWriterFunc<ElementType>(*this, Keyframe);
 	}
 
 	using IteratorType = TNetworkSimBufferIterator<TNetworkSimContiguousBuffer<ElementType, NumElements>, ElementType>;
@@ -248,9 +287,9 @@ struct TNetworkSimSparseBuffer : public TNetworkSimBufferBase<TNetworkSimSparseB
 		return this->template WriteKeyframeInitializedFromHeadImpl<ElementType>(Keyframe);
 	}
 
-	TUniqueFunction<T*()> WriteKeyframeFunc(int32 Keyframe)
+	TNetSimLazyWriterFunc<ElementType> LazyWriter(int32 Keyframe)
 	{
-		return [this, Keyframe]() { return WriteKeyframe(Keyframe); };
+		return TNetSimLazyWriterFunc<ElementType>(*this, Keyframe);
 	}
 
 	using IteratorType = TNetworkSimBufferIterator<TNetworkSimSparseBuffer<ElementType, NumElements>, ElementType>;
diff --git a/Engine/Plugins/Runtime/NetworkPrediction/Source/NetworkPrediction/Public/NetworkSimulationModelDebugger.h b/Engine/Plugins/Runtime/NetworkPrediction/Source/NetworkPrediction/Public/NetworkSimulationModelDebugger.h
--- a/Engine/Plugins/Runtime/NetworkPrediction/Source/NetworkPrediction/Public/NetworkSimulationModelDebugger.h
+++ b/Engine/Plugins/Runtime/NetworkPrediction/Source/NetworkPrediction/Public/NetworkSimulationModelDebugger.h
@@ -14,8 +14,8 @@
 #include "CanvasItem.h"
 #include "Containers/Array.h"
 #include "Templates/UniquePtr.h"
-#include "NetworkSimulationModel.h"
 #include "NetworkSimulationModelCVars.h"
+#include "VisualLogger/VisualLogger.h"
 
 DEFINE_LOG_CATEGORY_STATIC(LogNetworkSimDebug, Log, All);
 
@@ -28,7 +28,7 @@ namespace NetworkSimulationModelDebugCVars
 
 struct FNetworkSimulationModelDebuggerManager;
 
-NETWORKPREDICTION_API UObject* FindReplicatedObjectOnPIEServer(UObject* ClientObject);
+NETWORKPREDICTION_API UObject* FindReplicatedObjectOnPIEServer(const UObject* ClientObject);
 
 // ------------------------------------------------------------------------------------------------------------------------
 //	Debugger support classes
@@ -71,7 +71,7 @@ struct NETWORKPREDICTION_API FNetworkSimulationModelDebuggerManager: public FTic
 	// ---------------------------------------------------------------------------------------------------------------------------------------
 
 	template <typename T>
-	void RegisterNetworkSimulationModel(T* NetworkSim, AActor* OwningActor);
+	void RegisterNetworkSimulationModel(T* NetworkSim, const AActor* OwningActor);
 
 	void SetDebuggerActive(AActor* OwningActor, bool InActive)
 	{
@@ -158,7 +158,7 @@ struct NETWORKPREDICTION_API FNetworkSimulationModelDebuggerManager: public FTic
 	{
 		for (auto It = DebuggerMap.CreateIterator(); It; ++It)
 		{
-			AActor* Owner = It.Key().Get();
+			const AActor* Owner = It.Key().Get();
 			if (!Owner)
 			{
 				It.RemoveCurrent();
@@ -258,14 +258,14 @@ struct NETWORKPREDICTION_API FNetworkSimulationModelDebuggerManager: public FTic
 
 private:
 
-	INetworkSimulationModelDebugger* Find(AActor* Actor)
+	INetworkSimulationModelDebugger* Find(const AActor* Actor)
 	{
 		if (!Actor)
 		{
 			return nullptr;
 		}
 
-		INetworkSimulationModelDebugger* Debugger = DebuggerMap.FindRef(TWeakObjectPtr<AActor>(Actor));
+		INetworkSimulationModelDebugger* Debugger = DebuggerMap.FindRef(TWeakObjectPtr<const AActor>(Actor));
 		if (!Debugger)
 		{
 			UE_LOG(LogNetworkSimDebug, Warning, TEXT("Could not find NetworkSimulationModel associated with %s"), *GetPathNameSafe(Actor));
@@ -279,7 +279,7 @@ private:
 
 		for (auto It = DebuggerMap.CreateIterator(); It; ++It)
 		{
-			AActor* Owner = It.Key().Get();
+			const AActor* Owner = It.Key().Get();
 			if (!Owner)
 			{
 				It.RemoveCurrent();
@@ -291,7 +291,7 @@ private:
 				It.Value()->GatherCurrent(*this, C);
 				if (NetworkSimulationModelDebugCVars::GatherServerSidePIE() > 0)
 				{
-					if (AActor* ServerSideActor = Cast<AActor>(FindReplicatedObjectOnPIEServer(Owner)))
+					if (const AActor* ServerSideActor = Cast<AActor>(FindReplicatedObjectOnPIEServer(Owner)))
 					{
 						if (INetworkSimulationModelDebugger* ServerSideSim = Find(ServerSideActor))
 						{
@@ -316,7 +316,7 @@ private:
 		CanvasItems[1].Reset();
 	}
 
-	TMap<TWeakObjectPtr<AActor>, INetworkSimulationModelDebugger*>	DebuggerMap;
+	TMap<TWeakObjectPtr<const AActor>, INetworkSimulationModelDebugger*>	DebuggerMap;
 	bool bContinousGather = true; // Whether you should gather new data every frame
 
 	FDelegateHandle DrawDebugServicesHandle;
@@ -339,7 +339,7 @@ struct TNetworkSimulationModelDebugger : public INetworkSimulationModelDebugger
 	using TSimTime = typename TNetSimModel::TSimTime;
 	using TDebugState = typename TNetSimModel::TDebugState;
 
-	TNetworkSimulationModelDebugger(TNetSimModel* InNetSim, AActor* OwningActor)
+	TNetworkSimulationModelDebugger(TNetSimModel* InNetSim, const AActor* OwningActor)
 	{
 		NetworkSim = InNetSim;
 		WeakOwningActor = OwningActor;
@@ -523,7 +523,7 @@ struct TNetworkSimulationModelDebugger : public INetworkSimulationModelDebugger
 
 	void GatherCurrent(FNetworkSimulationModelDebuggerManager& Out, UCanvas* Canvas) override
 	{
-		AActor* Owner = WeakOwningActor.Get();
+		const AActor* Owner = WeakOwningActor.Get();
 		if (!ensure(Owner))
 		{
 			return;
@@ -653,7 +653,7 @@ struct TNetworkSimulationModelDebugger : public INetworkSimulationModelDebugger
 
 	void Tick( float DeltaTime ) override
 	{
-		AActor* Owner = WeakOwningActor.Get();
+		const AActor* Owner = WeakOwningActor.Get();
 		if (!Owner)
 		{
 			return;
@@ -663,7 +663,9 @@ struct TNetworkSimulationModelDebugger : public INetworkSimulationModelDebugger
 
 		if (auto* LatestSync = NetworkSim->Buffers.Sync.HeadElement())
 		{
-			LatestSync->VisualLog( FVisualLoggingParameters(EVisualLoggingContext::LastPredicted, NetworkSim->Buffers.Sync.HeadKeyframe(), EVisualLoggingLifetime::Transient), NetworkSim->Driver, NetworkSim->Driver);
+			FVisualLoggingParameters VLogParams(EVisualLoggingContext::LastPredicted, NetworkSim->Buffers.Sync.HeadKeyframe(), EVisualLoggingLifetime::Transient);
+			int32 SyncKeyframe = NetworkSim->Buffers.Sync.HeadKeyframe();
+			NetworkSim->Driver->VisualLog(NetworkSim->Buffers.Input[SyncKeyframe], LatestSync, NetworkSim->Buffers.Aux[SyncKeyframe], VLogParams);
 		}
 
 		FStuff ServerPIEStuff = GetServerPIEStuff();
@@ -671,7 +673,9 @@ struct TNetworkSimulationModelDebugger : public INetworkSimulationModelDebugger
 		{
 			if (auto* ServerLatestSync = ServerPIEStuff.NetworkSim->Buffers.Sync.HeadElement())
 			{
-				ServerLatestSync->VisualLog( FVisualLoggingParameters(EVisualLoggingContext::CurrentServerPIE, ServerPIEStuff.NetworkSim->Buffers.Sync.HeadKeyframe(), EVisualLoggingLifetime::Transient), ServerPIEStuff.NetworkSim->Driver, NetworkSim->Driver);
+				FVisualLoggingParameters VLogParams(EVisualLoggingContext::CurrentServerPIE, ServerPIEStuff.NetworkSim->Buffers.Sync.HeadKeyframe(), EVisualLoggingLifetime::Transient);
+				int32 SyncKeyframe = ServerPIEStuff.NetworkSim->Buffers.Sync.HeadKeyframe();
+				NetworkSim->Driver->VisualLog(ServerPIEStuff.NetworkSim->Buffers.Input[SyncKeyframe], ServerLatestSync, ServerPIEStuff.NetworkSim->Buffers.Aux[SyncKeyframe], VLogParams);
 			}
 		}
 
@@ -682,13 +686,17 @@ struct TNetworkSimulationModelDebugger : public INetworkSimulationModelDebugger
 				if (auto* SyncState = NetworkSim->Buffers.Sync[Keyframe])
 				{
 					const EVisualLoggingContext Context = (Keyframe == NetworkSim->RepProxy_Autonomous.GetLastSerializedKeyframe()) ? EVisualLoggingContext::LastConfirmed : EVisualLoggingContext::OtherPredicted;
-					SyncState->VisualLog( FVisualLoggingParameters(Context, NetworkSim->Buffers.Sync.HeadKeyframe(), EVisualLoggingLifetime::Transient), NetworkSim->Driver, NetworkSim->Driver);
+					FVisualLoggingParameters VLogParams(Context, NetworkSim->Buffers.Sync.HeadKeyframe(), EVisualLoggingLifetime::Transient);
+					NetworkSim->Driver->VisualLog(NetworkSim->Buffers.Input[Keyframe], SyncState, NetworkSim->Buffers.Aux[Keyframe], VLogParams);
 				}
 			}
 		}
 		else if (Owner->GetLocalRole() == ROLE_SimulatedProxy)
 		{
-			NetworkSim->RepProxy_Simulated.GetLastSerializedSyncState().VisualLog( FVisualLoggingParameters(EVisualLoggingContext::LastConfirmed, NetworkSim->Buffers.Sync.HeadKeyframe(), EVisualLoggingLifetime::Transient), NetworkSim->Driver, NetworkSim->Driver); 
+			FVisualLoggingParameters VLogParams(EVisualLoggingContext::LastConfirmed, NetworkSim->Buffers.Sync.HeadKeyframe(), EVisualLoggingLifetime::Transient);
+			int32 HeadKeyframe = NetworkSim->Buffers.Sync.HeadKeyframe();
+			NetworkSim->Driver->VisualLog(NetworkSim->Buffers.Input[HeadKeyframe], NetworkSim->Buffers.Sync[HeadKeyframe], NetworkSim->Buffers.Aux[HeadKeyframe], VLogParams);
+			
 			if (NetworkSim->GetSimulatedUpdateMode() != ESimulatedUpdateMode::Interpolate)
 			{
 				FVector2D ServerSimulationTimeData(Owner->GetWorld()->GetTimeSeconds(), NetworkSim->RepProxy_Simulated.GetLastSerializedSimulationTime().ToRealTimeMS());
@@ -716,23 +724,23 @@ struct TNetworkSimulationModelDebugger : public INetworkSimulationModelDebugger
 
 private:
 	
-	TWeakObjectPtr<AActor>	WeakOwningActor;
+	TWeakObjectPtr<const AActor> WeakOwningActor;
 	TNetSimModel* NetworkSim = nullptr;
 };
 
 template <typename T>
-void FNetworkSimulationModelDebuggerManager::RegisterNetworkSimulationModel(T* NetworkSim, AActor* OwningActor)
+void FNetworkSimulationModelDebuggerManager::RegisterNetworkSimulationModel(T* NetworkSim, const AActor* OwningActor)
 {
 	TNetworkSimulationModelDebugger<T>* Debugger = new TNetworkSimulationModelDebugger<T>(NetworkSim, OwningActor);
-	DebuggerMap.Add( TWeakObjectPtr<AActor>(OwningActor), Debugger );
+	DebuggerMap.Add( TWeakObjectPtr<const AActor>(OwningActor), Debugger );
 
 	// Gross stuff so that the debugger can find the ServerPIE equiv
-	TWeakObjectPtr<AActor> WeakOwner(OwningActor);
+	TWeakObjectPtr<const AActor> WeakOwner(OwningActor);
 	Debugger->GetServerPIEStuff = [WeakOwner, this]()
 	{
 		if (AActor* ServerOwner = Cast<AActor>(FindReplicatedObjectOnPIEServer(WeakOwner.Get())))
 		{
-			return ((TNetworkSimulationModelDebugger<T>*)DebuggerMap.FindRef(TWeakObjectPtr<AActor>(ServerOwner)))->GetStuff();
+			return ((TNetworkSimulationModelDebugger<T>*)DebuggerMap.FindRef(TWeakObjectPtr<const AActor>(ServerOwner)))->GetStuff();
 		}
 		return typename TNetworkSimulationModelDebugger<T>::FStuff();
 	};
diff --git a/Engine/Plugins/Runtime/NetworkPrediction/Source/NetworkPrediction/Public/NetworkSimulationModelInterpolator.h b/Engine/Plugins/Runtime/NetworkPrediction/Source/NetworkPrediction/Public/NetworkSimulationModelInterpolator.h
--- a/Engine/Plugins/Runtime/NetworkPrediction/Source/NetworkPrediction/Public/NetworkSimulationModelInterpolator.h
+++ b/Engine/Plugins/Runtime/NetworkPrediction/Source/NetworkPrediction/Public/NetworkSimulationModelInterpolator.h
@@ -7,6 +7,7 @@
 #include "GameFramework/Actor.h"
 #include "NetworkPredictionTypes.h"
 #include "Engine/World.h"
+#include "NetworkSimulationModelTypes.h"
 
 DEFINE_LOG_CATEGORY_STATIC(LogNetInterpolation, Log, All);
 
@@ -32,13 +33,13 @@ struct TInterpolator
 
 	bool bEnableVisualLog = true;
 
-	template<typename TDriver>
-	void PostSimTick(TDriver* Driver, const TNetworkSimBufferContainer<TBufferTypes>& Buffers, const TSimulationTickState<TTickSettings>& TickInfo, const FNetSimTickParameters& TickParameters)
+	template<typename TSystemDriver>
+	void PostSimTick(TSystemDriver* Driver, const TNetworkSimBufferContainer<TBufferTypes>& Buffers, const TSimulationTickState<TTickSettings>& TickInfo, const FNetSimTickParameters& TickParameters)
 	{
 		const bool bDoVLog = NetworkInterpolationDebugCVars::VLog() && bEnableVisualLog;
 		const float DeltaSeconds = TickParameters.LocalDeltaTimeSeconds;
 
-		const UObject* LogOwner = Driver->GetVLogOwner();
+		const AActor* LogOwner = Driver->GetVLogOwner();
 
 		// Interpolation disabled
 		if (NetworkInterpolationDebugCVars::Disable() > 0)
@@ -196,16 +197,31 @@ struct TInterpolator
 					const TSyncState* DebugTail = Buffers.Sync.TailElement();
 					const TSyncState* DebugHead = Buffers.Sync.HeadElement();
 
-					DebugTail->VisualLog( FVisualLoggingParameters(EVisualLoggingContext::InterpolationBufferTail, Buffers.Sync.TailKeyframe(), EVisualLoggingLifetime::Transient), Driver, Driver );
-					DebugHead->VisualLog( FVisualLoggingParameters(EVisualLoggingContext::InterpolationBufferHead, Buffers.Sync.HeadKeyframe(), EVisualLoggingLifetime::Transient), Driver, Driver );
-
-					InterpolationState.VisualLog( FVisualLoggingParameters(EVisualLoggingContext::InterpolationFrom, InterpolationKeyframe, EVisualLoggingLifetime::Transient), Driver, Driver );
-					ToState->VisualLog( FVisualLoggingParameters(EVisualLoggingContext::InterpolationTo, InterpolationKeyframe, EVisualLoggingLifetime::Transient), Driver, Driver );
-
-					NewInterpolatedState.VisualLog( FVisualLoggingParameters(LoggingContext, InterpolationKeyframe, EVisualLoggingLifetime::Transient), Driver, Driver );
+					auto VLogHelper = [&](int32 Keyframe, EVisualLoggingContext Context)
+					{
+						FVisualLoggingParameters VLogParams(Context, Keyframe, EVisualLoggingLifetime::Transient);
+						Driver->VisualLog(Buffers.Input[Keyframe], Buffers.Sync[Keyframe], Buffers.Aux[Keyframe], VLogParams);
+					};
+
+					VLogHelper(Buffers.Sync.TailKeyframe(), EVisualLoggingContext::InterpolationBufferTail);
+					VLogHelper(Buffers.Sync.HeadKeyframe(), EVisualLoggingContext::InterpolationBufferHead);
+
+					{
+						FVisualLoggingParameters VLogParams(EVisualLoggingContext::InterpolationFrom, InterpolationKeyframe-1, EVisualLoggingLifetime::Transient);
+						Driver->VisualLog(Buffers.Input[InterpolationKeyframe-1], &InterpolationState, Buffers.Aux[InterpolationKeyframe-1], VLogParams);
+					}
+
+					{
+						FVisualLoggingParameters VLogParams(EVisualLoggingContext::InterpolationTo, InterpolationKeyframe, EVisualLoggingLifetime::Transient);
+						Driver->VisualLog(Buffers.Input[InterpolationKeyframe], &InterpolationState, Buffers.Aux[InterpolationKeyframe], VLogParams);
+					}
+
+					{
+						FVisualLoggingParameters VLogParams(LoggingContext, InterpolationKeyframe, EVisualLoggingLifetime::Transient);
+						Driver->VisualLog(Buffers.Input[InterpolationKeyframe], &NewInterpolatedState, Buffers.Aux[InterpolationKeyframe], VLogParams);
+					}
 				}
 
-
 				InterpolationState = NewInterpolatedState;
 				InterpolationTime = NewInterpolationTime;
 			}
diff --git a/Engine/Plugins/Runtime/NetworkPrediction/Source/NetworkPrediction/Public/NetworkSimulationModelReplicators.h b/Engine/Plugins/Runtime/NetworkPrediction/Source/NetworkPrediction/Public/NetworkSimulationModelReplicators.h
--- a/Engine/Plugins/Runtime/NetworkPrediction/Source/NetworkPrediction/Public/NetworkSimulationModelReplicators.h
+++ b/Engine/Plugins/Runtime/NetworkPrediction/Source/NetworkPrediction/Public/NetworkSimulationModelReplicators.h
@@ -52,7 +52,12 @@ struct FNetworkSimulationSerialization
 };
 
 // ----------------------------------------------------------------------------------------------------------------------------------------------------------------------
-//	Templated Replicators
+//	Replicators
+//	The Replicators are the pieces of the TNetworkedSimulationModel that make up the role-specific functionality (Server, Autonomous Client, and Simulated Client).
+//	Mainly they NetSerialize, Reconcile, and have Pre/Post sim tick functions. The TNetworkedSimulationModel is still the core piece that ticks the sim, but the Replicators
+//	basically do everything else in a role-specific way.
+//
+//
 //
 // ----------------------------------------------------------------------------------------------------------------------------------------------------------------------
 
@@ -66,11 +71,11 @@ struct TReplicatorEmpty
 	void NetSerialize(const FNetSerializeParams& P, TNetworkSimBufferContainer<TBufferTypes>& Buffers, const TSimulationTickState<TTickSettings>& TickInfo) { }
 
 	// Reconcile: called after everyone has NetSerialized. "Get right with the server": this function is about reconciling what the server told you vs what you have extrapolated or forward predicted locally
-	template<typename T, typename TDriver>
-	void Reconcile(TDriver* Driver, TNetworkSimBufferContainer<TBufferTypes>& Buffers, TSimulationTickState<TTickSettings>& TickInfo) { }
+	template<typename TSimulation, typename TDriver>
+	void Reconcile(TSimulation* Simulation, TDriver* Driver, TNetworkSimBufferContainer<TBufferTypes>& Buffers, TSimulationTickState<TTickSettings>& TickInfo) { }
 
 	// Called prior to input processing. This function must updated TickInfo to allow simulation time (from TickParameters) and to possibly get new input.
-	template<typename T, typename TDriver>
+	template<typename TDriver>
 	void PreSimTick(TDriver* Driver, TNetworkSimBufferContainer<TBufferTypes>& Buffers, TSimulationTickState<TTickSettings>& TickInfo, const FNetSimTickParameters& TickParameters) { }
 
 	// Called after input processing. Should finalize the frame and do any smoothing/interpolation. This function is not allowed to modify the buffers or tick state, or even call the simulation/Update function.
@@ -91,11 +96,11 @@ struct TReplicatorBase
 	void NetSerialize(const FNetSerializeParams& P, TNetworkSimBufferContainer<TBufferTypes>& Buffers, const TSimulationTickState<TTickSettings>& TickInfo) { }
 
 	// Reconcile: called after everyone has NetSerialized. "Get right with the server": this function is about reconciling what the server told you vs what you have extrapolated or forward predicted locally
-	template<typename T, typename TDriver>
-	void Reconcile(TDriver* Driver, TNetworkSimBufferContainer<TBufferTypes>& Buffers, TSimulationTickState<TTickSettings>& TickInfo) { }
+	template<typename TSimulation, typename TDriver>
+	void Reconcile(TSimulation* Simulation, TDriver* Driver, TNetworkSimBufferContainer<TBufferTypes>& Buffers, TSimulationTickState<TTickSettings>& TickInfo) { }
 
 	// Called prior to input processing. This function must updated TickInfo to allow simulation time (from TickParameters) and to possibly get new input.
-	template<typename T, typename TDriver>
+	template<typename TDriver>
 	void PreSimTick(TDriver* Driver, TNetworkSimBufferContainer<TBufferTypes>& Buffers, TSimulationTickState<TTickSettings>& TickInfo, const FNetSimTickParameters& TickParameters)
 	{
 		// Accumulate local delta time into TickInfo
@@ -269,8 +274,8 @@ struct TReplicator_Server : public TBase
 {
 	using TInputCmd = typename TBufferTypes::TInputCmd;
 
-	template<typename T, typename TDriver>
-	void Reconcile(TDriver* Driver, TNetworkSimBufferContainer<TBufferTypes>& Buffers, TSimulationTickState<TTickSettings>& TickInfo)
+	template<typename TSimulation, typename TDriver>
+	void Reconcile(TSimulation* Simulation, TDriver* Driver, TNetworkSimBufferContainer<TBufferTypes>& Buffers, TSimulationTickState<TTickSettings>& TickInfo)
 	{
 		// After receiving input, server may process up to the latest received frames.
 		// (If we needed to buffer input server side for whatever reason, we would do it here)
@@ -429,8 +434,8 @@ struct TReplicator_Simulated : public TBase
 		}
 	}
 
-	template<typename T, typename TDriver>
-	void Reconcile(TDriver* Driver, TNetworkSimBufferContainer<TBufferTypes>& Buffers, TSimulationTickState<TTickSettings>& TickInfo)
+	template<typename TSimulation, typename TDriver>
+	void Reconcile(TSimulation* Simulation, TDriver* Driver, TNetworkSimBufferContainer<TBufferTypes>& Buffers, TSimulationTickState<TTickSettings>& TickInfo)
 	{
 		if (ReconcileSimulationTime.IsPositive() == false)
 		{
@@ -452,7 +457,7 @@ struct TReplicator_Simulated : public TBase
 			FNetworkSimTime DeltaSimTime = ReconcileSimulationTime - TickInfo.GetTotalProcessedSimulationTime();
 			if (DeltaSimTime.IsPositive() && NetworkSimulationModelCVars::EnableSimulatedReconcile())
 			{
-				SimulationExtrapolation<T, TDriver>(Driver, Buffers, TickInfo, DeltaSimTime);
+				SimulationExtrapolation<TSimulation, TDriver>(Simulation, Driver, Buffers, TickInfo, DeltaSimTime);
 			}
 		}
 		
@@ -460,7 +465,7 @@ struct TReplicator_Simulated : public TBase
 		ReconcileSimulationTime.Reset();
 	}
 
-	template<typename T, typename TDriver>
+	template<typename TDriver>
 	void PreSimTick(TDriver* Driver, TNetworkSimBufferContainer<TBufferTypes>& Buffers, TSimulationTickState<TTickSettings>& TickInfo, const FNetSimTickParameters& TickParameters)
 	{
 		// Tick if we are dependent simulation or extrapolation is enabled
@@ -497,8 +502,8 @@ struct TReplicator_Simulated : public TBase
 		}
 	}
 
-	template<typename T, typename TDriver>
-	void DependentRollbackBegin(TDriver* Driver, TNetworkSimBufferContainer<TBufferTypes>& Buffers, TSimulationTickState<TTickSettings>& TickInfo, const FNetworkSimTime& RollbackDeltaTime, const int32 ParentKeyframe)
+	template<typename TSimulation, typename TDriver>
+	void DependentRollbackBegin(TSimulation* Simulation, TDriver* Driver, TNetworkSimBufferContainer<TBufferTypes>& Buffers, TSimulationTickState<TTickSettings>& TickInfo, const FNetworkSimTime& RollbackDeltaTime, const int32 ParentKeyframe)
 	{
 		// For now, we make the assumption that our last serialized state and time match the parent simulation.
 		// This would not always be the case with low frequency simulated proxies. But could be handled by replicating the simulations together (at the replication level)
@@ -514,25 +519,34 @@ struct TReplicator_Simulated : public TBase
 
 		Driver->FinalizeFrame(LastSerializedSyncState, LastSerializedAuxState);
 
-		LastSerializedSyncState.VisualLog( FVisualLoggingParameters(EVisualLoggingContext::FirstMispredicted, ParentKeyframe, EVisualLoggingLifetime::Persistent), Driver, Driver );
+		if (NETSIM_MODEL_DEBUG)
+		{
+			FVisualLoggingParameters VLogParams(EVisualLoggingContext::FirstMispredicted, ParentKeyframe, EVisualLoggingLifetime::Persistent, TEXT("DependentRollbackBegin"));
+			Driver->VisualLog(nullptr, &LastSerializedSyncState, &LastSerializedAuxState, VLogParams);
+		}
 	}
 
-	template<typename T, typename TDriver>
-	void DependentRollbackStep(TDriver* Driver, TNetworkSimBufferContainer<TBufferTypes>& Buffers, TSimulationTickState<TTickSettings>& TickInfo, const FNetworkSimTime& StepTime, const int32 ParentKeyframe, const bool bFinalStep)
+	template<typename TSimulation, typename TDriver>
+	void DependentRollbackStep(TSimulation* Simulation, TDriver* Driver, TNetworkSimBufferContainer<TBufferTypes>& Buffers, TSimulationTickState<TTickSettings>& TickInfo, const FNetworkSimTime& StepTime, const int32 ParentKeyframe, const bool bFinalStep)
 	{
 		TickInfo.SetTotalAllowedSimulationTime( TickInfo.GetTotalAllowedSimulationTime() + StepTime );
 
-		SimulationExtrapolation<T, TDriver>(Driver, Buffers, TickInfo, StepTime);
+		SimulationExtrapolation<TSimulation, TDriver>(Simulation, Driver, Buffers, TickInfo, StepTime);
 
 		TSyncState* SyncState = Buffers.Sync.HeadElement();
 		check(SyncState);
-		SyncState->VisualLog( FVisualLoggingParameters(bFinalStep ? EVisualLoggingContext::LastMispredicted : EVisualLoggingContext::OtherMispredicted, ParentKeyframe, EVisualLoggingLifetime::Persistent), Driver, Driver );
+		
+		if (NETSIM_MODEL_DEBUG)
+		{
+			FVisualLoggingParameters VLogParams(bFinalStep ? EVisualLoggingContext::LastMispredicted : EVisualLoggingContext::OtherMispredicted, ParentKeyframe, EVisualLoggingLifetime::Persistent, TEXT("DependentRollbackStep"));
+			Driver->VisualLog(nullptr, SyncState, Buffers.Aux.HeadElement(), VLogParams);
+		}
 	}
 
 private:
 
-	template<typename T, typename TDriver>
-	void SimulationExtrapolation(TDriver* Driver, TNetworkSimBufferContainer<TBufferTypes>& Buffers, TSimulationTickState<TTickSettings>& TickInfo, const FNetworkSimTime DeltaSimTime)
+	template<typename TSimulation, typename TDriver>
+	void SimulationExtrapolation(TSimulation* Simulation, TDriver* Driver, TNetworkSimBufferContainer<TBufferTypes>& Buffers, TSimulationTickState<TTickSettings>& TickInfo, const FNetworkSimTime DeltaSimTime)
 	{
 		// We have extrapolated ahead of the server. The latest network update is now "in the past" from what we rendered last frame.
 		// We will insert a new keyframe to make up the difference from the last known state to where we want to be in the now.
@@ -554,7 +568,7 @@ private:
 		// Do the actual update
 		{
 			TScopedSimulationTick UpdateScope(TickInfo, OutputKeyframe, NewCmd->GetFrameDeltaTime());
-			T::Update(Driver, NewCmd->GetFrameDeltaTime().ToRealTimeSeconds(), *NewCmd, *PrevSyncState, *NextSyncState, *AuxState, Buffers.Aux.WriteKeyframeFunc(OutputKeyframe));
+			Simulation->Update(NewCmd->GetFrameDeltaTime().ToRealTimeSeconds(), *NewCmd, *PrevSyncState, *NextSyncState, *AuxState, Buffers.Aux.LazyWriter(OutputKeyframe));
 		}
 
 		TickInfo.MaxAllowedKeyframe = OutputKeyframe;
@@ -647,8 +661,8 @@ struct TReplicator_Autonomous : public TBase
 	// --------------------------------------------------------------------
 	//	Reconcile
 	// --------------------------------------------------------------------
-	template<typename T, typename TDriver>
-	void Reconcile(TDriver* Driver, TNetworkSimBufferContainer<TBufferTypes>& Buffers, TSimulationTickState<TTickSettings>& TickInfo)
+	template<typename TSimulation, typename TDriver>
+	void Reconcile(TSimulation* Simulation, TDriver* Driver, TNetworkSimBufferContainer<TBufferTypes>& Buffers, TSimulationTickState<TTickSettings>& TickInfo)
 	{
 		if (bPendingReconciliation == false && (bDependentSimulationNeedsReconcile == false || SerializedKeyframe == INDEX_NONE))
 		{
@@ -658,11 +672,12 @@ struct TReplicator_Autonomous : public TBase
 		bDependentSimulationNeedsReconcile = false;
 
 		TSyncState* ClientSyncState = Buffers.Sync[SerializedKeyframe];
-		const bool bDoVisualLog = NetworkSimulationModelCVars::EnableLocalPrediction() > 0; // don't visual log if we have prediction disabled
+		const bool bDoVisualLog = NetworkSimulationModelCVars::EnableLocalPrediction() > 0 && NETSIM_MODEL_DEBUG; // don't visual log if we have prediction disabled
 		
 		if (bDoVisualLog)
 		{
-			SerializedSyncState.VisualLog( FVisualLoggingParameters(EVisualLoggingContext::LastConfirmed, SerializedKeyframe, EVisualLoggingLifetime::Persistent), Driver, Driver);
+			FVisualLoggingParameters VLogParameters(EVisualLoggingContext::LastConfirmed, SerializedKeyframe, EVisualLoggingLifetime::Persistent, TEXT("Serialized State"));
+			Driver->VisualLog(Buffers.Input[SerializedKeyframe], &SerializedSyncState, &SerializedAuxState, VLogParameters);
 		}
 
 		if (ClientSyncState)
@@ -670,7 +685,8 @@ struct TReplicator_Autonomous : public TBase
 			// Existing ClientSyncState, log it before overwriting it
 			if (bDoVisualLog)
 			{
-				ClientSyncState->VisualLog( FVisualLoggingParameters(EVisualLoggingContext::FirstMispredicted, SerializedKeyframe, EVisualLoggingLifetime::Persistent), Driver, Driver);
+				FVisualLoggingParameters VLogParameters(EVisualLoggingContext::FirstMispredicted, SerializedKeyframe, EVisualLoggingLifetime::Persistent, TEXT("Mispredicted State"));
+				Driver->VisualLog(Buffers.Input[SerializedKeyframe], ClientSyncState, Buffers.Aux[SerializedKeyframe], VLogParameters);
 			}
 		}
 		else
@@ -723,16 +739,24 @@ struct TReplicator_Autonomous : public TBase
 			check(AuxState);
 
 			// Log out the Mispredicted state that we are about to overwrite.
-			NextSyncState->VisualLog( FVisualLoggingParameters(Keyframe == LastKeyframeToProcess ? EVisualLoggingContext::LastMispredicted : EVisualLoggingContext::OtherMispredicted, Keyframe, EVisualLoggingLifetime::Persistent), Driver, Driver);
+			if (NETSIM_MODEL_DEBUG)
+			{
+				FVisualLoggingParameters VLogParameters(Keyframe == LastKeyframeToProcess ? EVisualLoggingContext::LastMispredicted : EVisualLoggingContext::OtherMispredicted, Keyframe, EVisualLoggingLifetime::Persistent, TEXT("Resimulate Step: mispredicted"));
+				Driver->VisualLog(Buffers.Input[OutputKeyframe], Buffers.Sync[OutputKeyframe], Buffers.Aux[OutputKeyframe], VLogParameters);
+			}
 
 			// Do the actual update
 			{
 				TScopedSimulationTick UpdateScope(TickInfo, OutputKeyframe, ResimulateCmd->GetFrameDeltaTime());
-				T::Update(Driver, ResimulateCmd->GetFrameDeltaTime().ToRealTimeSeconds(), *ResimulateCmd, *PrevSyncState, *NextSyncState, *AuxState, Buffers.Aux.WriteKeyframeFunc(OutputKeyframe));
+				Simulation->Update(ResimulateCmd->GetFrameDeltaTime().ToRealTimeSeconds(), *ResimulateCmd, *PrevSyncState, *NextSyncState, *AuxState, Buffers.Aux.LazyWriter(OutputKeyframe));
 			}
 
 			// Log out the newly predicted state that we got.
-			NextSyncState->VisualLog( FVisualLoggingParameters(Keyframe == LastKeyframeToProcess ? EVisualLoggingContext::LastPredicted : EVisualLoggingContext::OtherPredicted, Keyframe, EVisualLoggingLifetime::Persistent), Driver, Driver);
+			if (NETSIM_MODEL_DEBUG)
+			{			
+				FVisualLoggingParameters VLogParameters(Keyframe == LastKeyframeToProcess ? EVisualLoggingContext::LastPredicted : EVisualLoggingContext::OtherPredicted, Keyframe, EVisualLoggingLifetime::Persistent, TEXT("Resimulate Step: repredicted"));
+				Driver->VisualLog(Buffers.Input[OutputKeyframe], Buffers.Sync[OutputKeyframe], Buffers.Aux[OutputKeyframe], VLogParameters);
+			}
 
 			// Tell dependent simulations to advance
 			for (INetworkSimulationModel* DependentSim : DependentSimulations)
@@ -745,7 +769,7 @@ struct TReplicator_Autonomous : public TBase
 	// --------------------------------------------------------------------
 	//	PreSimTick
 	// --------------------------------------------------------------------
-	template<typename T, typename TDriver>
+	template<typename TDriver>
 	void PreSimTick(TDriver* Driver, TNetworkSimBufferContainer<TBufferTypes>& Buffers, TSimulationTickState<TTickSettings>& TickInfo, const FNetSimTickParameters& TickParameters)
 	{
 		// If we have a reconcile fault, we cannot continue on with the simulation until it clears itself out. This effectively drops the input time and does not sample new inputs
diff --git a/Engine/Plugins/Runtime/NetworkPrediction/Source/NetworkPrediction/Public/NetworkSimulationModelTypes.h b/Engine/Plugins/Runtime/NetworkPrediction/Source/NetworkPrediction/Public/NetworkSimulationModelTypes.h
--- a/Engine/Plugins/Runtime/NetworkPrediction/Source/NetworkPrediction/Public/NetworkSimulationModelTypes.h
+++ b/Engine/Plugins/Runtime/NetworkPrediction/Source/NetworkPrediction/Public/NetworkSimulationModelTypes.h
@@ -73,6 +73,24 @@ struct TNetworkSimBufferAllocation
 	static constexpr ENetworkSimBufferAllocationType AllocationType() { return InAllocationType; }
 };
 
+template<typename T>
+struct HasNetSerialize
+{
+	template<typename U, void (U::*)(const FNetSerializeParams& P)> struct SFINAE {};
+	template<typename U> static char Test(SFINAE<U, &U::NetSerialize>*);
+	template<typename U> static int Test(...);
+	static const bool Has = sizeof(Test<T>(0)) == sizeof(char);
+};
+
+template<typename T>
+struct HasLog
+{
+	template<typename U, void (U::*)(FStandardLoggingParameters& P) const> struct SFINAE {};
+	template<typename U> static char Test(SFINAE<U, &U::Log>*);
+	template<typename U> static int Test(...);
+	static const bool Has = sizeof(Test<T>(0)) == sizeof(char);
+};
+
 // A collection of the system's buffer types. This allows us to collapse the 4 types into a single type to use a template argument elsewhere.
 template<typename InInputCmd, typename InSyncState, typename InAuxState, typename InDebugState = FNetSimProcessedFrameDebugInfo>
 struct TNetworkSimBufferTypes
@@ -89,6 +107,17 @@ struct TNetworkSimBufferTypes
 	{
 		using type = typename TSelectTypeHelper< TNetworkSimBufferTypes<TInputCmd, TSyncState, TAuxState, TDebugState>, Id >::type;
 	};
+
+	// Must implement NetSerialize and Log functions. This purposefully does not pass on inherited methods
+	static_assert(HasNetSerialize<InInputCmd>::Has == true, "InputCmd Must implement NetSerialize");
+	static_assert(HasNetSerialize<InSyncState>::Has == true, "SyncState Must implement NetSerialize");
+	static_assert(HasNetSerialize<InAuxState>::Has == true, "AuxState Must implement NetSerialize");
+	static_assert(HasNetSerialize<InDebugState>::Has == true, "DebugState Must implement NetSerialize");
+
+	static_assert(HasLog<InInputCmd>::Has == true, "InputCmd Must implement Log");
+	static_assert(HasLog<InSyncState>::Has == true, "SyncState Must implement Log");
+	static_assert(HasLog<InAuxState>::Has == true, "AuxState Must implement Log");
+	static_assert(HasLog<InDebugState>::Has == true, "DebugState Must implement Log");
 };
 
 // ---------------------------------------------------------------------------------------------------------------------
@@ -208,7 +237,7 @@ struct TLazyStateAccessor
 
 private:
 
-	T* CachedWriteNext = nullptr;
+	mutable T* CachedWriteNext = nullptr;
 	TUniqueFunction<T*()> GetWriteNextFunc;
 };
 
@@ -480,7 +509,8 @@ private:
 	TRealTimeAccumulator<TSettings>	RealtimeAccumulator;
 };
 
-// Scoped helper to be used right before entering a call to the sim's ::Update function
+// Scoped helper to be used right before entering a call to the sim's ::Update function.
+// Important to note that this advanced the PendingKeyFrame to the output Keyframe. So that any writes that occur to the buffers will go to the output frame.
 struct TScopedSimulationTick
 {
 	template<typename TSimulationTickState>
@@ -603,6 +633,8 @@ struct TFrameCmd : public BaseCmdType
 		BaseCmdType::NetSerialize(P); 
 	}
 
+	void Log(FStandardLoggingParameters& P) const { BaseCmdType::Log(P); }
+
 private:
 	FNetworkSimTime FrameDeltaTime;
 };
@@ -614,6 +646,7 @@ struct TFrameCmd<BaseCmdType, TickSettings, true> : public BaseCmdType
 	FNetworkSimTime GetFrameDeltaTime() const { return FNetworkSimTime::FromMSec(TickSettings::GetFixedStepMS()); }
 	void SetFrameDeltaTime(const FNetworkSimTime& InTime) { }
 	void NetSerialize(const FNetSerializeParams& P) { BaseCmdType::NetSerialize(P); }
+	void Log(FStandardLoggingParameters& P) const { BaseCmdType::Log(P); }
 };
 
 // Helper to turn user supplied buffer types into the "real" buffer types: the InputCmd struct is wrapped in TFrameCmd
@@ -630,16 +663,26 @@ struct TInternalBufferTypes : TNetworkSimBufferTypes<
 {
 };
 
-/** Interface base for the simulation driver. This is what is used by the Network Sim Model internally. Basically, the non simulation specific functions that need to be defined. */
+/** This is the "system driver", it has functions that the TNetworkedSimulationModel needs to call internally, that are specific to the types but not specific to the simulation itself. */
 template<typename TBufferTypes>
-class TNetworkSimDriverInterfaceBase
+class TNetworkedSimulationModelDriver
 {
 public:
-	virtual FString GetDebugName() const = 0; // Used for debugging. Recommended to emit the simulation name and the actor name/role.
-	virtual const UObject* GetVLogOwner() const = 0; // Owning object for Visual Logs
+	using TInputCmd = typename TBufferTypes::TInputCmd;
+	using TSyncState = typename TBufferTypes::TSyncState;
+	using TAuxState= typename TBufferTypes::TAuxState;
+
+	// Debug string that can be used in internal warning/error logs
+	virtual FString GetDebugName() const = 0;
+
+	// Owning object for Visual Logs so that the system can emit them internally
+	virtual const AActor* GetVLogOwner() const = 0;
+
+	// Call to visual log the given states. Note that not all 3 will always be present and you should check for nullptrs.
+	virtual void VisualLog(const TInputCmd* Input, const TSyncState* Sync, const TAuxState* Aux, const FVisualLoggingParameters& SystemParameters) const = 0;
 	
 	// Called whenever the sim is ready to process new local input.
-	virtual void ProduceInput(const FNetworkSimTime SimTime, typename TBufferTypes::TInputCmd&) = 0;
+	virtual void ProduceInput(const FNetworkSimTime SimTime, TInputCmd&) = 0;
 	
 	// Called from the Network Sim at the end of the sim frame when there is new sync data.
 	virtual void FinalizeFrame(const typename TBufferTypes::TSyncState& SyncState, const typename TBufferTypes::TAuxState& AuxState) = 0;
diff --git a/Engine/Plugins/Runtime/NetworkPrediction/readme.txt b/Engine/Plugins/Runtime/NetworkPrediction/readme.txt
--- a/Engine/Plugins/Runtime/NetworkPrediction/readme.txt
+++ b/Engine/Plugins/Runtime/NetworkPrediction/readme.txt
@@ -35,6 +35,7 @@ Code to look at:
 -NetworkPredictionExtrasFlyingPawn.h: Starting point to see how a pawn using flying movement is setup. In NetworkPredictionExtras.
 -FlyingMovement.h: The actual implementation of the flying movement system. In NetworkPrediction.
 -NetworkSimulationModel.h: this is the lowest level guts of the system. This is what every simulation ultimately uses to stay in sync.
+-FMockAbilitySimulation::Update an example "ability system" working on top of a movement simulation
  
 Extremely high level technical overview:
 -Concept of "NetworkedSimulationModel": defining a tightly constrained (gameplay) simulation that can be run predictively on clients and authoritatively on server. E.g., a "Movement System" is a NetworkedSimulationModel.
@@ -59,9 +60,7 @@ Simulated Proxy behavior / "modes"
  
  
 TODO: Major missing elements
--Aux Buffer is not implemented, but referenced in a few places.
 -"Events" into and out of the system are not in place. That is, something happens in your movement sim that needs to call outside code/generate an event. A system is needed to track this stuff rather than just calling directly into the handler.
--Interpolation/smoothing layer (visual, non simulation affecting smoothing to help with corrections and simulated proxies)
 -No optimizations (bandwidth, cpu) have been done.
 -[Movement specific]: must lock down SetActorLocation/Rotation API so that movement state can be changed out from underneath the system.
  
@@ -85,7 +84,11 @@ Road Map:
 // ----------------------------------------------------------------------------------------------------------
 // Glossary
 // ----------------------------------------------------------------------------------------------------------
- 
+
+"Simulation": should always refer to use code. The simulation is what we are networking. It doesn't matter what the simulation is actually doing, we are building a system to network it like a black box.
+"System": should always refer to the TNetworkedSimulationModel code: the generic code that is doing the networking. It is distinct from the simulation.
+
+
 Reconcile: When a predicting client receives authoritative data from the server that conflicts with what they previously predicted: sorting this out is called "reconciling". Usually means you will resimulate.
 Synthesize: In general means "making shit up". For example, "synthesizing user commands" would mean "creating user commands that were not generated by a player, but instead guessing what it might be".
  
@@ -105,6 +108,18 @@ Smoothing: Taking the output of the simulation and applying an additional layer
 // Release notes
 // ----------------------------------------------------------------------------------------------------------
 
+Update (10-25-19)
+-Another large refactor on the core classes. Things are feeling good!
+-"Simulation Driver" and a few other concepts have changed:
+	-"Simulation" is now implemented by a simulation class. What was previously a static method + a "driver" pointer, is now an instantiated C++ class that user code creates and passes into TNetworkedSimulationModel.
+	-The Simulation (previously driver) is now longer implemented by the owning actor/component. This simplifies a lot of things and make extending simulation much less awkward.
+	-TNetworkSimDriverInterfaceBase --> TNetworkedSimulationModelDriver. This is now distinct from "what the simulation needs to call internally" (previously defined by the driver now the simulation).
+	-TNetworkedSimulationModelDriver ("System Driver") is the set of functions the system internally needs to call.
+	-Visual logging has been refactored to be a single function of the system driver that takes all 3 states instead of single methods on the individual system types (Input,Sync,Aux). this simplifies a lot of stuff and allows you to visual sync+aux together.
+-"MockAbilitySystem" has been implemented.
+	-Simple, "hardcoded" abilities that function on top of the flying movement simulation. See notes in MockAbilitySimulation.h
+	-This code gives a good glimpse into what GameplayAbilities will look like (again, more hardcoded in this example but same techniques will be applied).
+
 
 Update (10-16-19)
 -Cleanup how we advance the simulation and unified which keyframes are accessed across the three main buffers
diff --git a/Engine/Plugins/Runtime/NetworkPredictionExtras/Source/NetworkPredictionExtras/Private/MockAbilitySimulation.cpp b/Engine/Plugins/Runtime/NetworkPredictionExtras/Source/NetworkPredictionExtras/Private/MockAbilitySimulation.cpp
new file mode 100644
--- /dev/null
+++ b/Engine/Plugins/Runtime/NetworkPredictionExtras/Source/NetworkPredictionExtras/Private/MockAbilitySimulation.cpp
@@ -0,0 +1,203 @@
+// Copyright 1998-2019 Epic Games, Inc. All Rights Reserved.
+
+#include "MockAbilitySimulation.h"
+#include "DrawDebugHelpers.h"
+#include "VisualLogger/VisualLogger.h"
+
+namespace MockAbilityCVars
+{
+	NETSIM_DEVCVAR_SHIPCONST_FLOAT(DefaultMaxSpeed, 1200.f, "mockability.DefaultMaxSpeed", "Default Speed");
+	NETSIM_DEVCVAR_SHIPCONST_FLOAT(DefaultAcceleration, 4000.f, "mockability.DefaultAcceleration", "Default Acceleration");
+
+	NETSIM_DEVCVAR_SHIPCONST_FLOAT(SprintMaxSpeed, 5000.f, "mockability.SprintMaxSpeed", "Max Speed when sprint is applied.");
+
+	NETSIM_DEVCVAR_SHIPCONST_FLOAT(DashMaxSpeed, 7500.f, "mockability.DashMaxSpeed", "Max Speed when dashing.");
+	NETSIM_DEVCVAR_SHIPCONST_FLOAT(DashAcceleration, 100000.f, "mockability.DashAcceleration", "Acceleration when dashing.");
+}
+
+// -------------------------------------------------------------------------------------------------------------
+//
+// -------------------------------------------------------------------------------------------------------------
+
+
+const FName FMockAbilitySimulation::GroupName(TEXT("Ability"));
+
+void FMockAbilitySimulation::Update(const float DeltaTimeSeconds, const FMockAbilityInputCmd& InputCmd, const FMockAbilitySyncState& InputState, FMockAbilitySyncState& OutputState, const FMockAbilityAuxstate& AuxState, const TNetSimLazyWriter<FMockAbilityAuxstate>& OutAuxStateAccessor)
+{
+	// Stamina passes through. Some code paths will modify this again, but if we don't set the output state it will be garbage/stale
+	// (considering implicit copying of old state to new state by the NetworkedSimModel code, but that could be undesired/inefficient in some cases)
+	OutputState.Stamina = InputState.Stamina;
+
+	static float BlinkCost = 75.f;
+	const bool bBlink = (InputCmd.bBlinkPressed && InputState.Stamina > BlinkCost);
+
+	if (bBlink)
+	{
+		static float BlinkDist = 1000.f;
+		static float OverlapCheckInflation = 0.100f;
+
+		FVector DestLocation = InputState.Location + InputState.Rotation.RotateVector( FVector(BlinkDist, 0.f, 0.f) );
+		AActor* OwningActor = UpdatedComponent->GetOwner();
+		check(OwningActor);
+
+		// DrawDebugLine(OwningActor->GetWorld(), InputState.Location, DestLocation, FColor::Red, false);	
+
+		// Its unfortunate teleporting is so complicated. It may make sense for a new movement simulation to define this themselves, 
+		// but for this mock one, we will just use the engine's AActor teleport.
+		if (OwningActor->TeleportTo(DestLocation, InputState.Rotation))
+		{
+			OutputState = InputState;
+
+			// Component now has the final location
+			const FTransform UpdateComponentTransform = GetUpdateComponentTransform();
+			OutputState.Location = UpdateComponentTransform.GetLocation();
+			OutputState.Stamina = InputState.Stamina - BlinkCost;
+
+			// And we skip the normal update simulation for this frame. This is just a choice. We could still allow it to run.
+			return;
+		}
+	}
+
+	// Dash is implemented in the following way:
+	//	-Stamina consumed on initial press
+	//	-MaxSpeed/Acceleration are jacked up
+	//	-Dash lasts for 400ms (DashDurationMS)
+	//		-division of frame times can cause you to dash for longer. We would have to break up simulation steps to support this 100% accurately.
+	//	-Movement input is synthesized while in dash state. That is, we force forward movement and ignore what was actually fed into the simulation.
+	//
+	// This is just a simple/interesting way of implementing dash in the this system. A real movement/ability system will probably have some concept of
+	// root motion/sources that drive movement forward.
+
+	static float DashCost = 75.f;
+	static int16 DashDurationMS = 400;
+	int16 DashTimeLeft = AuxState.DashTimeLeft;
+	bool bIsDashing = AuxState.DashTimeLeft > 0;
+	
+	if (InputCmd.bDashPressed && InputState.Stamina > DashCost && bIsDashing == false)
+	{
+		// Start dashing
+		DashTimeLeft = DashDurationMS;
+		OutputState.Stamina -= DashCost;
+		bIsDashing = true;
+	}
+
+	if (bIsDashing)
+	{
+		FMockAbilityAuxstate* OutAuxState = OutAuxStateAccessor.Get();
+		OutAuxState->DashTimeLeft = FMath::Max<int16>(DashTimeLeft - (int16)(DeltaTimeSeconds * 1000.f), 0);
+
+		FMockAbilityAuxstate LocalAuxState = AuxState;
+		LocalAuxState.MaxSpeed = MockAbilityCVars::DashMaxSpeed();
+		LocalAuxState.Acceleration = MockAbilityCVars::DashAcceleration();
+
+		FMockAbilityInputCmd LocalInputCmd = InputCmd;
+		LocalInputCmd.MovementInput = FVector(1.f, 0.f, 0.f);
+		
+		FlyingMovement::FMovementSimulation::Update(DeltaTimeSeconds, LocalInputCmd, InputState, OutputState, LocalAuxState, OutAuxStateAccessor);
+
+		if (OutAuxState->DashTimeLeft == 0)
+		{
+			// Stop when dash is over
+			OutputState.Velocity = FVector::ZeroVector;
+		}
+	}
+	else
+	{
+		// Sprint (mutually exclusive from Dash state)		
+		static float SprintBaseCost = 100.f;
+		const float SprintCostThisFrame = SprintBaseCost * DeltaTimeSeconds;
+		const bool bIsSprinting = (InputCmd.bSprintPressed && InputState.Stamina > SprintCostThisFrame);
+
+		// Set our max speed. This is an interesting case.
+		//	-Our input states are already "final". It doesn't make sense to modify the input AuxState data.
+		//	-But, we want to feed this locally calculated max speed into the base movement simulation. So, it is an "input" in that sense.
+		//	-Creating a local copy of the aux state, modifying it, and passing that into the base movement sim is clean, though a bit inefficient.
+		//	-This leaves us with a weird side effect that the calculated MaxSpeed is never written to the AuxBuffer. Though we could write it to the
+		//	output AuxState, that doesn't really do anything useful: the base move sim won't use it and it'll just be left there for next frame (to be overriden again).
+		//
+		// The way the base movement simulation is written, MaxSpeed is a clearly defined input. But in the MockAbility sim, MaxSpeed is a derived
+		// value from other input state. This difference in how the simulations treat the variable is what causes the need to do this.
+		//
+		// It would be possible to write the base movement sim in a way that MaxSpeed is transient value on the sim class (FlyingMovement::FMovementSimulation).
+		// Something like, "MaxSpeed" really means "Base max speed" and their would be a "GetCurrentMaxSpeed" virtual on the sim.
+		// This would make things a bit more awkward in the base case with no ability system. So, for now, this seems like a good pattern/precedence.
+
+		FMockAbilityAuxstate LocalAuxState = AuxState;
+		LocalAuxState.MaxSpeed = bIsSprinting ? MockAbilityCVars::SprintMaxSpeed() : MockAbilityCVars::DefaultMaxSpeed();
+
+		if (bIsSprinting)
+		{
+			OutputState.Stamina = FMath::Max<float>( InputState.Stamina - SprintCostThisFrame, 0.f );
+		}
+		else if (OutputState.Stamina < AuxState.MaxStamina)
+		{
+			OutputState.Stamina = FMath::Min<float>( InputState.Stamina + (DeltaTimeSeconds * AuxState.StaminaRegenRate), AuxState.MaxStamina );
+		}
+		else
+		{
+			OutputState.Stamina = InputState.Stamina;
+		}
+
+		FlyingMovement::FMovementSimulation::Update(DeltaTimeSeconds, InputCmd, InputState, OutputState, LocalAuxState, OutAuxStateAccessor);
+	}
+}
+
+// -------------------------------------------------------------------------------------------------------------
+//
+// -------------------------------------------------------------------------------------------------------------
+
+UMockFlyingAbilityComponent::UMockFlyingAbilityComponent()
+{
+
+}
+
+INetworkSimulationModel* UMockFlyingAbilityComponent::InstantiateNetworkSimulation()
+{
+	check(UpdatedComponent);
+	
+	// FIXME: init sim helper on super?
+	MockAbilitySimulation.Reset(new FMockAbilitySimulation());
+	MockAbilitySimulation->UpdatedComponent = UpdatedComponent;
+	MockAbilitySimulation->UpdatedPrimitive = UpdatedPrimitive;
+	
+	FMockAbilitySyncState InitialSyncState;
+	InitialSyncState.Location = UpdatedComponent->GetComponentLocation();
+	InitialSyncState.Rotation = UpdatedComponent->GetComponentQuat().Rotator();	
+
+	FMockAbilityAuxstate InitialAuxState;
+	InitialAuxState.MaxSpeed = MockAbilityCVars::DefaultMaxSpeed();
+
+	auto NewModel = new FMockAbilitySystem<0>(MockAbilitySimulation.Get(), this, InitialSyncState, InitialAuxState);
+
+	MovementSyncState.Init(NewModel);
+	MovementAuxState.Init(NewModel);
+	return NewModel;
+}
+
+void UMockFlyingAbilityComponent::ProduceInput(const FNetworkSimTime SimTime, FMockAbilityInputCmd& Cmd)
+{
+	// This isn't ideal. It probably makes sense for the component to do all the input binding rather.
+	ProduceInputDelegate.ExecuteIfBound(SimTime, Cmd);
+}
+
+void UMockFlyingAbilityComponent::FinalizeFrame(const FMockAbilitySyncState& SyncState, const FMockAbilityAuxstate& AuxState)
+{
+	Super::FinalizeFrame(SyncState, AuxState);
+}
+
+FString UMockFlyingAbilityComponent::GetDebugName() const
+{
+	return FString::Printf(TEXT("MockAbility. %s. %s"), *UEnum::GetValueAsString(TEXT("Engine.ENetRole"), GetOwnerRole()), *GetName());
+}
+
+const AActor* UMockFlyingAbilityComponent::GetVLogOwner() const
+{
+	return GetOwner();
+}
+
+void UMockFlyingAbilityComponent::VisualLog(const FMockAbilityInputCmd* Input, const FMockAbilitySyncState* Sync, const FMockAbilityAuxstate* Aux, const FVisualLoggingParameters& SystemParameters) const
+{
+	Super::VisualLog(Input, Sync, Aux, SystemParameters);
+}
+
+// ---------------------------------------------------------------------------------
diff --git a/Engine/Plugins/Runtime/NetworkPredictionExtras/Source/NetworkPredictionExtras/Private/MockNetworkSimulation.cpp b/Engine/Plugins/Runtime/NetworkPredictionExtras/Source/NetworkPredictionExtras/Private/MockNetworkSimulation.cpp
--- a/Engine/Plugins/Runtime/NetworkPredictionExtras/Source/NetworkPredictionExtras/Private/MockNetworkSimulation.cpp
+++ b/Engine/Plugins/Runtime/NetworkPredictionExtras/Source/NetworkPredictionExtras/Private/MockNetworkSimulation.cpp
@@ -1,7 +1,6 @@
 // Copyright 1998-2019 Epic Games, Inc. All Rights Reserved.
 
 #include "MockNetworkSimulation.h"
-#include "NetworkSimulationModelDebugger.h"
 #include "GameFramework/Actor.h"
 #include "Engine/World.h"
 #include "Components/PrimitiveComponent.h"
@@ -62,7 +61,7 @@ static bool ForceMispredict = false;
 
 const FName FMockNetworkSimulation::GroupName(TEXT("Mock"));
 
-void FMockNetworkSimulation::Update(IMockDriver* Driver, const float DeltaTimeSeconds, const FMockInputCmd& InputCmd, const FMockSyncState& InputState, FMockSyncState& OutputState, const FMockAuxState& AuxState, const TLazyStateAccessor<FMockAuxState>& OutAuxStateAccessor)
+void FMockNetworkSimulation::Update(const float DeltaTimeSeconds, const FMockInputCmd& InputCmd, const FMockSyncState& InputState, FMockSyncState& OutputState, const FMockAuxState& AuxState, const TNetSimLazyWriter<FMockAuxState>& OutAuxStateAccessor)
 {
 	OutputState.Total = InputState.Total + (InputCmd.InputValue * AuxState.Multiplier * DeltaTimeSeconds);
 
@@ -77,31 +76,6 @@ void FMockNetworkSimulation::Update(IMockDriver* Driver, const float DeltaTimeSe
 	}
 }
 
-void FMockSyncState::VisualLog(const FVisualLoggingParameters& Parameters, IMockDriver* Driver, IMockDriver* LogDriver) const
-{
-	// This function is awkward because the mock example doesn't include any positional data, so we added an explicit ::GetDebugWorldTransform to the driver interface.
-	// still, this doesn't give us historical location info. So we will just use Z offset based on the EvisualLoggingContext to space things out.
-	const bool DrawStyleMinor = (Parameters.Context == EVisualLoggingContext::OtherMispredicted || Parameters.Context == EVisualLoggingContext::OtherPredicted);
-	const FTransform WorldTransform = LogDriver->GetDebugWorldTransform(); // No sense using the Driver's world transform (this would be the server's position in PIE, not useful in mock sim)
-	UWorld* DebugWorld = LogDriver->GetDriverWorld();
-
-	FString DebugStr = FString::Printf(TEXT("[%d][%s] Total: %.4f"), Parameters.Keyframe, *LexToString(Parameters.Context), this->Total);
-	if (MockNetworkSimCVars::UseDrawDebug)
-	{
-		const float Lifetime = Parameters.Lifetime == EVisualLoggingLifetime::Transient ? 0.0001f : MockNetworkSimCVars::DrawDebugDefaultLifeTime;
-		const float YOffset = 0.f; //DrawStyleMinor ? ((100.f * (Parameters.Keyframe % 32)) - 50.f) : 0.f;
-		const float ZOffset = 20.f * (int32)Parameters.Context;
-		
-		DrawDebugString( DebugWorld, WorldTransform.GetLocation() + WorldTransform.GetRotation().RotateVector( FVector(0.f, YOffset, ZOffset) ), *DebugStr, nullptr, Parameters.GetDebugColor(), Lifetime );
-	}
-
-	if (MockNetworkSimCVars::UseVLogger)
-	{
-		// I don't see a "Draw string in th world" function in VLOG so we are just doing this
-		UE_VLOG(LogDriver->GetVLogOwner(), LogMockNetworkSim, Log, TEXT("%s"), *DebugStr);
-	}
-}
-
 // ===============================================================================================
 //
 // ===============================================================================================
@@ -134,12 +108,13 @@ UMockNetworkSimulationComponent::UMockNetworkSimulationComponent()
 
 INetworkSimulationModel* UMockNetworkSimulationComponent::InstantiateNetworkSimulation()
 {
+	MockNetworkSimulation.Reset(new FMockNetworkSimulation());
+
 	FMockSyncState InitialSyncState;
 	InitialSyncState.Total = MockValue;
 
-	auto* NewSim = new FMockNetworkModel(this);
-	DO_NETSIM_MODEL_DEBUG(FNetworkSimulationModelDebuggerManager::Get().RegisterNetworkSimulationModel(NewSim, GetOwner()));
-	return NewSim;
+	auto* NewModel = new FMockNetworkModel(MockNetworkSimulation.Get(), this, InitialSyncState);
+	return NewModel;
 }
 
 void UMockNetworkSimulationComponent::TickComponent(float DeltaTime, enum ELevelTick TickType, FActorComponentTickFunction *ThisTickFunction)
@@ -179,14 +154,20 @@ FString UMockNetworkSimulationComponent::GetDebugName() const
 	return FString::Printf(TEXT("MockSim. %s. %s"), *UEnum::GetValueAsString(TEXT("Engine.ENetRole"), GetOwnerRole()), *GetName());
 }
 
-FTransform UMockNetworkSimulationComponent::GetDebugWorldTransform() const
+AActor* UMockNetworkSimulationComponent::GetVLogOwner() const
 {
-	return GetOwner()->GetActorTransform();
+	return GetOwner();
 }
 
-UObject* UMockNetworkSimulationComponent::GetVLogOwner() const
+void UMockNetworkSimulationComponent::VisualLog(const FMockInputCmd* Input, const FMockSyncState* Sync, const FMockAuxState* Aux, const FVisualLoggingParameters& SystemParameters) const
 {
-	return GetOwner();
+	FVisualLoggingParameters LocalParameters = SystemParameters;
+	AActor* Owner = GetOwner();
+	FTransform Transform = Owner->GetActorTransform();
+
+	LocalParameters.DebugString += FString::Printf(TEXT(" [%d] Total: %.4f"), SystemParameters.Keyframe, Sync->Total);
+
+	FVisualLoggingHelpers::VisualLogActor(Owner, Transform, LocalParameters);
 }
 
 // ------------------------------------------------------------------------------------------------------------------------
diff --git a/Engine/Plugins/Runtime/NetworkPredictionExtras/Source/NetworkPredictionExtras/Private/NetworkPredictionExtrasFlyingPawn.cpp b/Engine/Plugins/Runtime/NetworkPredictionExtras/Source/NetworkPredictionExtras/Private/NetworkPredictionExtrasFlyingPawn.cpp
--- a/Engine/Plugins/Runtime/NetworkPredictionExtras/Source/NetworkPredictionExtras/Private/NetworkPredictionExtrasFlyingPawn.cpp
+++ b/Engine/Plugins/Runtime/NetworkPredictionExtras/Source/NetworkPredictionExtras/Private/NetworkPredictionExtrasFlyingPawn.cpp
@@ -17,6 +17,7 @@
 #include "Misc/AssertionMacros.h"
 #include "HAL/PlatformStackWalk.h"
 #include "HAL/ThreadHeartBeat.h"
+#include "MockAbilitySimulation.h"
 
 namespace FlyingPawnCVars
 {
@@ -33,9 +34,12 @@ static FAutoConsoleVariableRef CVarBindAutomatically(TEXT("NetworkPredictionExtr
 	BindAutomatically, TEXT("Binds local input and mispredict commands to 5 and 6 respectively"), ECVF_Default);
 }
 
-ANetworkPredictionExtrasFlyingPawn::ANetworkPredictionExtrasFlyingPawn()
+const FName Name_FlyingMovementComponent(TEXT("FlyingMovementComponent"));
+
+ANetworkPredictionExtrasFlyingPawn::ANetworkPredictionExtrasFlyingPawn(const FObjectInitializer& ObjectInitializer)
 {
-	FlyingMovementComponent = CreateDefaultSubobject<UFlyingMovementComponent>(TEXT("FlyingMovementComponent"));
+	FlyingMovementComponent = CreateDefaultSubobject<UFlyingMovementComponent>(Name_FlyingMovementComponent);
+	ensure(FlyingMovementComponent);
 }
 
 void ANetworkPredictionExtrasFlyingPawn::BeginPlay()
@@ -44,6 +48,7 @@ void ANetworkPredictionExtrasFlyingPawn::BeginPlay()
 
 	if (UWorld* World = GetWorld())
 	{
+		ensure(FlyingMovementComponent);
 		FlyingMovementComponent->ProduceInputDelegate.BindUObject(this, &ANetworkPredictionExtrasFlyingPawn::ProduceInput);
 
 		// Binds 0 and 9 to the debug hud commands. This is just a convenience for the extras plugin. Real projects should bind this themselves
@@ -84,7 +89,6 @@ UNetConnection* ANetworkPredictionExtrasFlyingPawn::GetNetConnection() const
 			return GetWorld()->GetFirstPlayerController()->GetNetConnection();
 		}
 	}
-	
 
 	return nullptr;
 }
@@ -115,8 +119,8 @@ void ANetworkPredictionExtrasFlyingPawn::SetupPlayerInputComponent(UInputCompone
 			PC->PlayerInput->AddAxisMapping(FInputAxisKeyMapping(TEXT("MoveRight"), EKeys::A, -1.f));
 			PC->PlayerInput->AddAxisMapping(FInputAxisKeyMapping(TEXT("MoveForward"), EKeys::W, 1.f));
 			PC->PlayerInput->AddAxisMapping(FInputAxisKeyMapping(TEXT("MoveForward"), EKeys::S, -1.f));
-			PC->PlayerInput->AddAxisMapping(FInputAxisKeyMapping(TEXT("LeftTriggerAxis"), EKeys::LeftControl, 1.f));
-			PC->PlayerInput->AddAxisMapping(FInputAxisKeyMapping(TEXT("RightTriggerAxis"), EKeys::LeftShift, 1.f));
+			PC->PlayerInput->AddAxisMapping(FInputAxisKeyMapping(TEXT("LeftTriggerAxis"), EKeys::C, 1.f));
+			PC->PlayerInput->AddAxisMapping(FInputAxisKeyMapping(TEXT("RightTriggerAxis"), EKeys::LeftControl, 1.f));
 
 			// Mouse
 			PC->PlayerInput->AddAxisMapping(FInputAxisKeyMapping(TEXT("LookYaw"), EKeys::MouseX));
@@ -445,4 +449,127 @@ void ANetworkPredictionExtrasFlyingPawn::ProduceInput(const FNetworkSimTime SimT
 
 	CachedMoveInput = FVector::ZeroVector;
 	CachedLookInput = FVector2D::ZeroVector;
+}
+
+// ------------------------------------------------------------------------
+
+ANetworkPredictionExtrasFlyingPawn_MockAbility::ANetworkPredictionExtrasFlyingPawn_MockAbility(const FObjectInitializer& ObjectInitializer)
+	: Super(ObjectInitializer.SetDefaultSubobjectClass<UMockFlyingAbilityComponent>(Name_FlyingMovementComponent))
+{
+
+}
+
+UMockFlyingAbilityComponent* ANetworkPredictionExtrasFlyingPawn_MockAbility::GetMockFlyingAbilityComponent()
+{
+	return Cast<UMockFlyingAbilityComponent>(FlyingMovementComponent);
+}
+
+const UMockFlyingAbilityComponent* ANetworkPredictionExtrasFlyingPawn_MockAbility::GetMockFlyingAbilityComponent() const
+{
+	return Cast<UMockFlyingAbilityComponent>(FlyingMovementComponent);
+}
+
+void ANetworkPredictionExtrasFlyingPawn_MockAbility::ProduceInput(const FNetworkSimTime SimTime, FMockAbilityInputCmd& Cmd)
+{
+	Super::ProduceInput(SimTime, Cmd);
+	Cmd.bSprintPressed = bSprintPressed;
+	Cmd.bDashPressed = bDashPressed;
+	Cmd.bBlinkPressed = bBlinkPressed;
+}
+
+void ANetworkPredictionExtrasFlyingPawn_MockAbility::BeginPlay()
+{
+	Super::BeginPlay();
+
+	if (UWorld* World = GetWorld())
+	{
+		if (UMockFlyingAbilityComponent* FlyingAbilityComponent = GetMockFlyingAbilityComponent())
+		{
+			FlyingAbilityComponent->ProduceInputDelegate.BindUObject(this, &ANetworkPredictionExtrasFlyingPawn_MockAbility::ProduceInput);
+		}
+	}
+}
+
+float ANetworkPredictionExtrasFlyingPawn_MockAbility::GetStamina() const
+{
+	if (const UMockFlyingAbilityComponent* FlyingAbilityComponent = GetMockFlyingAbilityComponent())
+	{
+		if (const FMockAbilitySyncState* SyncState = FlyingAbilityComponent->MovementSyncState.GetStateRead())
+		{
+			return SyncState->Stamina;
+		}
+	}
+	return 0.f;
+}
+
+float ANetworkPredictionExtrasFlyingPawn_MockAbility::GetMaxStamina() const
+{
+	if (const UMockFlyingAbilityComponent* FlyingAbilityComponent = GetMockFlyingAbilityComponent())
+	{
+		if (const FMockAbilityAuxstate* AuxState = FlyingAbilityComponent->MovementAuxState.GetStateRead())
+		{
+			return AuxState->MaxStamina;
+		}
+	}
+	return 0.f;
+}
+
+void ANetworkPredictionExtrasFlyingPawn_MockAbility::SetupPlayerInputComponent(UInputComponent* PlayerInputComponent)
+{
+	Super::SetupPlayerInputComponent(PlayerInputComponent);
+
+	if (APlayerController* PC = Cast<APlayerController>(Controller))
+	{
+		if (PC->PlayerInput)
+		{
+			//Gamepad
+			PC->PlayerInput->AddActionMapping(FInputActionKeyMapping(TEXT("Sprint"), EKeys::Gamepad_LeftThumbstick));
+			PC->PlayerInput->AddActionMapping(FInputActionKeyMapping(TEXT("Dash"), EKeys::Gamepad_FaceButton_Left));
+			PC->PlayerInput->AddActionMapping(FInputActionKeyMapping(TEXT("Blink"), EKeys::Gamepad_FaceButton_Top));
+
+			// Keyboard
+			PC->PlayerInput->AddActionMapping(FInputActionKeyMapping(TEXT("Sprint"), EKeys::LeftShift));
+			PC->PlayerInput->AddActionMapping(FInputActionKeyMapping(TEXT("Dash"), EKeys::Q));
+			PC->PlayerInput->AddActionMapping(FInputActionKeyMapping(TEXT("Blink"), EKeys::E));
+		}
+	}
+
+	PlayerInputComponent->BindAction(TEXT("Sprint"), IE_Pressed, this, &ThisClass::Action_Sprint_Pressed);
+	PlayerInputComponent->BindAction(TEXT("Sprint"), IE_Released, this, &ThisClass::Action_Sprint_Released);
+
+	PlayerInputComponent->BindAction(TEXT("Dash"), IE_Pressed, this, &ThisClass::Action_Dash_Pressed);
+	PlayerInputComponent->BindAction(TEXT("Dash"), IE_Released, this, &ThisClass::Action_Dash_Released);
+
+	PlayerInputComponent->BindAction(TEXT("Blink"), IE_Pressed, this, &ThisClass::Action_Blink_Pressed);
+	PlayerInputComponent->BindAction(TEXT("Blink"), IE_Released, this, &ThisClass::Action_Blink_Released);
+}
+
+void ANetworkPredictionExtrasFlyingPawn_MockAbility::Action_Sprint_Pressed()
+{
+	bSprintPressed = true;
+}
+
+void ANetworkPredictionExtrasFlyingPawn_MockAbility::Action_Sprint_Released()
+{
+	bSprintPressed = false;
+}
+
+void ANetworkPredictionExtrasFlyingPawn_MockAbility::Action_Dash_Pressed()
+{
+	bDashPressed = true;
+}
+
+void ANetworkPredictionExtrasFlyingPawn_MockAbility::Action_Dash_Released()
+{
+	bDashPressed = false;
+}
+
+void ANetworkPredictionExtrasFlyingPawn_MockAbility::Action_Blink_Pressed()
+{
+	bBlinkPressed = true;
+}
+
+void ANetworkPredictionExtrasFlyingPawn_MockAbility::Action_Blink_Released()
+{
+	bBlinkPressed = false;
 }
\ No newline at end of file
diff --git a/Engine/Plugins/Runtime/NetworkPredictionExtras/Source/NetworkPredictionExtras/Public/MockAbilitySimulation.h b/Engine/Plugins/Runtime/NetworkPredictionExtras/Source/NetworkPredictionExtras/Public/MockAbilitySimulation.h
new file mode 100644
--- /dev/null
+++ b/Engine/Plugins/Runtime/NetworkPredictionExtras/Source/NetworkPredictionExtras/Public/MockAbilitySimulation.h
@@ -0,0 +1,168 @@
+// Copyright 1998-2019 Epic Games, Inc. All Rights Reserved.
+
+#pragma once
+
+#include "CoreMinimal.h"
+#include "UObject/ObjectMacros.h"
+#include "UObject/UObjectGlobals.h"
+#include "Engine/NetSerialization.h"
+#include "Engine/EngineTypes.h"
+#include "Engine/EngineBaseTypes.h"
+#include "WorldCollision.h"
+#include "Components/ActorComponent.h"
+#include "Components/SceneComponent.h"
+#include "Misc/OutputDevice.h"
+#include "Misc/CoreDelegates.h"
+
+#include "NetworkSimulationModel.h"
+#include "NetworkPredictionComponent.h"
+#include "Movement/FlyingMovement.h"
+
+#include "MockAbilitySimulation.generated.h"
+
+// -------------------------------------------------------------------------------------------------------------------------------
+// Mock Ability Simulation
+//	This is meant to illustrate how a higher level simulation can build off an existing one. While something like GameplayAbilities
+//	is more generic and data driven, this illustrates how it will need to solve core issues.
+//
+//	This implements:
+//	1. Stamina "attribute": basic attribute with max + regen value that is consumed by abilities.
+//	2. Sprint: Increased max speed while sprint button is held. Drains stamina each frame.
+//	3. Dash: Immediate acceleration to a high speed for X seconds.
+//	4. Blink: Teleport to location X units ahead
+//
+//
+//	Notes/Todo:
+//	-Cooldown/timers for easier tracking (currently relying on high stamina cost to avoid activating each frame)
+//	-Outward events for gameplay cue type: tying cosmetic events to everything
+//	-More efficient packing of input/pressed buttons. Input handling in general should detect up/down presses rather than just current state
+//
+// -------------------------------------------------------------------------------------------------------------------------------
+
+struct FMockAbilityInputCmd : public FlyingMovement::FInputCmd
+{
+	bool bSprintPressed = false;
+	bool bDashPressed = false;
+	bool bBlinkPressed = false;
+
+	void NetSerialize(const FNetSerializeParams& P)
+	{
+		P.Ar << bSprintPressed;
+		P.Ar << bDashPressed;
+		P.Ar << bBlinkPressed;
+		FlyingMovement::FInputCmd::NetSerialize(P);
+	}
+
+	void Log(FStandardLoggingParameters& P) const
+	{
+		FlyingMovement::FInputCmd::Log(P);
+		if (P.Context == EStandardLoggingContext::Full)
+		{
+			P.Ar->Logf(TEXT("bSprintPressed: %d"), bSprintPressed);
+			P.Ar->Logf(TEXT("bDashPressed: %d"), bDashPressed);
+			P.Ar->Logf(TEXT("bBlinkPressed: %d"), bBlinkPressed);
+		}
+	}
+};
+
+struct FMockAbilitySyncState : public FlyingMovement::FMoveState
+{
+	float Stamina = 0.f;
+	
+	void NetSerialize(const FNetSerializeParams& P)
+	{
+		P.Ar << Stamina;
+		FlyingMovement::FMoveState::NetSerialize(P);
+	}
+
+	void Log(FStandardLoggingParameters& P) const
+	{
+		FlyingMovement::FMoveState::Log(P);
+		if (P.Context == EStandardLoggingContext::Full)
+		{
+			P.Ar->Logf(TEXT("Stamina: %.2f"), Stamina);
+		}
+	}
+};
+
+struct FMockAbilityAuxstate : public FlyingMovement::FAuxState
+{
+	float MaxStamina = 100.f;
+	float StaminaRegenRate = 20.f;
+	int16 DashTimeLeft = 0;
+
+	void NetSerialize(const FNetSerializeParams& P)
+	{
+		P.Ar << MaxStamina;
+		P.Ar << StaminaRegenRate;
+		P.Ar << DashTimeLeft;
+		FlyingMovement::FAuxState::NetSerialize(P);
+	}
+
+	void Log(FStandardLoggingParameters& P) const
+	{
+		FlyingMovement::FAuxState::Log(P);
+		if (P.Context == EStandardLoggingContext::Full)
+		{
+			P.Ar->Logf(TEXT("MaxStamina: %.2f"), MaxStamina);
+			P.Ar->Logf(TEXT("StaminaRegenRate: %.2f"), StaminaRegenRate);
+			P.Ar->Logf(TEXT("StaminaRegenRate: %.2f"), DashTimeLeft);
+		}
+	}
+};
+
+class FMockAbilitySimulation : public FlyingMovement::FMovementSimulation
+{
+public:
+	/** Tick group the simulation maps to */
+	static const FName GroupName;
+
+	/** Main update function */
+	void Update(const float DeltaTimeSeconds, const FMockAbilityInputCmd& InputCmd, const FMockAbilitySyncState& InputState, FMockAbilitySyncState& OutputState, const FMockAbilityAuxstate& AuxState, const TNetSimLazyWriter<FMockAbilityAuxstate>& OutAuxStateAccessor);
+};
+
+using TMockAbilityBufferTypes = TNetworkSimBufferTypes<FMockAbilityInputCmd, FMockAbilitySyncState, FMockAbilityAuxstate>;
+
+template<int32 InFixedStepMS=0>
+using FMockAbilitySystem = TNetworkedSimulationModel<FMockAbilitySimulation, TMockAbilityBufferTypes, TNetworkSimTickSettings<InFixedStepMS>>;
+
+class IMockFlyingAbilitySystemDriver : public TNetworkedSimulationModelDriver<TMockAbilityBufferTypes> { };
+
+// -------------------------------------------------------------------------------------------------------------------------------
+// ActorComponent for running Mock Ability Simulation 
+// -------------------------------------------------------------------------------------------------------------------------------
+
+UCLASS(BlueprintType, meta=(BlueprintSpawnableComponent))
+class NETWORKPREDICTIONEXTRAS_API UMockFlyingAbilityComponent : public UFlyingMovementComponent, public IMockFlyingAbilitySystemDriver
+{
+	GENERATED_BODY()
+
+public:
+
+	UMockFlyingAbilityComponent();
+
+	DECLARE_DELEGATE_TwoParams(FProduceMockAbilityInput, const FNetworkSimTime /*SimTime*/, FMockAbilityInputCmd& /*Cmd*/)
+	FProduceMockAbilityInput ProduceInputDelegate;
+
+	TNetworkSimStateAccessor<FMockAbilitySyncState> MovementSyncState;
+	TNetworkSimStateAccessor<FMockAbilityAuxstate> MovementAuxState;
+
+	// IMockFlyingAbilitySystemDriver
+	FString GetDebugName() const override;
+	const AActor* GetVLogOwner() const override;
+	void VisualLog(const FMockAbilityInputCmd* Input, const FMockAbilitySyncState* Sync, const FMockAbilityAuxstate* Aux, const FVisualLoggingParameters& SystemParameters) const override;
+
+	void ProduceInput(const FNetworkSimTime SimTime, FMockAbilityInputCmd& Cmd) override;
+	void FinalizeFrame(const FMockAbilitySyncState& SyncState, const FMockAbilityAuxstate& AuxState) override;
+
+	// Required to supress -Woverloaded-virtual warning. We are effectively hiding UFlyingMovementComponent's methods
+	using UFlyingMovementComponent::VisualLog;
+	using UFlyingMovementComponent::ProduceInput;
+	using UFlyingMovementComponent::FinalizeFrame;
+
+protected:
+
+	// Network Prediction
+	virtual INetworkSimulationModel* InstantiateNetworkSimulation() override;
+	TUniquePtr<FMockAbilitySimulation> MockAbilitySimulation;
+};
\ No newline at end of file
diff --git a/Engine/Plugins/Runtime/NetworkPredictionExtras/Source/NetworkPredictionExtras/Public/MockNetworkSimulation.h b/Engine/Plugins/Runtime/NetworkPredictionExtras/Source/NetworkPredictionExtras/Public/MockNetworkSimulation.h
--- a/Engine/Plugins/Runtime/NetworkPredictionExtras/Source/NetworkPredictionExtras/Public/MockNetworkSimulation.h
+++ b/Engine/Plugins/Runtime/NetworkPredictionExtras/Source/NetworkPredictionExtras/Public/MockNetworkSimulation.h
@@ -99,8 +99,6 @@ struct FMockSyncState
 		}
 	}
 
-	void VisualLog(const FVisualLoggingParameters& Parameters, IMockDriver* Driver, IMockDriver* LogDriver) const;
-
 	static void Interpolate(const FMockSyncState& From, const FMockSyncState& To, const float PCT, FMockSyncState& OutDest)
 	{
 		OutDest.Total = From.Total + ((To.Total - From.Total) * PCT);
@@ -133,41 +131,34 @@ struct FMockAuxState
 
 using TMockNetworkSimulationBufferTypes = TNetworkSimBufferTypes<FMockInputCmd, FMockSyncState, FMockAuxState>;
 
-static FName MockSimulationGroupName("Mock");
 
-// Interface between the simulation and owning component driving it. Functions added here are available in ::Update and must be implemented by UMockNetworkSimulationComponent.
-class IMockDriver : public TNetworkSimDriverInterfaceBase<TMockNetworkSimulationBufferTypes>
-{
-public:
 
-	virtual UWorld* GetDriverWorld() const = 0;
-	virtual FTransform GetDebugWorldTransform() const = 0;
-};
+static FName MockSimulationGroupName("Mock");
 
 class FMockNetworkSimulation
 {
 public:
 
 	/** Main update function */
-	static void Update(IMockDriver* Driver, const float DeltaTimeSeconds, const FMockInputCmd& InputCmd, const FMockSyncState& InputState, FMockSyncState& OutputState, const FMockAuxState& AuxState, const TLazyStateAccessor<FMockAuxState>& OutAuxStateAccessor);
+	void Update(const float DeltaTimeSeconds, const FMockInputCmd& InputCmd, const FMockSyncState& InputState, FMockSyncState& OutputState, const FMockAuxState& AuxState, const TNetSimLazyWriter<FMockAuxState>& OutAuxStateAccessor);
 
 	/** Tick group the simulation maps to */
 	static const FName GroupName;
 };
 
 // Actual definition of our network simulation.
-using FMockNetworkModel = TNetworkedSimulationModel<FMockNetworkSimulation, IMockDriver, TMockNetworkSimulationBufferTypes>;
+using FMockNetworkModel = TNetworkedSimulationModel<FMockNetworkSimulation, TMockNetworkSimulationBufferTypes>;
 
 // Needed to trick UHT into letting UMockNetworkSimulationComponent implement. UHT cannot parse the ::
 // Also needed for forward declaring. Can't just be a typedef/using =
-class IMockNetworkSimulationDriver : public IMockDriver { };
+class IMockSystemDriver : public TNetworkedSimulationModelDriver<TMockNetworkSimulationBufferTypes> { };
 
 // -------------------------------------------------------------------------------------------------------------------------------
 // ActorComponent for running a MockNetworkSimulation
 // -------------------------------------------------------------------------------------------------------------------------------
 
 UCLASS(BlueprintType, meta=(BlueprintSpawnableComponent))
-class NETWORKPREDICTIONEXTRAS_API UMockNetworkSimulationComponent : public UNetworkPredictionComponent, public IMockNetworkSimulationDriver
+class NETWORKPREDICTIONEXTRAS_API UMockNetworkSimulationComponent : public UNetworkPredictionComponent, public IMockSystemDriver
 {
 	GENERATED_BODY()
 
@@ -182,15 +173,14 @@ public:
 public:
 
 	FString GetDebugName() const override;
-	virtual UObject* GetVLogOwner() const override final;
+	AActor* GetVLogOwner() const override;
+	void VisualLog(const FMockInputCmd* Input, const FMockSyncState* Sync, const FMockAuxState* Aux, const FVisualLoggingParameters& SystemParameters) const override;
 	
 	void ProduceInput(const FNetworkSimTime SimTime, FMockInputCmd& Cmd);
-	void FinalizeFrame(const FMockSyncState& SyncState, const FMockAuxState& AuxState) override;
-	
-
-	virtual UWorld* GetDriverWorld() const override final { return GetWorld(); }
-	virtual FTransform GetDebugWorldTransform() const override final;
+	void FinalizeFrame(const FMockSyncState& SyncState, const FMockAuxState& AuxState) override;	
 
 	// Mock representation of "syncing' to the sync state in the network sim.
 	float MockValue = 1000.f;
+
+	TUniquePtr<FMockNetworkSimulation> MockNetworkSimulation;
 };
\ No newline at end of file
diff --git a/Engine/Plugins/Runtime/NetworkPredictionExtras/Source/NetworkPredictionExtras/Public/NetworkPredictionExtrasFlyingPawn.h b/Engine/Plugins/Runtime/NetworkPredictionExtras/Source/NetworkPredictionExtras/Public/NetworkPredictionExtrasFlyingPawn.h
--- a/Engine/Plugins/Runtime/NetworkPredictionExtras/Source/NetworkPredictionExtras/Public/NetworkPredictionExtrasFlyingPawn.h
+++ b/Engine/Plugins/Runtime/NetworkPredictionExtras/Source/NetworkPredictionExtras/Public/NetworkPredictionExtrasFlyingPawn.h
@@ -3,6 +3,7 @@
 
 #include "GameFramework/Pawn.h"
 #include "Movement/FlyingMovement.h"
+#include "MockAbilitySimulation.h"
 #include "NetworkPredictionExtrasFlyingPawn.generated.h"
 
 class UInputComponent;
@@ -41,19 +42,19 @@ enum class ENetworkPredictionExtrasFlyingInputPreset: uint8
 };
 
 /** Sample pawn that uses UFlyingMovementComponent. The main thing this provides is actually producing user input for the component/simulation to consume. */
-UCLASS(config = Game)
+UCLASS()
 class NETWORKPREDICTIONEXTRAS_API ANetworkPredictionExtrasFlyingPawn : public APawn
 {
 	GENERATED_BODY()
 
 public:
 
-	ANetworkPredictionExtrasFlyingPawn();
+	ANetworkPredictionExtrasFlyingPawn(const FObjectInitializer& ObjectInitializer);
 
 	virtual void BeginPlay() override;
 	virtual void SetupPlayerInputComponent(UInputComponent* PlayerInputComponent) override;
 	virtual void Tick( float DeltaSeconds) override;
-	virtual UNetConnection* GetNetConnection() const override;
+	virtual UNetConnection* GetNetConnection() const override; // For bFakeAutonomousProxy only
 
 
 	UPROPERTY(EditAnywhere, BlueprintReadWrite, Category="Automation")
@@ -84,10 +85,13 @@ protected:
 	const FlyingMovement::FAuxState* GetAuxStateRead() const;
 	FlyingMovement::FAuxState* GetAuxStateWrite();
 
-private:
-
 	void ProduceInput(const FNetworkSimTime SimTime, FlyingMovement::FInputCmd& Cmd);
 
+	UPROPERTY()
+	UFlyingMovementComponent* FlyingMovementComponent;
+
+private:
+
 	FVector CachedMoveInput;
 	FVector2D CachedLookInput;
 
@@ -102,7 +106,42 @@ private:
 	void Action_LeftShoulder_Released() { }
 	void Action_RightShoulder_Pressed() { }
 	void Action_RightShoulder_Released() { }
+};
 
-	UPROPERTY()
-	UFlyingMovementComponent* FlyingMovementComponent;
+// Example subclass of ANetworkPredictionExtrasFlyingPawn that uses the MockAbility simulation
+UCLASS()
+class NETWORKPREDICTIONEXTRAS_API ANetworkPredictionExtrasFlyingPawn_MockAbility : public ANetworkPredictionExtrasFlyingPawn
+{
+	GENERATED_BODY()
+
+public:
+
+	ANetworkPredictionExtrasFlyingPawn_MockAbility(const FObjectInitializer& ObjectInitializer);
+	virtual void BeginPlay() override;
+	virtual void SetupPlayerInputComponent(UInputComponent* PlayerInputComponent) override;
+
+	
+	UMockFlyingAbilityComponent* GetMockFlyingAbilityComponent();
+	const UMockFlyingAbilityComponent* GetMockFlyingAbilityComponent() const;
+
+	UFUNCTION(BlueprintCallable, Category="Ability")
+	float GetStamina() const;
+
+	UFUNCTION(BlueprintCallable, Category="Ability")
+	float GetMaxStamina() const;
+	
+protected:
+	using ANetworkPredictionExtrasFlyingPawn::ProduceInput;
+	void ProduceInput(const FNetworkSimTime SimTime, FMockAbilityInputCmd& Cmd);
+
+	void Action_Sprint_Pressed();
+	void Action_Sprint_Released();
+	void Action_Dash_Pressed();
+	void Action_Dash_Released();
+	void Action_Blink_Pressed();
+	void Action_Blink_Released();
+
+	bool bSprintPressed = false;
+	bool bDashPressed = false;
+	bool bBlinkPressed = false;
 };
\ No newline at end of file
