Network Prediction Plugin: Renaming/cleanup on ::Update function:
-"TSimulation::Update" is now "TSimulation::SimulationTick". Makes searching through source a bit easier, less generic name
-Collapsed parameters of SimulationTick into 3 containers: TNetSimTimeStep, TNetSimInput, TNetSimOutput
-Time passed into simulation is now FNetworkSimTime (instead of forcing float RealTimeSeconds conversion at callsite)
-TSimulationTickState broken up into FSimulationTickState and TSimulationTicker.
-::SimulationTick now also gets a reference to FSimulationTickState. This is setting us up to do timers/countdowns based on sim time within these functions.

#rnx

diff --git a/Engine/Plugins/Runtime/NetworkPrediction/Source/NetworkPrediction/Private/Movement/FlyingMovementSimulation.cpp b/Engine/Plugins/Runtime/NetworkPrediction/Source/NetworkPrediction/Private/Movement/FlyingMovementSimulation.cpp
--- a/Engine/Plugins/Runtime/NetworkPrediction/Source/NetworkPrediction/Private/Movement/FlyingMovementSimulation.cpp
+++ b/Engine/Plugins/Runtime/NetworkPrediction/Source/NetworkPrediction/Private/Movement/FlyingMovementSimulation.cpp
@@ -150,10 +150,11 @@ namespace FlyingMovement
 		return 0.f;
 	}
 
-
-	void FMovementSimulation::Update(const float DeltaSeconds, const FInputCmd& InputCmd, const FMoveState& InputState, FMoveState& OutputState, const FAuxState& AuxState, const TNetSimLazyWriter<FAuxState>& OutAuxStateAccessor)
+	void FMovementSimulation::SimulationTick(const TNetSimTimeStep& TimeStep, const TNetSimInput<TMovementBufferTypes>& Input, const TNetSimOutput<TMovementBufferTypes>& Output)
 	{
-		OutputState = InputState;
+		Output.Sync = Input.Sync;
+
+		const float DeltaSeconds = TimeStep.StepMS.ToRealTimeSeconds();
 
 		// --------------------------------------------------------------
 		//	Rotation Update
@@ -162,26 +163,26 @@ namespace FlyingMovement
 		//	In this simulation, the rotation update isn't allowed to "fail". We don't expect the collision query to be able to fail the rotational update.
 		// --------------------------------------------------------------
 
-		OutputState.Rotation += (InputCmd.RotationInput * DeltaSeconds);
-		OutputState.Rotation.Normalize();
+		Output.Sync.Rotation += (Input.Cmd.RotationInput * DeltaSeconds);
+		Output.Sync.Rotation.Normalize();
 
-		const FQuat OutputQuat = OutputState.Rotation.Quaternion();
+		const FQuat OutputQuat = Output.Sync.Rotation.Quaternion();
 
 		// After the rotation is known, we need to sync our driver to this state. This is unfortunate and probably not best for perf, but since the Driver owned
 		// primitive component is ultimately what does our scene queries, we have no choice: we must get the primitive component in the state we say it should be in.
-		PreSimSync(OutputState);
+		PreSimSync(Output.Sync);
 
-		const FVector LocalSpaceMovementInput = OutputState.Rotation.RotateVector( InputCmd.MovementInput );
+		const FVector LocalSpaceMovementInput = Output.Sync.Rotation.RotateVector( Input.Cmd.MovementInput );
 	   	
 		// --------------------------------------------------------------
-		// Calculate OutputState.RelativeVelocity based on Input
+		// Calculate Output.Sync.RelativeVelocity based on Input
 		// --------------------------------------------------------------
 		{
 			const FVector ControlAcceleration = LocalSpaceMovementInput.GetClampedToMaxSize(1.f);
-			FVector Velocity = InputState.Velocity;
+			FVector Velocity = Input.Sync.Velocity;
 
 			const float AnalogInputModifier = (ControlAcceleration.SizeSquared() > 0.f ? ControlAcceleration.Size() : 0.f);
-			const float MaxPawnSpeed = AuxState.MaxSpeed * AnalogInputModifier;
+			const float MaxPawnSpeed = Input.Aux.MaxSpeed * AnalogInputModifier;
 			const bool bExceedingMaxSpeed = IsExceedingMaxSpeed(Velocity, MaxPawnSpeed);
 
 			if (AnalogInputModifier > 0.f && !bExceedingMaxSpeed)
@@ -190,7 +191,7 @@ namespace FlyingMovement
 				if (Velocity.SizeSquared() > 0.f)
 				{
 					// Change direction faster than only using acceleration, but never increase velocity magnitude.
-					const float TimeScale = FMath::Clamp(DeltaSeconds * AuxState.TurningBoost, 0.f, 1.f);
+					const float TimeScale = FMath::Clamp(DeltaSeconds * Input.Aux.TurningBoost, 0.f, 1.f);
 					Velocity = Velocity + (ControlAcceleration * Velocity.Size() - Velocity) * TimeScale;
 				}
 			}
@@ -200,7 +201,7 @@ namespace FlyingMovement
 				if (Velocity.SizeSquared() > 0.f)
 				{
 					const FVector OldVelocity = Velocity;
-					const float VelSize = FMath::Max(Velocity.Size() - FMath::Abs(AuxState.Deceleration) * DeltaSeconds, 0.f);
+					const float VelSize = FMath::Max(Velocity.Size() - FMath::Abs(Input.Aux.Deceleration) * DeltaSeconds, 0.f);
 					Velocity = Velocity.GetSafeNormal() * VelSize;
 
 					// Don't allow braking to lower us below max speed if we started above it.
@@ -213,15 +214,15 @@ namespace FlyingMovement
 
 			// Apply acceleration and clamp velocity magnitude.
 			const float NewMaxSpeed = (IsExceedingMaxSpeed(Velocity, MaxPawnSpeed)) ? Velocity.Size() : MaxPawnSpeed;
-			Velocity += ControlAcceleration * FMath::Abs(AuxState.Acceleration) * DeltaSeconds;
+			Velocity += ControlAcceleration * FMath::Abs(Input.Aux.Acceleration) * DeltaSeconds;
 			Velocity = Velocity.GetClampedToMaxSize(NewMaxSpeed);
 
 			// Finally, output velocity that we calculated
-			OutputState.Velocity = Velocity;
+			Output.Sync.Velocity = Velocity;
 		
 			if (FMovementSimulation::ForceMispredict)
 			{
-				OutputState.Velocity += ForceMispredictVelocityMagnitude;
+				Output.Sync.Velocity += ForceMispredictVelocityMagnitude;
 				ForceMispredict = false;
 			}
 		}
@@ -231,7 +232,7 @@ namespace FlyingMovement
 		// --------------------------------------------------------------
 		//	Calculate the final movement delta and move the update component
 		// --------------------------------------------------------------
-		FVector Delta = OutputState.Velocity * DeltaSeconds;
+		FVector Delta = Output.Sync.Velocity * DeltaSeconds;
 
 		if (!Delta.IsNearlyZero(1e-6f))
 		{
@@ -256,7 +257,7 @@ namespace FlyingMovement
 
 		// Finalize. This is unfortunate. The component mirrors our internal motion state and since we call into it to update, at this point, it has the real position.
 		const FTransform UpdateComponentTransform = GetUpdateComponentTransform();
-		OutputState.Location = UpdateComponentTransform.GetLocation();
+		Output.Sync.Location = UpdateComponentTransform.GetLocation();
 
 		// Note that we don't pull the rotation out of the final update transform. Converting back from a quat will lead to a different FRotator than what we are storing
 		// here in the simulation layer. This may not be the best choice for all movement simulations, but is ok for this one.
diff --git a/Engine/Plugins/Runtime/NetworkPrediction/Source/NetworkPrediction/Private/Movement/ParametricMovement.cpp b/Engine/Plugins/Runtime/NetworkPrediction/Source/NetworkPrediction/Private/Movement/ParametricMovement.cpp
--- a/Engine/Plugins/Runtime/NetworkPrediction/Source/NetworkPrediction/Private/Movement/ParametricMovement.cpp
+++ b/Engine/Plugins/Runtime/NetworkPrediction/Source/NetworkPrediction/Private/Movement/ParametricMovement.cpp
@@ -41,13 +41,14 @@ namespace ParametricMovement
 {
 	const FName FMovementSimulation::GroupName(TEXT("Parametric"));
 
-	void FMovementSimulation::Update(const float DeltaSeconds, const FInputCmd& InputCmd, const FMoveState& InputState, FMoveState& OutputState, const FAuxState& AuxState, const TNetSimLazyWriter<FAuxState>& OutAuxStateAccessor)
+	void FMovementSimulation::SimulationTick(const TNetSimTimeStep& TimeStep, const TNetSimInput<TMovementBufferTypes>& Input, const TNetSimOutput<TMovementBufferTypes>& Output)
 	{
 		check(Motion != nullptr); // Must set motion mapping prior to running the simulation
+		const float DeltaSeconds = TimeStep.StepMS.ToRealTimeSeconds();
 
 		// Advance parametric time. This won't always be linear: we could loop, rewind/bounce, etc
-		const float InputPlayRate = InputCmd.PlayRate.Get(InputState.PlayRate); // Returns InputCmds playrate if set, else returns previous state's playrate		
-		Motion->AdvanceParametricTime(InputState.Position, InputPlayRate, OutputState.Position, OutputState.PlayRate, DeltaSeconds);
+		const float InputPlayRate = Input.Cmd.PlayRate.Get(Input.Sync.PlayRate); // Returns InputCmds playrate if set, else returns previous state's playrate		
+		Motion->AdvanceParametricTime(Input.Sync.Position, InputPlayRate, Output.Sync.Position, Output.Sync.PlayRate, DeltaSeconds);
 
 		// We have our time that we should be at. We just need to move primitive component to that position.
 		// Again, note that we expect this cannot fail. We move like this so that it can push things, but we don't expect failure.
@@ -55,10 +56,10 @@ namespace ParametricMovement
 		// the output Position in the case where a Move is blocked (E.g, you move 50% towards the desired location)
 
 		FTransform StartingTransform;
-		Motion->MapTimeToTransform(InputState.Position, StartingTransform);
+		Motion->MapTimeToTransform(Input.Sync.Position, StartingTransform);
 
 		FTransform NewTransform;
-		Motion->MapTimeToTransform(OutputState.Position, NewTransform);
+		Motion->MapTimeToTransform(Output.Sync.Position, NewTransform);
 
 		FHitResult Hit(1.f);
 
diff --git a/Engine/Plugins/Runtime/NetworkPrediction/Source/NetworkPrediction/Private/NetworkPredictionTests.cpp b/Engine/Plugins/Runtime/NetworkPrediction/Source/NetworkPrediction/Private/NetworkPredictionTests.cpp
--- a/Engine/Plugins/Runtime/NetworkPrediction/Source/NetworkPrediction/Private/NetworkPredictionTests.cpp
+++ b/Engine/Plugins/Runtime/NetworkPrediction/Source/NetworkPrediction/Private/NetworkPredictionTests.cpp
@@ -16,10 +16,10 @@ void TickSettingsTest(TTickSettings& Settings)
 {
 	using TBufferTypes = TInternalBufferTypes<FDummyBufferTypes, TTickSettings>;
 
-	TSimulationTickState<TTickSettings>	TickState;
+	TSimulationTicker<TTickSettings> Ticker;
 
 	float RealTime = FMath::Rand();
-	TickState.GiveSimulationTime(RealTime);
+	Ticker.GiveSimulationTime(RealTime);
 
 }
 
diff --git a/Engine/Plugins/Runtime/NetworkPrediction/Source/NetworkPrediction/Public/Movement/FlyingMovement.h b/Engine/Plugins/Runtime/NetworkPrediction/Source/NetworkPrediction/Public/Movement/FlyingMovement.h
--- a/Engine/Plugins/Runtime/NetworkPrediction/Source/NetworkPrediction/Public/Movement/FlyingMovement.h
+++ b/Engine/Plugins/Runtime/NetworkPrediction/Source/NetworkPrediction/Public/Movement/FlyingMovement.h
@@ -155,7 +155,7 @@ namespace FlyingMovement
 		static bool ForceMispredict;
 
 		/** Main update function */
-		NETWORKPREDICTION_API void Update(const float DeltaSeconds, const FInputCmd& InputCmd, const FMoveState& InputState, FMoveState& OutputState, const FAuxState& AuxState, const TNetSimLazyWriter<FAuxState>& OutAuxStateAccessor);
+		NETWORKPREDICTION_API void SimulationTick(const TNetSimTimeStep& TimeStep, const TNetSimInput<TMovementBufferTypes>& Input, const TNetSimOutput<TMovementBufferTypes>& Output);
 
 		// Called prior to running the sim to make sure to make sure the collision component is in the right place. 
 		// This is unfortunate and not good, but is needed to ensure our collision and world position have not been moved out from under us.
diff --git a/Engine/Plugins/Runtime/NetworkPrediction/Source/NetworkPrediction/Public/Movement/ParametricMovement.h b/Engine/Plugins/Runtime/NetworkPrediction/Source/NetworkPrediction/Public/Movement/ParametricMovement.h
--- a/Engine/Plugins/Runtime/NetworkPrediction/Source/NetworkPrediction/Public/Movement/ParametricMovement.h
+++ b/Engine/Plugins/Runtime/NetworkPrediction/Source/NetworkPrediction/Public/Movement/ParametricMovement.h
@@ -112,7 +112,7 @@ namespace ParametricMovement
 	public:
 		static const FName GroupName;
 
-		void Update(const float DeltaSeconds, const FInputCmd& InputCmd, const FMoveState& InputState, FMoveState& OutputState, const FAuxState& AuxState, const TNetSimLazyWriter<FAuxState>& OutAuxStateAccessor);
+		void SimulationTick(const TNetSimTimeStep& TimeStep, const TNetSimInput<TMovementBufferTypes>& Input, const TNetSimOutput<TMovementBufferTypes>& Output);
 
 		// Pointer to our static mapping of time->position
 		const FSimpleParametricMotion* Motion = nullptr;
diff --git a/Engine/Plugins/Runtime/NetworkPrediction/Source/NetworkPrediction/Public/NetworkSimulationModel.h b/Engine/Plugins/Runtime/NetworkPrediction/Source/NetworkPrediction/Public/NetworkSimulationModel.h
--- a/Engine/Plugins/Runtime/NetworkPrediction/Source/NetworkPrediction/Public/NetworkSimulationModel.h
+++ b/Engine/Plugins/Runtime/NetworkPrediction/Source/NetworkPrediction/Public/NetworkSimulationModel.h
@@ -67,7 +67,7 @@ public:
 		Driver = InDriver;
 		*Buffers.Sync.WriteKeyframe(0) = InitialSyncState;
 		*Buffers.Aux.WriteKeyframe(0) = InitialAuxState;
-		TickInfo.SetTotalProcessedSimulationTime(FNetworkSimTime(), 0);
+		Ticker.SetTotalProcessedSimulationTime(FNetworkSimTime(), 0);
 		DO_NETSIM_MODEL_DEBUG(FNetworkSimulationModelDebuggerManager::Get().RegisterNetworkSimulationModel(this, Driver->GetVLogOwner()));
 	}
 
@@ -123,34 +123,34 @@ public:
 		switch (Parameters.Role)
 		{
 			case ROLE_Authority:
-				RepProxy_ServerRPC.template PreSimTick<TDriver>(Driver, Buffers, TickInfo, Parameters);
+				RepProxy_ServerRPC.template PreSimTick<TDriver>(Driver, Buffers, Ticker, Parameters);
 			break;
 
 			case ROLE_AutonomousProxy:
-				RepProxy_Autonomous.template PreSimTick<TDriver>(Driver, Buffers, TickInfo, Parameters);
+				RepProxy_Autonomous.template PreSimTick<TDriver>(Driver, Buffers, Ticker, Parameters);
 			break;
 
 			case ROLE_SimulatedProxy:
-				RepProxy_Simulated.template PreSimTick<TDriver>(Driver, Buffers, TickInfo, Parameters);
+				RepProxy_Simulated.template PreSimTick<TDriver>(Driver, Buffers, Ticker, Parameters);
 			break;
 		}
 
 		// -------------------------------------------------------------------------------------------------------------------------------------------------
 		//												Input Processing & Simulation Update
 		// -------------------------------------------------------------------------------------------------------------------------------------------------
-		while (TickInfo.PendingKeyframe <= TickInfo.MaxAllowedKeyframe)
+		while (Ticker.PendingKeyframe <= Ticker.MaxAllowedKeyframe)
 		{
-			const int32 InputKeyframe = TickInfo.PendingKeyframe;
-			const int32 OutputKeyframe = TickInfo.PendingKeyframe + 1;
+			const int32 InputKeyframe = Ticker.PendingKeyframe;
+			const int32 OutputKeyframe = Ticker.PendingKeyframe + 1;
 
 			const TInputCmd* InputCmd = Buffers.Input[InputKeyframe];
-			if (!ensureMsgf(InputCmd, TEXT("No InputCmd available for Keyframe %d. PendingKeyframe: %d. MaxAllowedKeyframe: %d."), InputKeyframe, TickInfo.PendingKeyframe, TickInfo.MaxAllowedKeyframe))
+			if (!ensureMsgf(InputCmd, TEXT("No InputCmd available for Keyframe %d. PendingKeyframe: %d. MaxAllowedKeyframe: %d."), InputKeyframe, Ticker.PendingKeyframe, Ticker.MaxAllowedKeyframe))
 			{
 				break;
 			}
 
 			// We have an unprocessed command, do we have enough allotted simulation time to process it?
-			if (TickInfo.GetRemainingAllowedSimulationTime() >= InputCmd->GetFrameDeltaTime())
+			if (Ticker.GetRemainingAllowedSimulationTime() >= InputCmd->GetFrameDeltaTime())
 			{
 				TSyncState* InSyncState = Buffers.Sync[InputKeyframe];
 				if (InSyncState == nullptr)
@@ -162,14 +162,19 @@ public:
 				}
 
 				const TAuxState* InAuxState = Buffers.Aux[InputKeyframe];
-				checkf(InAuxState, TEXT("No AuxState available for Keyframe %d. PendingKeyframe: %d. MaxAllowedKeyframe: %d."), InputKeyframe, TickInfo.PendingKeyframe, TickInfo.MaxAllowedKeyframe);
+				checkf(InAuxState, TEXT("No AuxState available for Keyframe %d. PendingKeyframe: %d. MaxAllowedKeyframe: %d."), InputKeyframe, Ticker.PendingKeyframe, Ticker.MaxAllowedKeyframe);
 
 				TSyncState* OutSyncState = Buffers.Sync.WriteKeyframe(OutputKeyframe);
 				check(OutSyncState);
 
 				{
-					TScopedSimulationTick UpdateScope(TickInfo, OutputKeyframe, InputCmd->GetFrameDeltaTime());
-					Simulation->Update(InputCmd->GetFrameDeltaTime().ToRealTimeSeconds(), *InputCmd, *InSyncState, *OutSyncState, *InAuxState, Buffers.Aux.LazyWriter(OutputKeyframe));
+					TScopedSimulationTick UpdateScope(Ticker, OutputKeyframe, InputCmd->GetFrameDeltaTime());
+					
+					Simulation->SimulationTick( 
+						{ InputCmd->GetFrameDeltaTime(), Ticker },
+						{ *InputCmd, *InSyncState, *InAuxState },
+						{ *OutSyncState, Buffers.Aux.LazyWriter(OutputKeyframe) } );
+													
 				}
 
 				if (DebugState)
@@ -190,15 +195,15 @@ public:
 		switch (Parameters.Role)
 		{
 			case ROLE_Authority:
-				RepProxy_ServerRPC.template PostSimTick<TDriver>(Driver, Buffers, TickInfo, Parameters);
+				RepProxy_ServerRPC.template PostSimTick<TDriver>(Driver, Buffers, Ticker, Parameters);
 			break;
 
 			case ROLE_AutonomousProxy:
-				RepProxy_Autonomous.template PostSimTick<TDriver>(Driver, Buffers, TickInfo, Parameters);
+				RepProxy_Autonomous.template PostSimTick<TDriver>(Driver, Buffers, Ticker, Parameters);
 			break;
 
 			case ROLE_SimulatedProxy:
-				RepProxy_Simulated.template PostSimTick<TDriver>(Driver, Buffers, TickInfo, Parameters);
+				RepProxy_Simulated.template PostSimTick<TDriver>(Driver, Buffers, Ticker, Parameters);
 			break;
 		}
 
@@ -209,9 +214,9 @@ public:
 		// Finish debug state buffer recording (what the server processed each frame)
 		if (DebugState)
 		{
-			DebugState->PendingKeyframe = TickInfo.PendingKeyframe;
+			DebugState->PendingKeyframe = Ticker.PendingKeyframe;
 			DebugState->HeadKeyframe = Buffers.Input.HeadKeyframe();
-			DebugState->RemainingAllowedSimulationTimeSeconds = (float)TickInfo.GetRemainingAllowedSimulationTime().ToRealTimeSeconds();
+			DebugState->RemainingAllowedSimulationTimeSeconds = (float)Ticker.GetRemainingAllowedSimulationTime().ToRealTimeSeconds();
 		}
 
 		// Historical data recording (longer buffers for historical reference)
@@ -233,15 +238,15 @@ public:
 		switch (Role)
 		{
 			case ROLE_Authority:
-				RepProxy_ServerRPC.template Reconcile<TSimulation, TDriver>(Simulation, Driver, Buffers, TickInfo);
+				RepProxy_ServerRPC.template Reconcile<TSimulation, TDriver>(Simulation, Driver, Buffers, Ticker);
 			break;
 
 			case ROLE_AutonomousProxy:
-				RepProxy_Autonomous.template Reconcile<TSimulation, TDriver>(Simulation, Driver, Buffers, TickInfo);
+				RepProxy_Autonomous.template Reconcile<TSimulation, TDriver>(Simulation, Driver, Buffers, Ticker);
 			break;
 
 			case ROLE_SimulatedProxy:
-				RepProxy_Simulated.template Reconcile<TSimulation, TDriver>(Simulation, Driver, Buffers, TickInfo);
+				RepProxy_Simulated.template Reconcile<TSimulation, TDriver>(Simulation, Driver, Buffers, Ticker);
 			break;
 		}
 	}
@@ -266,7 +271,7 @@ public:
 			//MyHistoricBuffers->Aux.SetBufferSize(Parameters.HistoricBufferSize); AUXFIXME
 		}
 
-		//TickInfo.InitSimulationTimeBuffer(Parameters.SyncedBufferSize);
+		//Ticker.InitSimulationTimeBuffer(Parameters.SyncedBufferSize);
 	}
 
 	void NetSerializeProxy(EReplicationProxyTarget Target, const FNetSerializeParams& Params) final override
@@ -274,20 +279,20 @@ public:
 		switch(Target)
 		{
 		case EReplicationProxyTarget::ServerRPC:
-			RepProxy_ServerRPC.NetSerialize(Params, Buffers, TickInfo);
+			RepProxy_ServerRPC.NetSerialize(Params, Buffers, Ticker);
 			break;
 		case EReplicationProxyTarget::AutonomousProxy:
-			RepProxy_Autonomous.NetSerialize(Params, Buffers, TickInfo);
+			RepProxy_Autonomous.NetSerialize(Params, Buffers, Ticker);
 			break;
 		case EReplicationProxyTarget::SimulatedProxy:
-			RepProxy_Simulated.NetSerialize(Params, Buffers, TickInfo);
+			RepProxy_Simulated.NetSerialize(Params, Buffers, Ticker);
 			break;
 		case EReplicationProxyTarget::Replay:
-			RepProxy_Replay.NetSerialize(Params, Buffers, TickInfo);
+			RepProxy_Replay.NetSerialize(Params, Buffers, Ticker);
 			break;
 		case EReplicationProxyTarget::Debug:
 #if NETSIM_MODEL_DEBUG
-			RepProxy_Debug.NetSerialize(Params, Buffers, TickInfo);
+			RepProxy_Debug.NetSerialize(Params, Buffers, Ticker);
 			break;
 #endif
 		default:
@@ -362,12 +367,12 @@ public:
 
 	void BeginRollback(const FNetworkSimTime& RollbackDeltaTime, const int32 ParentKeyframe) final override
 	{
-		RepProxy_Simulated.template DependentRollbackBegin<TSimulation, TDriver>(Simulation, Driver, Buffers, TickInfo, RollbackDeltaTime, ParentKeyframe);
+		RepProxy_Simulated.template DependentRollbackBegin<TSimulation, TDriver>(Simulation, Driver, Buffers, Ticker, RollbackDeltaTime, ParentKeyframe);
 	}
 
 	void StepRollback(const FNetworkSimTime& Step, const int32 ParentKeyframe, const bool bFinalStep) final override
 	{
-		RepProxy_Simulated.template DependentRollbackStep<TSimulation, TDriver>(Simulation, Driver, Buffers, TickInfo, Step, ParentKeyframe, bFinalStep);
+		RepProxy_Simulated.template DependentRollbackStep<TSimulation, TDriver>(Simulation, Driver, Buffers, Ticker, Step, ParentKeyframe, bFinalStep);
 	}
 
 	void ClearAllDependentSimulations()
@@ -384,7 +389,7 @@ public:
 	TSimulation* Simulation = nullptr;
 	TDriver* Driver = nullptr;
 
-	TSimulationTickState<TTickSettings> TickInfo;	// Manages simulation time and what inputs we are processed
+	TSimulationTicker<TTickSettings> Ticker;
 
 	TNetworkSimBufferContainer<TBufferTypes> Buffers;
 
diff --git a/Engine/Plugins/Runtime/NetworkPrediction/Source/NetworkPrediction/Public/NetworkSimulationModelDebugger.h b/Engine/Plugins/Runtime/NetworkPrediction/Source/NetworkPrediction/Public/NetworkSimulationModelDebugger.h
--- a/Engine/Plugins/Runtime/NetworkPrediction/Source/NetworkPrediction/Public/NetworkSimulationModelDebugger.h
+++ b/Engine/Plugins/Runtime/NetworkPrediction/Source/NetworkPrediction/Public/NetworkSimulationModelDebugger.h
@@ -534,7 +534,7 @@ struct TNetworkSimulationModelDebugger : public INetworkSimulationModelDebugger
 		// ------------------------------------------------------------------------------------------------------------------------------------------------
 
 		Out.Emit(FString::Printf(TEXT("%s - %s"), *Owner->GetName(), *UEnum::GetValueAsString(TEXT("Engine.ENetRole"), Owner->GetLocalRole())), FColor::Yellow);
-		Out.Emit(FString::Printf(TEXT("PendingKeyframe: %d (%d Buffered)"), NetworkSim->TickInfo.PendingKeyframe, NetworkSim->Buffers.Input.HeadKeyframe() - NetworkSim->TickInfo.PendingKeyframe));
+		Out.Emit(FString::Printf(TEXT("PendingKeyframe: %d (%d Buffered)"), NetworkSim->Ticker.PendingKeyframe, NetworkSim->Buffers.Input.HeadKeyframe() - NetworkSim->Ticker.PendingKeyframe));
 				
 		if (Owner->GetLocalRole() == ROLE_AutonomousProxy)
 		{			
@@ -563,17 +563,17 @@ struct TNetworkSimulationModelDebugger : public INetworkSimulationModelDebugger
 
 			Out.Emit(*ConfirmedFrameStr, Color);
 
-			FString SimulationTimeString = FString::Printf(TEXT("Local SimulationTime: %s. SerializedSimulationTime: %s. Difference MS: %s"), *NetworkSim->TickInfo.GetTotalProcessedSimulationTime().ToString(),
-				*NetworkSim->RepProxy_Autonomous.GetLastSerializedSimTime().ToString(), *(NetworkSim->TickInfo.GetTotalProcessedSimulationTime() - NetworkSim->RepProxy_Autonomous.GetLastSerializedSimTime()).ToString());
+			FString SimulationTimeString = FString::Printf(TEXT("Local SimulationTime: %s. SerializedSimulationTime: %s. Difference MS: %s"), *NetworkSim->Ticker.GetTotalProcessedSimulationTime().ToString(),
+				*NetworkSim->RepProxy_Autonomous.GetLastSerializedSimTime().ToString(), *(NetworkSim->Ticker.GetTotalProcessedSimulationTime() - NetworkSim->RepProxy_Autonomous.GetLastSerializedSimTime()).ToString());
 			Out.Emit(*SimulationTimeString, Color);
 
-			FString AllowedSimulationTimeString = FString::Printf(TEXT("Allowed Simulation Time: %s. Keyframe: %d/%d/%d"), *NetworkSim->TickInfo.GetRemainingAllowedSimulationTime().ToString(), NetworkSim->TickInfo.MaxAllowedKeyframe, NetworkSim->TickInfo.PendingKeyframe, NetworkSim->Buffers.Input.HeadKeyframe());
+			FString AllowedSimulationTimeString = FString::Printf(TEXT("Allowed Simulation Time: %s. Keyframe: %d/%d/%d"), *NetworkSim->Ticker.GetRemainingAllowedSimulationTime().ToString(), NetworkSim->Ticker.MaxAllowedKeyframe, NetworkSim->Ticker.PendingKeyframe, NetworkSim->Buffers.Input.HeadKeyframe());
 			Out.Emit(*AllowedSimulationTimeString, Color);
 		}
 		else if (Owner->GetLocalRole() == ROLE_SimulatedProxy)
 		{
 			FColor Color = FColor::White;
-			FString TimeString = FString::Printf(TEXT("Total Processed Simulation Time: %s. Last Serialized Simulation Time: %s. Delta: %s"), *NetworkSim->TickInfo.GetTotalProcessedSimulationTime().ToString(), *NetworkSim->RepProxy_Simulated.GetLastSerializedSimulationTime().ToString(), *(NetworkSim->RepProxy_Simulated.GetLastSerializedSimulationTime() - NetworkSim->TickInfo.GetTotalProcessedSimulationTime()).ToString());
+			FString TimeString = FString::Printf(TEXT("Total Processed Simulation Time: %s. Last Serialized Simulation Time: %s. Delta: %s"), *NetworkSim->Ticker.GetTotalProcessedSimulationTime().ToString(), *NetworkSim->RepProxy_Simulated.GetLastSerializedSimulationTime().ToString(), *(NetworkSim->RepProxy_Simulated.GetLastSerializedSimulationTime() - NetworkSim->Ticker.GetTotalProcessedSimulationTime()).ToString());
 			Out.Emit(*TimeString, Color);
 		}
 
@@ -702,7 +702,7 @@ struct TNetworkSimulationModelDebugger : public INetworkSimulationModelDebugger
 				FVector2D ServerSimulationTimeData(Owner->GetWorld()->GetTimeSeconds(), NetworkSim->RepProxy_Simulated.GetLastSerializedSimulationTime().ToRealTimeMS());
 				UE_VLOG_HISTOGRAM(Owner, LogNetworkSimDebug, Log, "Simulated Time Graph", "Serialized Simulation Time", ServerSimulationTimeData);
 
-				FVector2D LocalSimulationTimeData(Owner->GetWorld()->GetTimeSeconds(), NetworkSim->TickInfo.GetTotalProcessedSimulationTime().ToRealTimeMS());
+				FVector2D LocalSimulationTimeData(Owner->GetWorld()->GetTimeSeconds(), NetworkSim->Ticker.GetTotalProcessedSimulationTime().ToRealTimeMS());
 				UE_VLOG_HISTOGRAM(Owner, LogNetworkSimDebug, Log, "Simulated Time Graph", "Local Simulation Time", LocalSimulationTimeData);
 			
 			}
diff --git a/Engine/Plugins/Runtime/NetworkPrediction/Source/NetworkPrediction/Public/NetworkSimulationModelInterpolator.h b/Engine/Plugins/Runtime/NetworkPrediction/Source/NetworkPrediction/Public/NetworkSimulationModelInterpolator.h
--- a/Engine/Plugins/Runtime/NetworkPrediction/Source/NetworkPrediction/Public/NetworkSimulationModelInterpolator.h
+++ b/Engine/Plugins/Runtime/NetworkPrediction/Source/NetworkPrediction/Public/NetworkSimulationModelInterpolator.h
@@ -34,7 +34,7 @@ struct TInterpolator
 	bool bEnableVisualLog = true;
 
 	template<typename TSystemDriver>
-	void PostSimTick(TSystemDriver* Driver, const TNetworkSimBufferContainer<TBufferTypes>& Buffers, const TSimulationTickState<TTickSettings>& TickInfo, const FNetSimTickParameters& TickParameters)
+	void PostSimTick(TSystemDriver* Driver, const TNetworkSimBufferContainer<TBufferTypes>& Buffers, const FSimulationTickState& TickInfo, const FNetSimTickParameters& TickParameters)
 	{
 		const bool bDoVLog = NetworkInterpolationDebugCVars::VLog() && bEnableVisualLog;
 		const float DeltaSeconds = TickParameters.LocalDeltaTimeSeconds;
diff --git a/Engine/Plugins/Runtime/NetworkPrediction/Source/NetworkPrediction/Public/NetworkSimulationModelReplicators.h b/Engine/Plugins/Runtime/NetworkPrediction/Source/NetworkPrediction/Public/NetworkSimulationModelReplicators.h
--- a/Engine/Plugins/Runtime/NetworkPrediction/Source/NetworkPrediction/Public/NetworkSimulationModelReplicators.h
+++ b/Engine/Plugins/Runtime/NetworkPrediction/Source/NetworkPrediction/Public/NetworkSimulationModelReplicators.h
@@ -68,19 +68,19 @@ struct TReplicatorEmpty
 	int32 GetProxyDirtyCount(TNetworkSimBufferContainer<TBufferTypes>& Buffers) const { return 0; }
 
 	// NetSerialize: just serialize the network data. Don't run simulation steps. Every replicator will be NetSerialized before moving on to Reconcile phase.
-	void NetSerialize(const FNetSerializeParams& P, TNetworkSimBufferContainer<TBufferTypes>& Buffers, const TSimulationTickState<TTickSettings>& TickInfo) { }
+	void NetSerialize(const FNetSerializeParams& P, TNetworkSimBufferContainer<TBufferTypes>& Buffers, const TSimulationTicker<TTickSettings>& Ticker) { }
 
 	// Reconcile: called after everyone has NetSerialized. "Get right with the server": this function is about reconciling what the server told you vs what you have extrapolated or forward predicted locally
 	template<typename TSimulation, typename TDriver>
-	void Reconcile(TSimulation* Simulation, TDriver* Driver, TNetworkSimBufferContainer<TBufferTypes>& Buffers, TSimulationTickState<TTickSettings>& TickInfo) { }
+	void Reconcile(TSimulation* Simulation, TDriver* Driver, TNetworkSimBufferContainer<TBufferTypes>& Buffers, TSimulationTicker<TTickSettings>& Ticker) { }
 
-	// Called prior to input processing. This function must updated TickInfo to allow simulation time (from TickParameters) and to possibly get new input.
+	// Called prior to input processing. This function must update Ticker to allow simulation time (from TickParameters) and to possibly get new input.
 	template<typename TDriver>
-	void PreSimTick(TDriver* Driver, TNetworkSimBufferContainer<TBufferTypes>& Buffers, TSimulationTickState<TTickSettings>& TickInfo, const FNetSimTickParameters& TickParameters) { }
+	void PreSimTick(TDriver* Driver, TNetworkSimBufferContainer<TBufferTypes>& Buffers, TSimulationTicker<TTickSettings>& Ticker, const FNetSimTickParameters& TickParameters) { }
 
 	// Called after input processing. Should finalize the frame and do any smoothing/interpolation. This function is not allowed to modify the buffers or tick state, or even call the simulation/Update function.
 	template<typename TDriver>
-	void PostSimTick(TDriver* Driver, const TNetworkSimBufferContainer<TBufferTypes>& Buffers, const TSimulationTickState<TTickSettings>& TickInfo, const FNetSimTickParameters& TickParameters) { }
+	void PostSimTick(TDriver* Driver, const TNetworkSimBufferContainer<TBufferTypes>& Buffers, const TSimulationTicker<TTickSettings>& Ticker, const FNetSimTickParameters& TickParameters) { }
 };
 
 // This is the "templated base class" for replicators but is not required (i.e., this not an official interface used by TNetworkedSimulation model. Just a base implementation you can start with)
@@ -93,23 +93,23 @@ struct TReplicatorBase
 	int32 GetProxyDirtyCount(TNetworkSimBufferContainer<TBufferTypes>& Buffers) const { return 0; }
 
 	// NetSerialize: just serialize the network data. Don't run simulation steps. Every replicator will be NetSerialized before moving on to Reconcile phase.
-	void NetSerialize(const FNetSerializeParams& P, TNetworkSimBufferContainer<TBufferTypes>& Buffers, const TSimulationTickState<TTickSettings>& TickInfo) { }
+	void NetSerialize(const FNetSerializeParams& P, TNetworkSimBufferContainer<TBufferTypes>& Buffers, const TSimulationTicker<TTickSettings>& Ticker) { }
 
 	// Reconcile: called after everyone has NetSerialized. "Get right with the server": this function is about reconciling what the server told you vs what you have extrapolated or forward predicted locally
 	template<typename TSimulation, typename TDriver>
-	void Reconcile(TSimulation* Simulation, TDriver* Driver, TNetworkSimBufferContainer<TBufferTypes>& Buffers, TSimulationTickState<TTickSettings>& TickInfo) { }
+	void Reconcile(TSimulation* Simulation, TDriver* Driver, TNetworkSimBufferContainer<TBufferTypes>& Buffers, TSimulationTicker<TTickSettings>& Ticker) { }
 
-	// Called prior to input processing. This function must updated TickInfo to allow simulation time (from TickParameters) and to possibly get new input.
+	// Called prior to input processing. This function must updated Ticker to allow simulation time (from TickParameters) and to possibly get new input.
 	template<typename TDriver>
-	void PreSimTick(TDriver* Driver, TNetworkSimBufferContainer<TBufferTypes>& Buffers, TSimulationTickState<TTickSettings>& TickInfo, const FNetSimTickParameters& TickParameters)
+	void PreSimTick(TDriver* Driver, TNetworkSimBufferContainer<TBufferTypes>& Buffers, TSimulationTicker<TTickSettings>& Ticker, const FNetSimTickParameters& TickParameters)
 	{
-		// Accumulate local delta time into TickInfo
-		TickInfo.GiveSimulationTime(TickParameters.LocalDeltaTimeSeconds);
+		// Accumulate local delta time into Ticker
+		Ticker.GiveSimulationTime(TickParameters.LocalDeltaTimeSeconds);
 	}
 
 	// Called after input processing. Should finalize the frame and do any smoothing/interpolation. This function is not allowed to modify the buffers or tick state, or even call the simulation/Update function.
 	template<typename TDriver>
-	void PostSimTick(TDriver* Driver, const TNetworkSimBufferContainer<TBufferTypes>& Buffers, const TSimulationTickState<TTickSettings>& TickInfo, const FNetSimTickParameters& TickParameters)
+	void PostSimTick(TDriver* Driver, const TNetworkSimBufferContainer<TBufferTypes>& Buffers, const TSimulationTicker<TTickSettings>& Ticker, const FNetSimTickParameters& TickParameters)
 	{
 		// Sync to latest frame if there is any
 		if (Buffers.Sync.Num() > 0)
@@ -127,10 +127,10 @@ struct TReplicatorBase
 template<typename TBufferTypes, typename TTickSettings, typename TBase=TReplicatorBase<TBufferTypes, TTickSettings>, bool Enabled=true>
 struct TReplicator_SimTime : public TBase
 {
-	void NetSerialize(const FNetSerializeParams& P, TNetworkSimBufferContainer<TBufferTypes>& Buffers, TSimulationTickState<TTickSettings>& TickInfo)
+	void NetSerialize(const FNetSerializeParams& P, TNetworkSimBufferContainer<TBufferTypes>& Buffers, TSimulationTicker<TTickSettings>& Ticker)
 	{
-		TBase::NetSerialize(P, Buffers, TickInfo);
-		SerializedTime = TickInfo.GetTotalProcessedSimulationTime();
+		TBase::NetSerialize(P, Buffers, Ticker);
+		SerializedTime = Ticker.GetTotalProcessedSimulationTime();
 		SerializedTime.NetSerialize(P.Ar);
 	}
 
@@ -156,9 +156,9 @@ struct TReplicator_Sequence : public TBase
 		return Buffers.template Get<BufferId>().GetDirtyCount() ^ (TBase::GetProxyDirtyCount(Buffers) << 2); 
 	}
 
-	void NetSerialize(const FNetSerializeParams& P, TNetworkSimBufferContainer<TBufferTypes>& Buffers, TSimulationTickState<TTickSettings>& TickInfo)
+	void NetSerialize(const FNetSerializeParams& P, TNetworkSimBufferContainer<TBufferTypes>& Buffers, TSimulationTicker<TTickSettings>& Ticker)
 	{
-		TBase::NetSerialize(P, Buffers, TickInfo);
+		TBase::NetSerialize(P, Buffers, Ticker);
 
 		auto& Buffer = Buffers.template Get<BufferId>();
 		FArchive& Ar = P.Ar;
@@ -209,9 +209,9 @@ struct TReplicator_Single : public TBase
 		return Buffers.template Get<BufferId>().GetDirtyCount() ^ (TBase::GetProxyDirtyCount(Buffers) << 2); ; 
 	}
 
-	void NetSerialize(const FNetSerializeParams& P, TNetworkSimBufferContainer<TBufferTypes>& Buffers, TSimulationTickState<TTickSettings>& TickInfo)
+	void NetSerialize(const FNetSerializeParams& P, TNetworkSimBufferContainer<TBufferTypes>& Buffers, TSimulationTicker<TTickSettings>& Ticker)
 	{
-		TBase::NetSerialize(P, Buffers, TickInfo);
+		TBase::NetSerialize(P, Buffers, Ticker);
 
 		FArchive& Ar = P.Ar;
 
@@ -252,15 +252,15 @@ void GenerateLocalInputCmdAtKeyframe(TDriver* Driver, TNetworkSimBufferContainer
 
 /** Helper to generate a local input cmd if we have simulation time to spend and advance the simulation's MaxAllowedKeyframe so that it can be processed. */
 template<typename TDriver, typename TBufferTypes, typename TTickSettings>
-void TryGenerateLocalInput(TDriver* Driver, TNetworkSimBufferContainer<TBufferTypes>& Buffers, TSimulationTickState<TTickSettings>& TickInfo)
+void TryGenerateLocalInput(TDriver* Driver, TNetworkSimBufferContainer<TBufferTypes>& Buffers, TSimulationTicker<TTickSettings>& Ticker)
 {
 	using TInputCmd = typename TBufferTypes::TInputCmd;
 
-	FNetworkSimTime DeltaSimTime = TickInfo.GetRemainingAllowedSimulationTime();
+	FNetworkSimTime DeltaSimTime = Ticker.GetRemainingAllowedSimulationTime();
 	if (DeltaSimTime.IsPositive())
 	{
-		GenerateLocalInputCmdAtKeyframe(Driver, Buffers, DeltaSimTime, TickInfo.PendingKeyframe);
-		TickInfo.MaxAllowedKeyframe = TickInfo.PendingKeyframe;
+		GenerateLocalInputCmdAtKeyframe(Driver, Buffers, DeltaSimTime, Ticker.PendingKeyframe);
+		Ticker.MaxAllowedKeyframe = Ticker.PendingKeyframe;
 	}
 }
 
@@ -275,28 +275,28 @@ struct TReplicator_Server : public TBase
 	using TInputCmd = typename TBufferTypes::TInputCmd;
 
 	template<typename TSimulation, typename TDriver>
-	void Reconcile(TSimulation* Simulation, TDriver* Driver, TNetworkSimBufferContainer<TBufferTypes>& Buffers, TSimulationTickState<TTickSettings>& TickInfo)
+	void Reconcile(TSimulation* Simulation, TDriver* Driver, TNetworkSimBufferContainer<TBufferTypes>& Buffers, TSimulationTicker<TTickSettings>& Ticker)
 	{
 		// After receiving input, server may process up to the latest received frames.
 		// (If we needed to buffer input server side for whatever reason, we would do it here)
 		// (also note that we will implicitly guard against speed hacks in the core update loop by not processing cmds past what we have been "allowed")
-		TickInfo.MaxAllowedKeyframe = Buffers.Input.HeadKeyframe();
+		Ticker.MaxAllowedKeyframe = Buffers.Input.HeadKeyframe();
 
 		// Check for gaps in commands
-		if ( TickInfo.PendingKeyframe < Buffers.Input.TailKeyframe() )
+		if ( Ticker.PendingKeyframe < Buffers.Input.TailKeyframe() )
 		{
-			UE_LOG(LogNetworkSim, Warning, TEXT("TReplicator_Server::Reconcile missing inputcmds. PendingKeyframe: %d. %s. This can happen via packet loss"), TickInfo.PendingKeyframe, *Buffers.Input.GetBasicDebugStr());
-			TickInfo.PendingKeyframe = Buffers.Input.TailKeyframe();
+			UE_LOG(LogNetworkSim, Warning, TEXT("TReplicator_Server::Reconcile missing inputcmds. PendingKeyframe: %d. %s. This can happen via packet loss"), Ticker.PendingKeyframe, *Buffers.Input.GetBasicDebugStr());
+			Ticker.PendingKeyframe = Buffers.Input.TailKeyframe();
 		}
 	}
 
 	template<typename T, typename TDriver>
-	void PreSimTick(TDriver* Driver, TNetworkSimBufferContainer<TBufferTypes>& Buffers, TSimulationTickState<TTickSettings>& TickInfo, const FNetSimTickParameters& TickParameters)
+	void PreSimTick(TDriver* Driver, TNetworkSimBufferContainer<TBufferTypes>& Buffers, TSimulationTicker<TTickSettings>& Ticker, const FNetSimTickParameters& TickParameters)
 	{
-		TickInfo.GiveSimulationTime(TickParameters.LocalDeltaTimeSeconds);
+		Ticker.GiveSimulationTime(TickParameters.LocalDeltaTimeSeconds);
 		if (TickParameters.bGenerateLocalInputCmds)
 		{
-			TryGenerateLocalInput(Driver, Buffers, TickInfo);
+			TryGenerateLocalInput(Driver, Buffers, Ticker);
 		}
 	}
 };
@@ -341,14 +341,14 @@ struct TReplicator_Simulated : public TBase
 		return Buffers.Sync.GetDirtyCount() ^ (TBase::GetProxyDirtyCount(Buffers) << 2); 
 	}
 
-	void NetSerialize(const FNetSerializeParams& P, TNetworkSimBufferContainer<TBufferTypes>& Buffers, TSimulationTickState<TTickSettings>& TickInfo)
+	void NetSerialize(const FNetSerializeParams& P, TNetworkSimBufferContainer<TBufferTypes>& Buffers, TSimulationTicker<TTickSettings>& Ticker)
 	{
 		FArchive& Ar = P.Ar;
 
 		FNetworkSimTime PrevLastSerializedSimulationTime = LastSerializedSimulationTime;
 
 		// Serialize latest simulation time
-		LastSerializedSimulationTime = TickInfo.GetTotalProcessedSimulationTime();
+		LastSerializedSimulationTime = Ticker.GetTotalProcessedSimulationTime();
 		LastSerializedSimulationTime.NetSerialize(P.Ar);
 
 		// Serialize latest element
@@ -363,28 +363,28 @@ struct TReplicator_Simulated : public TBase
 		}
 		else
 		{
-			check(TickInfo.SimulationTimeBuffer.HeadKeyframe() == Buffers.Sync.HeadKeyframe());
-			check(TickInfo.GetTotalProcessedSimulationTime() <= TickInfo.GetTotalAllowedSimulationTime());
+			check(Ticker.SimulationTimeBuffer.HeadKeyframe() == Buffers.Sync.HeadKeyframe());
+			check(Ticker.GetTotalProcessedSimulationTime() <= Ticker.GetTotalAllowedSimulationTime());
 
 			// Cache off our "starting" time before possibly overwriting it. We will use this in Reconcile to catch back up in some cases.
-			if (TickInfo.GetTotalProcessedSimulationTime() > ReconcileSimulationTime)
+			if (Ticker.GetTotalProcessedSimulationTime() > ReconcileSimulationTime)
 			{
-				ReconcileSimulationTime = TickInfo.GetTotalProcessedSimulationTime();
+				ReconcileSimulationTime = Ticker.GetTotalProcessedSimulationTime();
 			}
 
 			// Find out where this should go in the local buffer based on the serialized time
 			int32 DestinationKeyframe = INDEX_NONE;
-			if (LastSerializedSimulationTime > TickInfo.GetTotalProcessedSimulationTime())
+			if (LastSerializedSimulationTime > Ticker.GetTotalProcessedSimulationTime())
 			{
 				// We are getting new state that is ahead of what we have locally, so it can safety go right to head
-				DestinationKeyframe = TickInfo.SimulationTimeBuffer.HeadKeyframe()+1;
+				DestinationKeyframe = Ticker.SimulationTimeBuffer.HeadKeyframe()+1;
 			}
 			else
 			{
 				// We are getting state that is behind what we have locally
-				for (int32 Keyframe = TickInfo.SimulationTimeBuffer.HeadKeyframe(); Keyframe >= TickInfo.SimulationTimeBuffer.TailKeyframe(); --Keyframe)
+				for (int32 Keyframe = Ticker.SimulationTimeBuffer.HeadKeyframe(); Keyframe >= Ticker.SimulationTimeBuffer.TailKeyframe(); --Keyframe)
 				{
-					if (LastSerializedSimulationTime > *TickInfo.SimulationTimeBuffer[Keyframe])
+					if (LastSerializedSimulationTime > *Ticker.SimulationTimeBuffer[Keyframe])
 					{
 						DestinationKeyframe = Keyframe+1;
 						break;
@@ -393,13 +393,13 @@ struct TReplicator_Simulated : public TBase
 
 				if (DestinationKeyframe == INDEX_NONE)
 				{
-					FNetworkSimTime TotalTimeAhead = TickInfo.GetTotalProcessedSimulationTime() - LastSerializedSimulationTime;
+					FNetworkSimTime TotalTimeAhead = Ticker.GetTotalProcessedSimulationTime() - LastSerializedSimulationTime;
 					FNetworkSimTime SerializeDelta = LastSerializedSimulationTime - PrevLastSerializedSimulationTime;
 
 					// We are way far ahead of the server... we will need to clear out sync buffers, take what they gave us, and catch up in reconcile
 					//UE_LOG(LogNetworkSim, Warning, TEXT("!!! TReplicator_Simulated. Large gap detected. SerializedTime: %s. Buffer time: [%s-%s]. %d Elements. DeltaFromHead: %s. DeltaSerialize: %s"), *LastSerializedSimulationTime.ToString(), 
-					//	*TickInfo.SimulationTimeBuffer.GetElementFromTail(0)->ToString(), *TickInfo.SimulationTimeBuffer.GetElementFromHead(0)->ToString(), TickInfo.SimulationTimeBuffer.GetNumValidElements(), *TotalTimeAhead.ToString(), *SerializeDelta.ToString());
-					DestinationKeyframe = TickInfo.SimulationTimeBuffer.HeadKeyframe()+2; // (Skip ahead 2 to force a break in continuity)
+					//	*Ticker.SimulationTimeBuffer.GetElementFromTail(0)->ToString(), *Ticker.SimulationTimeBuffer.GetElementFromHead(0)->ToString(), Ticker.SimulationTimeBuffer.GetNumValidElements(), *TotalTimeAhead.ToString(), *SerializeDelta.ToString());
+					DestinationKeyframe = Ticker.SimulationTimeBuffer.HeadKeyframe()+2; // (Skip ahead 2 to force a break in continuity)
 				}
 			}
 
@@ -411,16 +411,16 @@ struct TReplicator_Simulated : public TBase
 			Buffers.Input.WriteKeyframe(DestinationKeyframe);
 
 			// Update tick info
-			TickInfo.SetTotalProcessedSimulationTime(LastSerializedSimulationTime, DestinationKeyframe);
-			if (TickInfo.GetTotalAllowedSimulationTime() < LastSerializedSimulationTime)
+			Ticker.SetTotalProcessedSimulationTime(LastSerializedSimulationTime, DestinationKeyframe);
+			if (Ticker.GetTotalAllowedSimulationTime() < LastSerializedSimulationTime)
 			{
-				TickInfo.SetTotalAllowedSimulationTime(LastSerializedSimulationTime);
+				Ticker.SetTotalAllowedSimulationTime(LastSerializedSimulationTime);
 			}
 
-			check(TickInfo.GetTotalProcessedSimulationTime() <= TickInfo.GetTotalAllowedSimulationTime());
+			check(Ticker.GetTotalProcessedSimulationTime() <= Ticker.GetTotalAllowedSimulationTime());
 
-			TickInfo.PendingKeyframe = DestinationKeyframe;	// We are about to serialize state to DestinationKeyframe which will be "unprocessed" (has not been used to generate a new frame)
-			TickInfo.MaxAllowedKeyframe = DestinationKeyframe-1; // Do not process PendingKeyframe on our account. ::PreSimTick will advance this based on our interpolation/extrapolation settings
+			Ticker.PendingKeyframe = DestinationKeyframe;	// We are about to serialize state to DestinationKeyframe which will be "unprocessed" (has not been used to generate a new frame)
+			Ticker.MaxAllowedKeyframe = DestinationKeyframe-1; // Do not process PendingKeyframe on our account. ::PreSimTick will advance this based on our interpolation/extrapolation settings
 		}
 
 		check(SyncState && AuxState);
@@ -435,14 +435,14 @@ struct TReplicator_Simulated : public TBase
 	}
 
 	template<typename TSimulation, typename TDriver>
-	void Reconcile(TSimulation* Simulation, TDriver* Driver, TNetworkSimBufferContainer<TBufferTypes>& Buffers, TSimulationTickState<TTickSettings>& TickInfo)
+	void Reconcile(TSimulation* Simulation, TDriver* Driver, TNetworkSimBufferContainer<TBufferTypes>& Buffers, TSimulationTicker<TTickSettings>& Ticker)
 	{
 		if (ReconcileSimulationTime.IsPositive() == false)
 		{
 			return;
 		}
 
-		check(TickInfo.GetTotalProcessedSimulationTime() <= TickInfo.GetTotalAllowedSimulationTime());
+		check(Ticker.GetTotalProcessedSimulationTime() <= Ticker.GetTotalAllowedSimulationTime());
 
 		if (bAllowSimulatedExtrapolation && ParentSimulation == nullptr && NetworkSimulationModelCVars::EnableSimulatedExtrapolation() && NetworkSimulationModelCVars::EnableSimulatedReconcile())
 		{
@@ -454,38 +454,38 @@ struct TReplicator_Simulated : public TBase
 			}
 
 			// Do we have time to make up? We may have extrapolated ahead of the server (totally fine - can happen with small amount of latency variance)
-			FNetworkSimTime DeltaSimTime = ReconcileSimulationTime - TickInfo.GetTotalProcessedSimulationTime();
+			FNetworkSimTime DeltaSimTime = ReconcileSimulationTime - Ticker.GetTotalProcessedSimulationTime();
 			if (DeltaSimTime.IsPositive() && NetworkSimulationModelCVars::EnableSimulatedReconcile())
 			{
-				SimulationExtrapolation<TSimulation, TDriver>(Simulation, Driver, Buffers, TickInfo, DeltaSimTime);
+				SimulationExtrapolation<TSimulation, TDriver>(Simulation, Driver, Buffers, Ticker, DeltaSimTime);
 			}
 		}
 		
-		check(TickInfo.GetTotalProcessedSimulationTime() <= TickInfo.GetTotalAllowedSimulationTime());
+		check(Ticker.GetTotalProcessedSimulationTime() <= Ticker.GetTotalAllowedSimulationTime());
 		ReconcileSimulationTime.Reset();
 	}
 
 	template<typename TDriver>
-	void PreSimTick(TDriver* Driver, TNetworkSimBufferContainer<TBufferTypes>& Buffers, TSimulationTickState<TTickSettings>& TickInfo, const FNetSimTickParameters& TickParameters)
+	void PreSimTick(TDriver* Driver, TNetworkSimBufferContainer<TBufferTypes>& Buffers, TSimulationTicker<TTickSettings>& Ticker, const FNetSimTickParameters& TickParameters)
 	{
 		// Tick if we are dependent simulation or extrapolation is enabled
 		if (ParentSimulation || (bAllowSimulatedExtrapolation && NetworkSimulationModelCVars::EnableSimulatedExtrapolation()))
 		{
 			// Don't start this simulation until you've gotten at least one update from the server
-			if (TickInfo.GetTotalProcessedSimulationTime().IsPositive())
+			if (Ticker.GetTotalProcessedSimulationTime().IsPositive())
 			{
-				TickInfo.GiveSimulationTime(TickParameters.LocalDeltaTimeSeconds);
+				Ticker.GiveSimulationTime(TickParameters.LocalDeltaTimeSeconds);
 			}
 
 			if (TickParameters.bGenerateLocalInputCmds)
 			{
-				TryGenerateLocalInput(Driver, Buffers, TickInfo);
+				TryGenerateLocalInput(Driver, Buffers, Ticker);
 			}
 		}
 	}
 
 	template<typename TDriver>
-	void PostSimTick(TDriver* Driver, const TNetworkSimBufferContainer<TBufferTypes>& Buffers, const TSimulationTickState<TTickSettings>& TickInfo, const FNetSimTickParameters& TickParameters)
+	void PostSimTick(TDriver* Driver, const TNetworkSimBufferContainer<TBufferTypes>& Buffers, const TSimulationTicker<TTickSettings>& Ticker, const FNetSimTickParameters& TickParameters)
 	{
 		if (bAllowSimulatedExtrapolation || ParentSimulation)
 		{
@@ -498,21 +498,21 @@ struct TReplicator_Simulated : public TBase
 		}
 		else
 		{
-			Interpolator.template PostSimTick<TDriver>(Driver, Buffers, TickInfo, TickParameters);
+			Interpolator.template PostSimTick<TDriver>(Driver, Buffers, Ticker, TickParameters);
 		}
 	}
 
 	template<typename TSimulation, typename TDriver>
-	void DependentRollbackBegin(TSimulation* Simulation, TDriver* Driver, TNetworkSimBufferContainer<TBufferTypes>& Buffers, TSimulationTickState<TTickSettings>& TickInfo, const FNetworkSimTime& RollbackDeltaTime, const int32 ParentKeyframe)
+	void DependentRollbackBegin(TSimulation* Simulation, TDriver* Driver, TNetworkSimBufferContainer<TBufferTypes>& Buffers, TSimulationTicker<TTickSettings>& Ticker, const FNetworkSimTime& RollbackDeltaTime, const int32 ParentKeyframe)
 	{
 		// For now, we make the assumption that our last serialized state and time match the parent simulation.
 		// This would not always be the case with low frequency simulated proxies. But could be handled by replicating the simulations together (at the replication level)
 		const int32 NewHeadKeyframe = Buffers.Sync.HeadKeyframe()+1;
 		
-		TickInfo.SetTotalProcessedSimulationTime(LastSerializedSimulationTime, NewHeadKeyframe);
-		TickInfo.SetTotalAllowedSimulationTime(LastSerializedSimulationTime);
-		TickInfo.PendingKeyframe = NewHeadKeyframe;
-		TickInfo.MaxAllowedKeyframe = NewHeadKeyframe;
+		Ticker.SetTotalProcessedSimulationTime(LastSerializedSimulationTime, NewHeadKeyframe);
+		Ticker.SetTotalAllowedSimulationTime(LastSerializedSimulationTime);
+		Ticker.PendingKeyframe = NewHeadKeyframe;
+		Ticker.MaxAllowedKeyframe = NewHeadKeyframe;
 
 		*Buffers.Sync.WriteKeyframe(NewHeadKeyframe) = LastSerializedSyncState;
 		*Buffers.Input.WriteKeyframe(NewHeadKeyframe) = TInputCmd();
@@ -527,11 +527,11 @@ struct TReplicator_Simulated : public TBase
 	}
 
 	template<typename TSimulation, typename TDriver>
-	void DependentRollbackStep(TSimulation* Simulation, TDriver* Driver, TNetworkSimBufferContainer<TBufferTypes>& Buffers, TSimulationTickState<TTickSettings>& TickInfo, const FNetworkSimTime& StepTime, const int32 ParentKeyframe, const bool bFinalStep)
+	void DependentRollbackStep(TSimulation* Simulation, TDriver* Driver, TNetworkSimBufferContainer<TBufferTypes>& Buffers, TSimulationTicker<TTickSettings>& Ticker, const FNetworkSimTime& StepTime, const int32 ParentKeyframe, const bool bFinalStep)
 	{
-		TickInfo.SetTotalAllowedSimulationTime( TickInfo.GetTotalAllowedSimulationTime() + StepTime );
+		Ticker.SetTotalAllowedSimulationTime( Ticker.GetTotalAllowedSimulationTime() + StepTime );
 
-		SimulationExtrapolation<TSimulation, TDriver>(Simulation, Driver, Buffers, TickInfo, StepTime);
+		SimulationExtrapolation<TSimulation, TDriver>(Simulation, Driver, Buffers, Ticker, StepTime);
 
 		TSyncState* SyncState = Buffers.Sync.HeadElement();
 		check(SyncState);
@@ -546,7 +546,7 @@ struct TReplicator_Simulated : public TBase
 private:
 
 	template<typename TSimulation, typename TDriver>
-	void SimulationExtrapolation(TSimulation* Simulation, TDriver* Driver, TNetworkSimBufferContainer<TBufferTypes>& Buffers, TSimulationTickState<TTickSettings>& TickInfo, const FNetworkSimTime DeltaSimTime)
+	void SimulationExtrapolation(TSimulation* Simulation, TDriver* Driver, TNetworkSimBufferContainer<TBufferTypes>& Buffers, TSimulationTicker<TTickSettings>& Ticker, const FNetworkSimTime DeltaSimTime)
 	{
 		// We have extrapolated ahead of the server. The latest network update is now "in the past" from what we rendered last frame.
 		// We will insert a new keyframe to make up the difference from the last known state to where we want to be in the now.
@@ -567,11 +567,14 @@ private:
 
 		// Do the actual update
 		{
-			TScopedSimulationTick UpdateScope(TickInfo, OutputKeyframe, NewCmd->GetFrameDeltaTime());
-			Simulation->Update(NewCmd->GetFrameDeltaTime().ToRealTimeSeconds(), *NewCmd, *PrevSyncState, *NextSyncState, *AuxState, Buffers.Aux.LazyWriter(OutputKeyframe));
+			TScopedSimulationTick UpdateScope(Ticker, OutputKeyframe, NewCmd->GetFrameDeltaTime());
+			Simulation->SimulationTick( 
+				{ NewCmd->GetFrameDeltaTime(), Ticker },
+				{ *NewCmd, *PrevSyncState, *AuxState },
+				{ *NextSyncState, Buffers.Aux.LazyWriter(OutputKeyframe) } );
 		}
 
-		TickInfo.MaxAllowedKeyframe = OutputKeyframe;
+		Ticker.MaxAllowedKeyframe = OutputKeyframe;
 	}
 	
 	FNetworkSimTime ReconcileSimulationTime;
@@ -604,13 +607,13 @@ struct TReplicator_Autonomous : public TBase
 	// --------------------------------------------------------------------
 	//	NetSerialize
 	// --------------------------------------------------------------------
-	void NetSerialize(const FNetSerializeParams& P, TNetworkSimBufferContainer<TBufferTypes>& Buffers, TSimulationTickState<TTickSettings>& TickInfo)
+	void NetSerialize(const FNetSerializeParams& P, TNetworkSimBufferContainer<TBufferTypes>& Buffers, TSimulationTicker<TTickSettings>& Ticker)
 	{
 		FArchive& Ar = P.Ar;
 		
 		SerializedKeyframe = FNetworkSimulationSerialization::SerializeKeyframe(Ar, Buffers.Sync.HeadKeyframe());
 		
-		SerializedTime = TickInfo.GetTotalProcessedSimulationTime();
+		SerializedTime = Ticker.GetTotalProcessedSimulationTime();
 		SerializedTime.NetSerialize(P.Ar);
 
 		if (Ar.IsSaving())
@@ -662,7 +665,7 @@ struct TReplicator_Autonomous : public TBase
 	//	Reconcile
 	// --------------------------------------------------------------------
 	template<typename TSimulation, typename TDriver>
-	void Reconcile(TSimulation* Simulation, TDriver* Driver, TNetworkSimBufferContainer<TBufferTypes>& Buffers, TSimulationTickState<TTickSettings>& TickInfo)
+	void Reconcile(TSimulation* Simulation, TDriver* Driver, TNetworkSimBufferContainer<TBufferTypes>& Buffers, TSimulationTicker<TTickSettings>& Ticker)
 	{
 		if (bPendingReconciliation == false && (bDependentSimulationNeedsReconcile == false || SerializedKeyframe == INDEX_NONE))
 		{
@@ -702,17 +705,17 @@ struct TReplicator_Autonomous : public TBase
 		// Set Client's aux state to the server version
 		*Buffers.Aux.WriteKeyframe(SerializedKeyframe) = SerializedAuxState;
 
-		const FNetworkSimTime RollbackDeltaTime = SerializedTime - TickInfo.GetTotalProcessedSimulationTime();
+		const FNetworkSimTime RollbackDeltaTime = SerializedTime - Ticker.GetTotalProcessedSimulationTime();
 
 		// Set the canonical simulation time to what we received (we will advance it as we resimulate)
-		TickInfo.SetTotalProcessedSimulationTime(SerializedTime, SerializedKeyframe);
-		TickInfo.PendingKeyframe = SerializedKeyframe;
-		//TickInfo.MaxAllowedKeyframe = FMath::Max(TickInfo.MaxAllowedKeyframe, TickInfo.PendingKeyframe); // Make sure this doesn't lag behind. This is the only place we should need to do this.
+		Ticker.SetTotalProcessedSimulationTime(SerializedTime, SerializedKeyframe);
+		Ticker.PendingKeyframe = SerializedKeyframe;
+		//Ticker.MaxAllowedKeyframe = FMath::Max(Ticker.MaxAllowedKeyframe, Ticker.PendingKeyframe); // Make sure this doesn't lag behind. This is the only place we should need to do this.
 
 		if (NetworkSimulationModelCVars::EnableLocalPrediction() == 0)
 		{
 			// If we aren't predicting at all, then we advanced the allowed sim time here, (since we aren't doing it in PreSimTick). This just keeps us constantly falling behind and not being able to toggle prediction on/off for debugging.
-			TickInfo.SetTotalAllowedSimulationTime(SerializedTime);
+			Ticker.SetTotalAllowedSimulationTime(SerializedTime);
 		}
 
 		// Tell dependent simulations to rollback
@@ -722,7 +725,7 @@ struct TReplicator_Autonomous : public TBase
 		}
 		
 		// Resimulate all user commands 
-		const int32 LastKeyframeToProcess = TickInfo.MaxAllowedKeyframe;
+		const int32 LastKeyframeToProcess = Ticker.MaxAllowedKeyframe;
 		for (int32 Keyframe = SerializedKeyframe; Keyframe <= LastKeyframeToProcess; ++Keyframe)
 		{
 			const int32 OutputKeyframe = Keyframe+1;
@@ -747,8 +750,11 @@ struct TReplicator_Autonomous : public TBase
 
 			// Do the actual update
 			{
-				TScopedSimulationTick UpdateScope(TickInfo, OutputKeyframe, ResimulateCmd->GetFrameDeltaTime());
-				Simulation->Update(ResimulateCmd->GetFrameDeltaTime().ToRealTimeSeconds(), *ResimulateCmd, *PrevSyncState, *NextSyncState, *AuxState, Buffers.Aux.LazyWriter(OutputKeyframe));
+				TScopedSimulationTick UpdateScope(Ticker, OutputKeyframe, ResimulateCmd->GetFrameDeltaTime());
+				Simulation->SimulationTick( 
+					{ ResimulateCmd->GetFrameDeltaTime(), Ticker },
+					{ *ResimulateCmd, *PrevSyncState, *AuxState },
+					{ *NextSyncState, Buffers.Aux.LazyWriter(OutputKeyframe) } );
 			}
 
 			// Log out the newly predicted state that we got.
@@ -770,7 +776,7 @@ struct TReplicator_Autonomous : public TBase
 	//	PreSimTick
 	// --------------------------------------------------------------------
 	template<typename TDriver>
-	void PreSimTick(TDriver* Driver, TNetworkSimBufferContainer<TBufferTypes>& Buffers, TSimulationTickState<TTickSettings>& TickInfo, const FNetSimTickParameters& TickParameters)
+	void PreSimTick(TDriver* Driver, TNetworkSimBufferContainer<TBufferTypes>& Buffers, TSimulationTicker<TTickSettings>& Ticker, const FNetSimTickParameters& TickParameters)
 	{
 		// If we have a reconcile fault, we cannot continue on with the simulation until it clears itself out. This effectively drops the input time and does not sample new inputs
 		if (bReconcileFaultDetected)
@@ -783,8 +789,8 @@ struct TReplicator_Autonomous : public TBase
 			if (NetworkSimulationModelCVars::EnableLocalPrediction() > 0)
 			{
 				// Prediction: add simulation time and generate new commands
-				TickInfo.GiveSimulationTime(TickParameters.LocalDeltaTimeSeconds);
-				TryGenerateLocalInput(Driver, Buffers, TickInfo);
+				Ticker.GiveSimulationTime(TickParameters.LocalDeltaTimeSeconds);
+				TryGenerateLocalInput(Driver, Buffers, Ticker);
 			}
 			else
 			{
@@ -826,9 +832,9 @@ struct TReplicator_Debug : public TBase
 		return Buffers.Debug.GetDirtyCount() ^ (TBase::GetProxyDirtyCount(Buffers) << 2); 
 	}
 
-	void NetSerialize(const FNetSerializeParams& P, TNetworkSimBufferContainer<TBufferTypes>& Buffers, TSimulationTickState<TTickSettings>& TickInfo)
+	void NetSerialize(const FNetSerializeParams& P, TNetworkSimBufferContainer<TBufferTypes>& Buffers, TSimulationTicker<TTickSettings>& Ticker)
 	{
-		TBase::NetSerialize(P, Buffers, TickInfo);
+		TBase::NetSerialize(P, Buffers, Ticker);
 		FArchive& Ar = P.Ar;
 
 		TDebugBuffer& Buffer = Ar.IsSaving() ? Buffers.Debug : ReceivedBuffer;
diff --git a/Engine/Plugins/Runtime/NetworkPrediction/Source/NetworkPrediction/Public/NetworkSimulationModelTypes.h b/Engine/Plugins/Runtime/NetworkPrediction/Source/NetworkPrediction/Public/NetworkSimulationModelTypes.h
--- a/Engine/Plugins/Runtime/NetworkPrediction/Source/NetworkPrediction/Public/NetworkSimulationModelTypes.h
+++ b/Engine/Plugins/Runtime/NetworkPrediction/Source/NetworkPrediction/Public/NetworkSimulationModelTypes.h
@@ -219,7 +219,7 @@ struct TBufferGetterHelper_ByState<TContainer, typename TContainer::TDebugState>
 
 // -------------------------------------------------------------------
 
-// Struct that encapsulates writing a new element to a buffer. This is used to allow a new Aux state to be created in the ::Update loop.
+// Struct that encapsulates writing a new element to a buffer. This is used to allow a new Aux state to be created in the ::SimulationTick loop.
 template<typename T>
 struct TLazyStateAccessor
 {
@@ -434,11 +434,11 @@ private:
 };
 
 // ----------------------------------------------------------------------------------------------------------------------------------------------
-//	TSimulationTickState: Holds active state for simulation ticking. We track two things: keyframes and time.
+//	FSimulationTickState: Holds active state for simulation ticking. We track two things: keyframes and time.
 //
-//	PendingKeyframe is the next keyframe we will process: the input/sync/aux state @ PendingKeyframe will be run through T::Update and produce the 
+//	PendingKeyframe is the next keyframe we will process: the input/sync/aux state @ PendingKeyframe will be run through ::SimulationTick and produce the 
 //	next frame's (PendingKeyframe+1) Sync and possibly Aux state (if it changes). "Out of band" modifications to the sync/aux state should happen
-//	to PendingKeyframe (e.g, before it is processed. Once a frame has been processed, we won't run it through T::Update again!).
+//	to PendingKeyframe (e.g, before it is processed. Once a frame has been processed, we won't run it through ::SimulationTick again!).
 //
 //	MaxAllowedKeyframe is a keyframe based limiter on simulation updates. This must be incremented to allow the simulation to advance.
 //
@@ -448,14 +448,12 @@ private:
 //	Consider that Keyframes are essentially client dependent and gaps can happen due to packet loss, etc. Time will always be continuous though.
 //	
 // ----------------------------------------------------------------------------------------------------------------------------------------------
-template<typename TickSettings=TNetworkSimTickSettings<>>
-struct TSimulationTickState
-{
-	using TSettings = TickSettings;
 
+struct FSimulationTickState
+{
 	int32 PendingKeyframe = 0;
 	int32 MaxAllowedKeyframe = -1;
-	bool UpdateInProgress = false;
+	bool bUpdateInProgress = false;
 	
 	FNetworkSimTime GetTotalProcessedSimulationTime() const 
 	{ 
@@ -477,19 +475,6 @@ struct TSimulationTickState
 	// Historic tracking of simulation time. This allows us to timestamp sync data as its produced
 	TNetworkSimContiguousBuffer<FNetworkSimTime>	SimulationTimeBuffer;
 
-	void SetTotalAllowedSimulationTime(const FNetworkSimTime& SimTime)
-	{
-		TotalAllowedSimulationTime = SimTime;
-		RealtimeAccumulator.Reset();
-	}
-
-	// "Grants" allowed simulation time to this tick state. That is, we are now allowed to advance the simulation by this amount the next time the sim ticks.
-	// Note the input is RealTime in SECONDS. This is what the rest of the engine uses when dealing with float delta time.
-	void GiveSimulationTime(float RealTimeSeconds)
-	{
-		RealtimeAccumulator.Accumulate(TotalAllowedSimulationTime, RealTimeSeconds);
-	}
-
 	// How much granted simulation time is left to process
 	FNetworkSimTime GetRemainingAllowedSimulationTime() const
 	{
@@ -501,32 +486,54 @@ struct TSimulationTickState
 		return TotalAllowedSimulationTime;
 	}
 
-private:
+protected:
 
 	FNetworkSimTime TotalAllowedSimulationTime;	// Total time we have been "given" to process. We cannot process more simulation time than this: doing so would be speed hacking.
 	FNetworkSimTime TotalProcessedSimulationTime;	// How much time we've actually processed. The only way to increment this is to process user commands or receive authoritative state from the network.
+};
+
+template<typename TickSettings=TNetworkSimTickSettings<>>
+struct TSimulationTicker : public FSimulationTickState
+{
+	using TSettings = TickSettings;
+
+	void SetTotalAllowedSimulationTime(const FNetworkSimTime& SimTime)
+	{
+		TotalAllowedSimulationTime = SimTime;
+		RealtimeAccumulator.Reset();
+	}
 
-	TRealTimeAccumulator<TSettings>	RealtimeAccumulator;
+	// "Grants" allowed simulation time to this tick state. That is, we are now allowed to advance the simulation by this amount the next time the sim ticks.
+	// Note the input is RealTime in SECONDS. This is what the rest of the engine uses when dealing with float delta time.
+	void GiveSimulationTime(float RealTimeSeconds)
+	{
+		RealtimeAccumulator.Accumulate(TotalAllowedSimulationTime, RealTimeSeconds);
+	}	
+
+private:
+
+	TRealTimeAccumulator<TSettings>	RealtimeAccumulator;	
 };
 
-// Scoped helper to be used right before entering a call to the sim's ::Update function.
+// Scoped helper to be used right before entering a call to the sim's ::SimulationTick function.
 // Important to note that this advanced the PendingKeyFrame to the output Keyframe. So that any writes that occur to the buffers will go to the output frame.
 struct TScopedSimulationTick
 {
-	template<typename TSimulationTickState>
-	TScopedSimulationTick(TSimulationTickState& TickState, const int32& OutputKeyframe, const FNetworkSimTime& DeltaSimTime)
+	TScopedSimulationTick(FSimulationTickState& InTicker, const int32& InOutputKeyframe, const FNetworkSimTime& InDeltaSimTime)
+		: Ticker(InTicker), OutputKeyframe(InOutputKeyframe), DeltaSimTime(InDeltaSimTime)
 	{
-		check(TickState.UpdateInProgress == false);
-		TickState.UpdateInProgress = true;
-		TickState.PendingKeyframe = OutputKeyframe;
-		TickState.IncrementTotalProcessedSimulationTime(DeltaSimTime, OutputKeyframe);
-		UpdateInProgressPtr = &TickState.UpdateInProgress;
+		check(Ticker.bUpdateInProgress == false);
+		Ticker.PendingKeyframe = OutputKeyframe;
+		Ticker.bUpdateInProgress = true;
 	}
 	~TScopedSimulationTick()
 	{
-		*UpdateInProgressPtr = false;
+		Ticker.IncrementTotalProcessedSimulationTime(DeltaSimTime, OutputKeyframe);
+		Ticker.bUpdateInProgress = false;
 	}
-	bool* UpdateInProgressPtr;
+	FSimulationTickState& Ticker;
+	const int32& OutputKeyframe;
+	const FNetworkSimTime& DeltaSimTime;
 };
 
 // ----------------------------------------------------------------------------------------------------------------------------------------------
@@ -535,12 +542,12 @@ struct TScopedSimulationTick
 
 // Accessor conditionally gives access to the current (pending) Sync/Aux state to outside code.
 // Reads are always allowed
-// Writes are conditional. Authority can always write to the pending keyframe. Non authority requires the netsim to be currently processing an ::Update.
-// If you aren't inside an ::Update call, it is really not safe to predict state changes. It is safest and simplest to just not predict these changes.
+// Writes are conditional. Authority can always write to the pending keyframe. Non authority requires the netsim to be currently processing an ::SimulationTick.
+// If you aren't inside an ::SimulationTick call, it is really not safe to predict state changes. It is safest and simplest to just not predict these changes.
 //
-// Explanation: During the scope of an ::Update call, we know exactly 'when' we are relative to what the server is processing. If the predicting client wants
+// Explanation: During the scope of an ::SimulationTick call, we know exactly 'when' we are relative to what the server is processing. If the predicting client wants
 // to predict a change to sync/aux state during an update, the server will do it at the exact same time (assuming not a mis prediction). When a state change
-// happens "out of band" (outside an ::Update call) - we really have no way to correlate when the server will do it. While its tempting to think "we will get
+// happens "out of band" (outside an ::SimulationTick call) - we really have no way to correlate when the server will do it. While its tempting to think "we will get
 // a correction anyways, might as well guess at it and maybe get a smaller correction" - but this opens us up to other problems. The server may actually not 
 // change the state at all and you may not get an update that corrects you. You could add a timeout and track the state change somewhere but that really complicates
 // things and could leave you open to "double" problems: if the state change is additive, you may stack the authority change on top of the local predicted change, or
@@ -558,8 +565,8 @@ struct TNetworkSimStateAccessor
 			// Gross: find the buffer our TState is in. This allows these accessors to be declared as class member variables without knowing the exact net sim it will
 			// be pulling from. (For example, you may want to templatize your network sim model).
 			auto& Buffer = TBufferGetterHelper_ByState<TNetworkSimBufferContainer<typename TNetworkSimModel::TBufferTypes>, TState>::GetBuffer(NetSimModel->Buffers);
-			OutState = bWrite ? Buffer.WriteKeyframeInitializedFromHead(NetSimModel->TickInfo.PendingKeyframe) : Buffer[NetSimModel->TickInfo.PendingKeyframe];
-			OutSafe = NetSimModel->TickInfo.UpdateInProgress;
+			OutState = bWrite ? Buffer.WriteKeyframeInitializedFromHead(NetSimModel->Ticker.PendingKeyframe) : Buffer[NetSimModel->Ticker.PendingKeyframe];
+			OutSafe = NetSimModel->Ticker.bUpdateInProgress;
 		};
 	}
 
@@ -618,7 +625,7 @@ struct FNetworkSimEventAccessor
 };
 
 // ----------------------------------------------------------------------------------------------------------------------------------------------
-//
+//	FrameCmd - in variable tick simulations we store the timestep of each frame with the input. TFrameCmd wraps the user struct to do this.
 // ----------------------------------------------------------------------------------------------------------------------------------------------
 
 // Wraps an input command (BaseType) in a NetworkSimulation time. 
@@ -686,4 +693,46 @@ public:
 	
 	// Called from the Network Sim at the end of the sim frame when there is new sync data.
 	virtual void FinalizeFrame(const typename TBufferTypes::TSyncState& SyncState, const typename TBufferTypes::TAuxState& AuxState) = 0;
-};
\ No newline at end of file
+};
+
+// ----------------------------------------------------------------------------------------------------------------------------------------------
+//	SimulationTick parameters. These are the data structures passed into the simulation code each frame
+// ----------------------------------------------------------------------------------------------------------------------------------------------
+
+struct TNetSimTimeStep
+{
+	// The delta time step for this tick (in MS by default)
+	const FNetworkSimTime& StepMS;
+	// The tick state of the simulation prior to running this tick. E.g, does not "include" the above StepMS that we are simulating now.
+	// The first time ::SimulationTick runs, TickState.GetTotalProcessedSimulationTime() == 0.
+	const FSimulationTickState& TickState;
+};
+
+// Input state: const references to the InputCmd/SyncState/AuxStates
+template<typename TBufferTypes>
+struct TNetSimInput
+{
+	const typename TBufferTypes::TInputCmd& Cmd;
+	const typename TBufferTypes::TSyncState& Sync;
+	const typename TBufferTypes::TAuxState& Aux;
+
+	TNetSimInput(const typename TBufferTypes::TInputCmd& InInputCmd, const typename TBufferTypes::TSyncState& InSync, const typename TBufferTypes::TAuxState& InAux)
+		: Cmd(InInputCmd), Sync(InSync), Aux(InAux) { }	
+	
+	// Allows implicit downcasting to a parent simulation class's input types
+	template<typename T>
+	TNetSimInput(const TNetSimInput<T>& Other)
+		: Cmd(Other.Cmd), Sync(Other.Sync), Aux(Other.Aux) { }
+};
+
+// Output state: the output SyncState (always created) and TNetSimLazyWriter for the AuxState (created on demand since every tick does not generate a new aux frame)
+template<typename TBufferTypes>
+struct TNetSimOutput
+{
+	typename TBufferTypes::TSyncState& Sync;
+	const TNetSimLazyWriter<typename TBufferTypes::TAuxState>& Aux;
+
+	TNetSimOutput(typename TBufferTypes::TSyncState& InSync, const TNetSimLazyWriter<typename TBufferTypes::TAuxState>& InAux)
+		: Sync(InSync), Aux(InAux) { }
+};
+
diff --git a/Engine/Plugins/Runtime/NetworkPrediction/readme.txt b/Engine/Plugins/Runtime/NetworkPrediction/readme.txt
--- a/Engine/Plugins/Runtime/NetworkPrediction/readme.txt
+++ b/Engine/Plugins/Runtime/NetworkPrediction/readme.txt
@@ -108,6 +108,14 @@ Smoothing: Taking the output of the simulation and applying an additional layer
 // Release notes
 // ----------------------------------------------------------------------------------------------------------
 
+Update (10-28-19)
+Renaming/cleanup on ::Update function:
+-"TSimulation::Update" is now "TSimulation::SimulationTick". Makes searching through source a bit easier, less generic name
+-Collapsed parameters of SimulationTick into 3 containers: TNetSimTimeStep, TNetSimInput, TNetSimOutput
+-Time passed into simulation is now FNetworkSimTime (instead of forcing float RealTimeSeconds conversion at callsite)
+-TSimulationTickState broken up into FSimulationTickState and TSimulationTicker. 
+-::SimulationTick now also gets a reference to FSimulationTickState. This is setting us up to do timers/countdowns based on sim time within these functions.
+
 Update (10-25-19)
 -Another large refactor on the core classes. Things are feeling good!
 -"Simulation Driver" and a few other concepts have changed:
diff --git a/Engine/Plugins/Runtime/NetworkPredictionExtras/Source/NetworkPredictionExtras/Private/MockAbilitySimulation.cpp b/Engine/Plugins/Runtime/NetworkPredictionExtras/Source/NetworkPredictionExtras/Private/MockAbilitySimulation.cpp
--- a/Engine/Plugins/Runtime/NetworkPredictionExtras/Source/NetworkPredictionExtras/Private/MockAbilitySimulation.cpp
+++ b/Engine/Plugins/Runtime/NetworkPredictionExtras/Source/NetworkPredictionExtras/Private/MockAbilitySimulation.cpp
@@ -19,39 +19,40 @@ namespace MockAbilityCVars
 //
 // -------------------------------------------------------------------------------------------------------------
 
-
 const FName FMockAbilitySimulation::GroupName(TEXT("Ability"));
 
-void FMockAbilitySimulation::Update(const float DeltaTimeSeconds, const FMockAbilityInputCmd& InputCmd, const FMockAbilitySyncState& InputState, FMockAbilitySyncState& OutputState, const FMockAbilityAuxstate& AuxState, const TNetSimLazyWriter<FMockAbilityAuxstate>& OutAuxStateAccessor)
+void FMockAbilitySimulation::SimulationTick(const TNetSimTimeStep& TimeStep, const TNetSimInput<TMockAbilityBufferTypes>& Input, const TNetSimOutput<TMockAbilityBufferTypes>& Output)
 {
+	const float DeltaSeconds = TimeStep.StepMS.ToRealTimeSeconds();
+
 	// Stamina passes through. Some code paths will modify this again, but if we don't set the output state it will be garbage/stale
 	// (considering implicit copying of old state to new state by the NetworkedSimModel code, but that could be undesired/inefficient in some cases)
-	OutputState.Stamina = InputState.Stamina;
+	Output.Sync.Stamina = Input.Sync.Stamina;
 
 	static float BlinkCost = 75.f;
-	const bool bBlink = (InputCmd.bBlinkPressed && InputState.Stamina > BlinkCost);
+	const bool bBlink = (Input.Cmd.bBlinkPressed && Input.Sync.Stamina > BlinkCost);
 
 	if (bBlink)
 	{
 		static float BlinkDist = 1000.f;
 		static float OverlapCheckInflation = 0.100f;
 
-		FVector DestLocation = InputState.Location + InputState.Rotation.RotateVector( FVector(BlinkDist, 0.f, 0.f) );
+		FVector DestLocation = Input.Sync.Location + Input.Sync.Rotation.RotateVector( FVector(BlinkDist, 0.f, 0.f) );
 		AActor* OwningActor = UpdatedComponent->GetOwner();
 		check(OwningActor);
 
-		// DrawDebugLine(OwningActor->GetWorld(), InputState.Location, DestLocation, FColor::Red, false);	
+		// DrawDebugLine(OwningActor->GetWorld(), Input.Sync.Location, DestLocation, FColor::Red, false);	
 
 		// Its unfortunate teleporting is so complicated. It may make sense for a new movement simulation to define this themselves, 
 		// but for this mock one, we will just use the engine's AActor teleport.
-		if (OwningActor->TeleportTo(DestLocation, InputState.Rotation))
+		if (OwningActor->TeleportTo(DestLocation, Input.Sync.Rotation))
 		{
-			OutputState = InputState;
+			Output.Sync = Input.Sync;
 
 			// Component now has the final location
 			const FTransform UpdateComponentTransform = GetUpdateComponentTransform();
-			OutputState.Location = UpdateComponentTransform.GetLocation();
-			OutputState.Stamina = InputState.Stamina - BlinkCost;
+			Output.Sync.Location = UpdateComponentTransform.GetLocation();
+			Output.Sync.Stamina = Input.Sync.Stamina - BlinkCost;
 
 			// And we skip the normal update simulation for this frame. This is just a choice. We could still allow it to run.
 			return;
@@ -70,43 +71,43 @@ void FMockAbilitySimulation::Update(const float DeltaTimeSeconds, const FMockAbi
 
 	static float DashCost = 75.f;
 	static int16 DashDurationMS = 400;
-	int16 DashTimeLeft = AuxState.DashTimeLeft;
-	bool bIsDashing = AuxState.DashTimeLeft > 0;
+	int16 DashTimeLeft = Input.Aux.DashTimeLeft;
+	bool bIsDashing = Input.Aux.DashTimeLeft > 0;
 	
-	if (InputCmd.bDashPressed && InputState.Stamina > DashCost && bIsDashing == false)
+	if (Input.Cmd.bDashPressed && Input.Sync.Stamina > DashCost && bIsDashing == false)
 	{
 		// Start dashing
 		DashTimeLeft = DashDurationMS;
-		OutputState.Stamina -= DashCost;
+		Output.Sync.Stamina -= DashCost;
 		bIsDashing = true;
 	}
 
 	if (bIsDashing)
 	{
-		FMockAbilityAuxstate* OutAuxState = OutAuxStateAccessor.Get();
-		OutAuxState->DashTimeLeft = FMath::Max<int16>(DashTimeLeft - (int16)(DeltaTimeSeconds * 1000.f), 0);
+		FMockAbilityAuxstate* OutAuxState = Output.Aux.Get();
+		OutAuxState->DashTimeLeft = FMath::Max<int16>(DashTimeLeft - (int16)(DeltaSeconds * 1000.f), 0);
 
-		FMockAbilityAuxstate LocalAuxState = AuxState;
+		FMockAbilityAuxstate LocalAuxState = Input.Aux;
 		LocalAuxState.MaxSpeed = MockAbilityCVars::DashMaxSpeed();
 		LocalAuxState.Acceleration = MockAbilityCVars::DashAcceleration();
 
-		FMockAbilityInputCmd LocalInputCmd = InputCmd;
+		FMockAbilityInputCmd LocalInputCmd = Input.Cmd;
 		LocalInputCmd.MovementInput = FVector(1.f, 0.f, 0.f);
 		
-		FlyingMovement::FMovementSimulation::Update(DeltaTimeSeconds, LocalInputCmd, InputState, OutputState, LocalAuxState, OutAuxStateAccessor);
+		FlyingMovement::FMovementSimulation::SimulationTick(TimeStep, { LocalInputCmd, Input.Sync, LocalAuxState }, { Output.Sync, Output.Aux });
 
 		if (OutAuxState->DashTimeLeft == 0)
 		{
 			// Stop when dash is over
-			OutputState.Velocity = FVector::ZeroVector;
+			Output.Sync.Velocity = FVector::ZeroVector;
 		}
 	}
 	else
 	{
 		// Sprint (mutually exclusive from Dash state)		
 		static float SprintBaseCost = 100.f;
-		const float SprintCostThisFrame = SprintBaseCost * DeltaTimeSeconds;
-		const bool bIsSprinting = (InputCmd.bSprintPressed && InputState.Stamina > SprintCostThisFrame);
+		const float SprintCostThisFrame = SprintBaseCost * DeltaSeconds;
+		const bool bIsSprinting = (Input.Cmd.bSprintPressed && Input.Sync.Stamina > SprintCostThisFrame);
 
 		// Set our max speed. This is an interesting case.
 		//	-Our input states are already "final". It doesn't make sense to modify the input AuxState data.
@@ -119,26 +120,26 @@ void FMockAbilitySimulation::Update(const float DeltaTimeSeconds, const FMockAbi
 		// value from other input state. This difference in how the simulations treat the variable is what causes the need to do this.
 		//
 		// It would be possible to write the base movement sim in a way that MaxSpeed is transient value on the sim class (FlyingMovement::FMovementSimulation).
-		// Something like, "MaxSpeed" really means "Base max speed" and their would be a "GetCurrentMaxSpeed" virtual on the sim.
+		// Something like, "MaxSpeed" really means "Base max speed" and there would be a "GetCurrentMaxSpeed" virtual on the sim.
 		// This would make things a bit more awkward in the base case with no ability system. So, for now, this seems like a good pattern/precedence.
 
-		FMockAbilityAuxstate LocalAuxState = AuxState;
+		FMockAbilityAuxstate LocalAuxState = Input.Aux;
 		LocalAuxState.MaxSpeed = bIsSprinting ? MockAbilityCVars::SprintMaxSpeed() : MockAbilityCVars::DefaultMaxSpeed();
 
 		if (bIsSprinting)
 		{
-			OutputState.Stamina = FMath::Max<float>( InputState.Stamina - SprintCostThisFrame, 0.f );
+			Output.Sync.Stamina = FMath::Max<float>( Input.Sync.Stamina - SprintCostThisFrame, 0.f );
 		}
-		else if (OutputState.Stamina < AuxState.MaxStamina)
+		else if (Output.Sync.Stamina < Input.Aux.MaxStamina)
 		{
-			OutputState.Stamina = FMath::Min<float>( InputState.Stamina + (DeltaTimeSeconds * AuxState.StaminaRegenRate), AuxState.MaxStamina );
+			Output.Sync.Stamina = FMath::Min<float>( Input.Sync.Stamina + (DeltaSeconds * Input.Aux.StaminaRegenRate), Input.Aux.MaxStamina );
 		}
 		else
 		{
-			OutputState.Stamina = InputState.Stamina;
+			Output.Sync.Stamina = Input.Sync.Stamina;
 		}
 
-		FlyingMovement::FMovementSimulation::Update(DeltaTimeSeconds, InputCmd, InputState, OutputState, LocalAuxState, OutAuxStateAccessor);
+		FlyingMovement::FMovementSimulation::SimulationTick(TimeStep, Input, { Output.Sync, Output.Aux });
 	}
 }
 
diff --git a/Engine/Plugins/Runtime/NetworkPredictionExtras/Source/NetworkPredictionExtras/Private/MockNetworkSimulation.cpp b/Engine/Plugins/Runtime/NetworkPredictionExtras/Source/NetworkPredictionExtras/Private/MockNetworkSimulation.cpp
--- a/Engine/Plugins/Runtime/NetworkPredictionExtras/Source/NetworkPredictionExtras/Private/MockNetworkSimulation.cpp
+++ b/Engine/Plugins/Runtime/NetworkPredictionExtras/Source/NetworkPredictionExtras/Private/MockNetworkSimulation.cpp
@@ -61,16 +61,16 @@ static bool ForceMispredict = false;
 
 const FName FMockNetworkSimulation::GroupName(TEXT("Mock"));
 
-void FMockNetworkSimulation::Update(const float DeltaTimeSeconds, const FMockInputCmd& InputCmd, const FMockSyncState& InputState, FMockSyncState& OutputState, const FMockAuxState& AuxState, const TNetSimLazyWriter<FMockAuxState>& OutAuxStateAccessor)
+void FMockNetworkSimulation::SimulationTick(const TNetSimTimeStep& TimeStep, const TNetSimInput<TMockNetworkSimulationBufferTypes>& Input, const TNetSimOutput<TMockNetworkSimulationBufferTypes>& Output)
 {
-	OutputState.Total = InputState.Total + (InputCmd.InputValue * AuxState.Multiplier * DeltaTimeSeconds);
+	Output.Sync.Total = Input.Sync.Total + (Input.Cmd.InputValue * Input.Aux.Multiplier * TimeStep.StepMS.ToRealTimeSeconds());
 
 	// Dev hack to force mispredict
 	if (ForceMispredict)
 	{
 		const float Fudge = FMath::FRand() * 100.f;
-		OutputState.Total += Fudge;
-		UE_LOG(LogMockNetworkSim, Warning, TEXT("ForcingMispredict via CVar. Fudge=%.2f. NewTotal: %.2f"), Fudge, OutputState.Total);
+		Output.Sync.Total += Fudge;
+		UE_LOG(LogMockNetworkSim, Warning, TEXT("ForcingMispredict via CVar. Fudge=%.2f. NewTotal: %.2f"), Fudge, Output.Sync.Total);
 		
 		ForceMispredict = false;
 	}
diff --git a/Engine/Plugins/Runtime/NetworkPredictionExtras/Source/NetworkPredictionExtras/Public/MockAbilitySimulation.h b/Engine/Plugins/Runtime/NetworkPredictionExtras/Source/NetworkPredictionExtras/Public/MockAbilitySimulation.h
--- a/Engine/Plugins/Runtime/NetworkPredictionExtras/Source/NetworkPredictionExtras/Public/MockAbilitySimulation.h
+++ b/Engine/Plugins/Runtime/NetworkPredictionExtras/Source/NetworkPredictionExtras/Public/MockAbilitySimulation.h
@@ -111,6 +111,8 @@ struct FMockAbilityAuxstate : public FlyingMovement::FAuxState
 	}
 };
 
+using TMockAbilityBufferTypes = TNetworkSimBufferTypes<FMockAbilityInputCmd, FMockAbilitySyncState, FMockAbilityAuxstate>;
+
 class FMockAbilitySimulation : public FlyingMovement::FMovementSimulation
 {
 public:
@@ -118,11 +120,9 @@ public:
 	static const FName GroupName;
 
 	/** Main update function */
-	void Update(const float DeltaTimeSeconds, const FMockAbilityInputCmd& InputCmd, const FMockAbilitySyncState& InputState, FMockAbilitySyncState& OutputState, const FMockAbilityAuxstate& AuxState, const TNetSimLazyWriter<FMockAbilityAuxstate>& OutAuxStateAccessor);
+	void SimulationTick(const TNetSimTimeStep& TimeStep, const TNetSimInput<TMockAbilityBufferTypes>& Input, const TNetSimOutput<TMockAbilityBufferTypes>& Output);
 };
 
-using TMockAbilityBufferTypes = TNetworkSimBufferTypes<FMockAbilityInputCmd, FMockAbilitySyncState, FMockAbilityAuxstate>;
-
 template<int32 InFixedStepMS=0>
 using FMockAbilitySystem = TNetworkedSimulationModel<FMockAbilitySimulation, TMockAbilityBufferTypes, TNetworkSimTickSettings<InFixedStepMS>>;
 
diff --git a/Engine/Plugins/Runtime/NetworkPredictionExtras/Source/NetworkPredictionExtras/Public/MockNetworkSimulation.h b/Engine/Plugins/Runtime/NetworkPredictionExtras/Source/NetworkPredictionExtras/Public/MockNetworkSimulation.h
--- a/Engine/Plugins/Runtime/NetworkPredictionExtras/Source/NetworkPredictionExtras/Public/MockNetworkSimulation.h
+++ b/Engine/Plugins/Runtime/NetworkPredictionExtras/Source/NetworkPredictionExtras/Public/MockNetworkSimulation.h
@@ -30,7 +30,7 @@ class IMockDriver;
 //	in this example. This is just meant to show the bare minimum hook ups into the system to make it easier to understand.
 //
 //	Highlights:
-//		FMockNetworkModel::Update			The "core update" function of the simulation.
+//		FMockNetworkModel::SimulationTick		The "core update" function of the simulation.
 //		UMockNetworkSimulationComponent			The UE4 Actor Component that anchors the system to an actor.
 //
 //	Usage:
@@ -131,8 +131,6 @@ struct FMockAuxState
 
 using TMockNetworkSimulationBufferTypes = TNetworkSimBufferTypes<FMockInputCmd, FMockSyncState, FMockAuxState>;
 
-
-
 static FName MockSimulationGroupName("Mock");
 
 class FMockNetworkSimulation
@@ -140,7 +138,7 @@ class FMockNetworkSimulation
 public:
 
 	/** Main update function */
-	void Update(const float DeltaTimeSeconds, const FMockInputCmd& InputCmd, const FMockSyncState& InputState, FMockSyncState& OutputState, const FMockAuxState& AuxState, const TNetSimLazyWriter<FMockAuxState>& OutAuxStateAccessor);
+	void SimulationTick(const TNetSimTimeStep& TimeStep, const TNetSimInput<TMockNetworkSimulationBufferTypes>& Input, const TNetSimOutput<TMockNetworkSimulationBufferTypes>& Output);
 
 	/** Tick group the simulation maps to */
 	static const FName GroupName;
diff --git a/Engine/Plugins/Runtime/NetworkPredictionExtras/Source/NetworkPredictionExtras/Public/NetworkPredictionExtrasFlyingPawn.h b/Engine/Plugins/Runtime/NetworkPredictionExtras/Source/NetworkPredictionExtras/Public/NetworkPredictionExtrasFlyingPawn.h
--- a/Engine/Plugins/Runtime/NetworkPredictionExtras/Source/NetworkPredictionExtras/Public/NetworkPredictionExtrasFlyingPawn.h
+++ b/Engine/Plugins/Runtime/NetworkPredictionExtras/Source/NetworkPredictionExtras/Public/NetworkPredictionExtrasFlyingPawn.h
@@ -18,7 +18,7 @@ class UFlyingMovementComponent;
 //	that are used by the flying movement simulation. This includes some basic camera/aiming code.
 //
 //	Highlights:
-//		FlyingMovement::FMovementSystem::Update						The "core update" function of the flying movement simulation.
+//		FlyingMovement::FMovementSystem::SimulationTick				The "core update" function of the flying movement simulation.
 //		ANetworkPredictionExtrasFlyingPawn::GenerateLocalInput		Function that generates local input commands that are fed into the movement system.
 //
 //	Usage:
