Network Prediction Plugin:
-bulk renaming of the term "keyframe" -> "frame" within the system.
-Reasoning: "Keyframe" has a tendency to imply subframes or the existence of some other type of frame which is not the case in this system. It is also shorter and less typing.

#rnx

diff --git a/Engine/Plugins/Runtime/NetworkPrediction/Source/NetworkPrediction/Private/Movement/ParametricMovement.cpp b/Engine/Plugins/Runtime/NetworkPrediction/Source/NetworkPrediction/Private/Movement/ParametricMovement.cpp
--- a/Engine/Plugins/Runtime/NetworkPrediction/Source/NetworkPrediction/Private/Movement/ParametricMovement.cpp
+++ b/Engine/Plugins/Runtime/NetworkPrediction/Source/NetworkPrediction/Private/Movement/ParametricMovement.cpp
@@ -152,7 +152,7 @@ void UParametricMovementComponent::VisualLog(const ParametricMovement::FInputCmd
 	FTransform Transform;
 	ParametricMotion.MapTimeToTransform(Sync->Position, Transform);
 	FVisualLoggingHelpers::VisualLogActor(GetOwner(), Transform, SystemParameters);
-	//DrawParams.LogText = FString::Printf(TEXT("[%d] %s. Position: %.4f. Location: %s. Rotation: %s"), Parameters.Keyframe, *LexToString(Parameters.Context), Position, *Transform.GetLocation().ToString(), *Transform.GetRotation().Rotator().ToString());
+	//DrawParams.LogText = FString::Printf(TEXT("[%d] %s. Position: %.4f. Location: %s. Rotation: %s"), Parameters.Frame, *LexToString(Parameters.Context), Position, *Transform.GetLocation().ToString(), *Transform.GetRotation().Rotator().ToString());
 }
 
 
diff --git a/Engine/Plugins/Runtime/NetworkPrediction/Source/NetworkPrediction/Private/NetworkPredictionTests.cpp b/Engine/Plugins/Runtime/NetworkPrediction/Source/NetworkPrediction/Private/NetworkPredictionTests.cpp
--- a/Engine/Plugins/Runtime/NetworkPrediction/Source/NetworkPrediction/Private/NetworkPredictionTests.cpp
+++ b/Engine/Plugins/Runtime/NetworkPrediction/Source/NetworkPrediction/Private/NetworkPredictionTests.cpp
@@ -45,17 +45,17 @@ void NetSimulationTestBuffer(TBuffer& Buffer)
 	for (int32 i=0; i < 10; ++i)
 	{
 		typename TBuffer::ElementType NewData;
-		*Buffer.WriteKeyframe(i) = NewData;
+		*Buffer.WriteFrame(i) = NewData;
 		UE_LOG(NetworkSimulationTests, Display, TEXT("%s"), *Buffer.GetBasicDebugStr());
 
 		for (auto It = Buffer.CreateIterator(); It; ++It)
 		{
-			UE_LOG(NetworkSimulationTests, Display, TEXT("   [%d] %s"), It.Keyframe(), *It.Element()->ToString());
+			UE_LOG(NetworkSimulationTests, Display, TEXT("   [%d] %s"), It.Frame(), *It.Element()->ToString());
 		}
 
 		for (auto It = Buffer.CreateConstIterator(); It; ++It)
 		{
-			UE_LOG(NetworkSimulationTests, Display, TEXT("   [%d] %s"), It.Keyframe(), *It.Element()->ToString());
+			UE_LOG(NetworkSimulationTests, Display, TEXT("   [%d] %s"), It.Frame(), *It.Element()->ToString());
 		}
 	}
 }
@@ -80,5 +80,5 @@ void NetSimulationBufferTest()
 
 
 	//TNetworkSimBufferIterator<TNetworkSimSparseBuffer<FState>> It(SparseBuffer);
-	//It.Keyframe();
+	//It.Frame();
 }
diff --git a/Engine/Plugins/Runtime/NetworkPrediction/Source/NetworkPrediction/Public/Movement/FlyingMovement.h b/Engine/Plugins/Runtime/NetworkPrediction/Source/NetworkPrediction/Public/Movement/FlyingMovement.h
--- a/Engine/Plugins/Runtime/NetworkPrediction/Source/NetworkPrediction/Public/Movement/FlyingMovement.h
+++ b/Engine/Plugins/Runtime/NetworkPrediction/Source/NetworkPrediction/Public/Movement/FlyingMovement.h
@@ -48,7 +48,7 @@ namespace FlyingMovement
 		{
 			if (P.Context == EStandardLoggingContext::HeaderOnly)
 			{
-				P.Ar->Logf(TEXT(" %d "), P.Keyframe);
+				P.Ar->Logf(TEXT(" %d "), P.Frame);
 			}
 			else if (P.Context == EStandardLoggingContext::Full)
 			{
@@ -89,11 +89,11 @@ namespace FlyingMovement
 		{
 			if (Params.Context == EStandardLoggingContext::HeaderOnly)
 			{
-				Params.Ar->Logf(TEXT(" %d "), Params.Keyframe);
+				Params.Ar->Logf(TEXT(" %d "), Params.Frame);
 			}
 			else if (Params.Context == EStandardLoggingContext::Full)
 			{
-				Params.Ar->Logf(TEXT("Frame: %d"), Params.Keyframe);
+				Params.Ar->Logf(TEXT("Frame: %d"), Params.Frame);
 				Params.Ar->Logf(TEXT("Loc: %s"), *Location.ToString());
 				Params.Ar->Logf(TEXT("Vel: %s"), *Velocity.ToString());
 				Params.Ar->Logf(TEXT("Rot: %s"), *Rotation.ToString());
@@ -128,7 +128,7 @@ namespace FlyingMovement
 		{
 			if (P.Context == EStandardLoggingContext::HeaderOnly)
 			{
-				P.Ar->Logf(TEXT(" %d "), P.Keyframe);
+				P.Ar->Logf(TEXT(" %d "), P.Frame);
 			}
 			else if (P.Context == EStandardLoggingContext::Full)
 			{
diff --git a/Engine/Plugins/Runtime/NetworkPrediction/Source/NetworkPrediction/Public/Movement/ParametricMovement.h b/Engine/Plugins/Runtime/NetworkPrediction/Source/NetworkPrediction/Public/Movement/ParametricMovement.h
--- a/Engine/Plugins/Runtime/NetworkPrediction/Source/NetworkPrediction/Public/Movement/ParametricMovement.h
+++ b/Engine/Plugins/Runtime/NetworkPrediction/Source/NetworkPrediction/Public/Movement/ParametricMovement.h
@@ -28,7 +28,7 @@ namespace ParametricMovement
 		{
 			if (P.Context == EStandardLoggingContext::HeaderOnly)
 			{
-				P.Ar->Logf(TEXT(" %d "), P.Keyframe);
+				P.Ar->Logf(TEXT(" %d "), P.Frame);
 			}
 			else if (P.Context == EStandardLoggingContext::Full)
 			{
@@ -63,11 +63,11 @@ namespace ParametricMovement
 		{
 			if (Params.Context == EStandardLoggingContext::HeaderOnly)
 			{
-				Params.Ar->Logf(TEXT(" %d "), Params.Keyframe);
+				Params.Ar->Logf(TEXT(" %d "), Params.Frame);
 			}
 			else if (Params.Context == EStandardLoggingContext::Full)
 			{
-				Params.Ar->Logf(TEXT("Frame: %d"), Params.Keyframe);
+				Params.Ar->Logf(TEXT("Frame: %d"), Params.Frame);
 				Params.Ar->Logf(TEXT("Pos: %.2f"), Position);
 				Params.Ar->Logf(TEXT("Rate: %.2f"), PlayRate);
 			}
@@ -96,7 +96,7 @@ namespace ParametricMovement
 		{
 			if (Params.Context == EStandardLoggingContext::HeaderOnly)
 			{
-				Params.Ar->Logf(TEXT(" %d "), Params.Keyframe);
+				Params.Ar->Logf(TEXT(" %d "), Params.Frame);
 			}
 			else if (Params.Context == EStandardLoggingContext::Full)
 			{
diff --git a/Engine/Plugins/Runtime/NetworkPrediction/Source/NetworkPrediction/Public/NetworkPredictionTypes.h b/Engine/Plugins/Runtime/NetworkPrediction/Source/NetworkPrediction/Public/NetworkPredictionTypes.h
--- a/Engine/Plugins/Runtime/NetworkPrediction/Source/NetworkPrediction/Public/NetworkPredictionTypes.h
+++ b/Engine/Plugins/Runtime/NetworkPrediction/Source/NetworkPrediction/Public/NetworkPredictionTypes.h
@@ -49,11 +49,11 @@ enum class EStandardLoggingContext: uint8
 
 struct FStandardLoggingParameters
 {
-	FStandardLoggingParameters(FOutputDevice* InAr, EStandardLoggingContext InContext, int32 InKeyframe)
-		: Ar(InAr), Context(InContext), Keyframe(InKeyframe) { }
+	FStandardLoggingParameters(FOutputDevice* InAr, EStandardLoggingContext InContext, int32 InFrame)
+		: Ar(InAr), Context(InContext), Frame(InFrame) { }
 	FOutputDevice* Ar;
 	EStandardLoggingContext Context;
-	int32 Keyframe;
+	int32 Frame;
 };
 
 struct FNetworkSimulationModelInitParameters
@@ -69,47 +69,47 @@ struct FNetSimProcessedFrameDebugInfo
 {
 	int32 LocalGFrameNumber; // Local GFrame number
 	float LocalDeltaTimeSeconds; // Local frame time
-	TArray<int32> ProcessedKeyframes; // Which keyframes were processed this frame
-	int32 PendingKeyframe; // What PendingKeyframe was at the end of the frame. Does NOT mean we processed it this frame!
-	int32 HeadKeyframe; // Head keyframe of the inputbuffer when the frame ended
+	TArray<int32> ProcessedFrames; // Which frames were processed this frame
+	int32 PendingFrame; // What PendingFrame was at the end of the frame. Does NOT mean we processed it this frame!
+	int32 HeadFrame; // Head frame of the inputbuffer when the frame ended
 
 	float RemainingAllowedSimulationTimeSeconds;
 	
-	int32 LastSentInputKeyframe;
-	int32 LastReceivedInputKeyframe;
+	int32 LastSentInputFrame;
+	int32 LastReceivedInputFrame;
 
 	void NetSerialize(const FNetSerializeParams& P)
 	{	
 		P.Ar << LocalGFrameNumber;
 		P.Ar << LocalDeltaTimeSeconds;
-		P.Ar << PendingKeyframe;
-		P.Ar << HeadKeyframe;
+		P.Ar << PendingFrame;
+		P.Ar << HeadFrame;
 		P.Ar << RemainingAllowedSimulationTimeSeconds;
 
-		P.Ar << LastSentInputKeyframe;
-		P.Ar << LastReceivedInputKeyframe;
+		P.Ar << LastSentInputFrame;
+		P.Ar << LastReceivedInputFrame;
 
-		SafeNetSerializeTArray_Default<31>(P.Ar, ProcessedKeyframes);
+		SafeNetSerializeTArray_Default<31>(P.Ar, ProcessedFrames);
 	}
 	void Log(FStandardLoggingParameters& P) const
 	{
 		P.Ar->Logf(TEXT("LocalGFrameNumber: %d"), LocalGFrameNumber);
 		P.Ar->Logf(TEXT("LocalDeltaTimeSeconds: %.4f"), LocalDeltaTimeSeconds);
-		P.Ar->Logf(TEXT("PendingKeyframe: %d"), PendingKeyframe);
-		P.Ar->Logf(TEXT("HeadKeyframe: %d"), HeadKeyframe);
+		P.Ar->Logf(TEXT("PendingFrame: %d"), PendingFrame);
+		P.Ar->Logf(TEXT("HeadFrame: %d"), HeadFrame);
 		P.Ar->Logf(TEXT("RemainingAllowedSimulationTimeSeconds: %.4f"), RemainingAllowedSimulationTimeSeconds);
 
-		P.Ar->Logf(TEXT("LastSentInputKeyframe: %d"), LastSentInputKeyframe);
-		P.Ar->Logf(TEXT("LastReceivedInputKeyframe: %d"), LastReceivedInputKeyframe);
+		P.Ar->Logf(TEXT("LastSentInputFrame: %d"), LastSentInputFrame);
+		P.Ar->Logf(TEXT("LastReceivedInputFrame: %d"), LastReceivedInputFrame);
 
-		FString ProcessedKeyframesStr;
-		for (int32 k : ProcessedKeyframes)
+		FString ProcessedFramesStr;
+		for (int32 k : ProcessedFrames)
 		{
-			ProcessedKeyframesStr += LexToString(k);
-			ProcessedKeyframesStr += TEXT(" ");
+			ProcessedFramesStr += LexToString(k);
+			ProcessedFramesStr += TEXT(" ");
 		}
-		ProcessedKeyframesStr.TrimEndInline();
-		P.Ar->Logf(TEXT("ProcessedKeyframes: [%s]"), *ProcessedKeyframesStr);
+		ProcessedFramesStr.TrimEndInline();
+		P.Ar->Logf(TEXT("ProcessedFrames: [%s]"), *ProcessedFramesStr);
 	}
 };
 
@@ -130,8 +130,8 @@ enum class EVisualLoggingContext: uint8
 	// (Contexts used by interpolation)
 	InterpolationBufferHead,// Head end of sync buffer while interpolating
 	InterpolationBufferTail,// Tail end of sync buffer while interpolating
-	InterpolationFrom,		// Immediate "from" interpolation keyframe
-	InterpolationTo,		// Immediate "to" interpolation keyframe
+	InterpolationFrom,		// Immediate "from" interpolation frame
+	InterpolationTo,		// Immediate "to" interpolation frame
 	
 	InterpolationLatest,	// Latest interpolation position in normal circumstances
 	InterpolationWaiting,	// Latest interpolation while waiting (overrun)
@@ -153,11 +153,11 @@ enum class EVisualLoggingLifetime : uint8
 
 struct FVisualLoggingParameters
 {
-	FVisualLoggingParameters(EVisualLoggingContext InContext, int32 InKeyframe, EVisualLoggingLifetime InLifetime, const FString& Str=FString()) :
-		 Context(InContext), Keyframe(InKeyframe), Lifetime(InLifetime), DebugString(Str) { }
+	FVisualLoggingParameters(EVisualLoggingContext InContext, int32 InFrame, EVisualLoggingLifetime InLifetime, const FString& Str=FString()) :
+		 Context(InContext), Frame(InFrame), Lifetime(InLifetime), DebugString(Str) { }
 
 	EVisualLoggingContext Context;
-	int32 Keyframe;
+	int32 Frame;
 	EVisualLoggingLifetime Lifetime;
 	FString DebugString;
 
@@ -249,8 +249,8 @@ public:
 	virtual void NotifyDependentSimNeedsReconcile() = 0;
 	
 	// Called by parent sim on the dependent sim as it reconciles
-	virtual void BeginRollback(const struct FNetworkSimTime& RollbackDeltaTime, const int32 ParentKeyframe) = 0;
-	virtual void StepRollback(const struct FNetworkSimTime& Step, const int32 ParentKeyframe, const bool FinalStep) = 0;
+	virtual void BeginRollback(const struct FNetworkSimTime& RollbackDeltaTime, const int32 ParentFrame) = 0;
+	virtual void StepRollback(const struct FNetworkSimTime& Step, const int32 ParentFrame, const bool FinalStep) = 0;
 };
 
 // -------------------------------------------------------------------------------------------------------------------------------
diff --git a/Engine/Plugins/Runtime/NetworkPrediction/Source/NetworkPrediction/Public/NetworkSimulationModel.h b/Engine/Plugins/Runtime/NetworkPrediction/Source/NetworkPrediction/Public/NetworkSimulationModel.h
--- a/Engine/Plugins/Runtime/NetworkPrediction/Source/NetworkPrediction/Public/NetworkSimulationModel.h
+++ b/Engine/Plugins/Runtime/NetworkPrediction/Source/NetworkPrediction/Public/NetworkSimulationModel.h
@@ -65,8 +65,8 @@ public:
 		check(InSimulation && InDriver);
 		Simulation = InSimulation;
 		Driver = InDriver;
-		*Buffers.Sync.WriteKeyframe(0) = InitialSyncState;
-		*Buffers.Aux.WriteKeyframe(0) = InitialAuxState;
+		*Buffers.Sync.WriteFrame(0) = InitialSyncState;
+		*Buffers.Aux.WriteFrame(0) = InitialAuxState;
 		Ticker.SetTotalProcessedSimulationTime(FNetworkSimTime(), 0);
 		DO_NETSIM_MODEL_DEBUG(FNetworkSimulationModelDebuggerManager::Get().RegisterNetworkSimulationModel(this, Driver->GetVLogOwner()));
 	}
@@ -83,15 +83,15 @@ public:
 		// (property replication and ServerRPC get sent after the tick, rather than forcing awkward callback into the NetSim post replication, we can check it here)
 		if (auto* DebugBuffer = GetLocalDebugBuffer())
 		{
-			if (TDebugState* const PrevDebugState = DebugBuffer->Get(DebugBuffer->HeadKeyframe()))
+			if (TDebugState* const PrevDebugState = DebugBuffer->Get(DebugBuffer->HeadFrame()))
 			{
 				if (Parameters.Role == ROLE_AutonomousProxy)
 				{
-					PrevDebugState->LastSentInputKeyframe = RepProxy_ServerRPC.GetLastSerializedKeyframe();
+					PrevDebugState->LastSentInputFrame = RepProxy_ServerRPC.GetLastSerializedFrame();
 				}
 				else if (Parameters.Role == ROLE_Authority)
 				{
-					PrevDebugState->LastSentInputKeyframe = RepProxy_Autonomous.GetLastSerializedKeyframe();
+					PrevDebugState->LastSentInputFrame = RepProxy_Autonomous.GetLastSerializedFrame();
 				}
 			}
 		}
@@ -103,15 +103,15 @@ public:
 			*DebugState = TDebugState();
 			DebugState->LocalDeltaTimeSeconds = Parameters.LocalDeltaTimeSeconds;
 			DebugState->LocalGFrameNumber = GFrameNumber;
-			DebugState->ProcessedKeyframes.Reset();
+			DebugState->ProcessedFrames.Reset();
 			
 			if (Parameters.Role == ROLE_AutonomousProxy)
 			{
-				DebugState->LastReceivedInputKeyframe = RepProxy_Autonomous.GetLastSerializedKeyframe();
+				DebugState->LastReceivedInputFrame = RepProxy_Autonomous.GetLastSerializedFrame();
 			}
 			else if (Parameters.Role == ROLE_Authority)
 			{
-				DebugState->LastReceivedInputKeyframe = RepProxy_ServerRPC.GetLastSerializedKeyframe();
+				DebugState->LastReceivedInputFrame = RepProxy_ServerRPC.GetLastSerializedFrame();
 			}
 		}
 
@@ -138,13 +138,13 @@ public:
 		// -------------------------------------------------------------------------------------------------------------------------------------------------
 		//												Input Processing & Simulation Update
 		// -------------------------------------------------------------------------------------------------------------------------------------------------
-		while (Ticker.PendingKeyframe <= Ticker.MaxAllowedKeyframe)
+		while (Ticker.PendingFrame <= Ticker.MaxAllowedFrame)
 		{
-			const int32 InputKeyframe = Ticker.PendingKeyframe;
-			const int32 OutputKeyframe = Ticker.PendingKeyframe + 1;
+			const int32 InputFrame = Ticker.PendingFrame;
+			const int32 OutputFrame = Ticker.PendingFrame + 1;
 
-			const TInputCmd* InputCmd = Buffers.Input[InputKeyframe];
-			if (!ensureMsgf(InputCmd, TEXT("No InputCmd available for Keyframe %d. PendingKeyframe: %d. MaxAllowedKeyframe: %d."), InputKeyframe, Ticker.PendingKeyframe, Ticker.MaxAllowedKeyframe))
+			const TInputCmd* InputCmd = Buffers.Input[InputFrame];
+			if (!ensureMsgf(InputCmd, TEXT("No InputCmd available for Frame %d. PendingFrame: %d. MaxAllowedFrame: %d."), InputFrame, Ticker.PendingFrame, Ticker.MaxAllowedFrame))
 			{
 				break;
 			}
@@ -152,34 +152,34 @@ public:
 			// We have an unprocessed command, do we have enough allotted simulation time to process it?
 			if (Ticker.GetRemainingAllowedSimulationTime() >= InputCmd->GetFrameDeltaTime())
 			{
-				TSyncState* InSyncState = Buffers.Sync[InputKeyframe];
+				TSyncState* InSyncState = Buffers.Sync[InputFrame];
 				if (InSyncState == nullptr)
 				{
 					// We don't have valid sync state for this frame. This could mean we skipped ahead somehow, such as heavy packet loss where the server doesn't receive all input cmds
-					// We will use the last known sync state, but first write it out to the current InputKeyframe
-					UE_LOG(LogNetworkSim, Log, TEXT("%s: No sync state found @ keyframe %d. Using previous head element @ keyframe %d."), *Driver->GetDebugName(), InputKeyframe, Buffers.Sync.HeadKeyframe());
-					InSyncState = Buffers.Sync.WriteKeyframeInitializedFromHead(InputKeyframe);
+					// We will use the last known sync state, but first write it out to the current InputFrame
+					UE_LOG(LogNetworkSim, Log, TEXT("%s: No sync state found @ frame %d. Using previous head element @ frame %d."), *Driver->GetDebugName(), InputFrame, Buffers.Sync.HeadFrame());
+					InSyncState = Buffers.Sync.WriteFrameInitializedFromHead(InputFrame);
 				}
 
-				const TAuxState* InAuxState = Buffers.Aux[InputKeyframe];
-				checkf(InAuxState, TEXT("No AuxState available for Keyframe %d. PendingKeyframe: %d. MaxAllowedKeyframe: %d."), InputKeyframe, Ticker.PendingKeyframe, Ticker.MaxAllowedKeyframe);
+				const TAuxState* InAuxState = Buffers.Aux[InputFrame];
+				checkf(InAuxState, TEXT("No AuxState available for Frame %d. PendingFrame: %d. MaxAllowedFrame: %d."), InputFrame, Ticker.PendingFrame, Ticker.MaxAllowedFrame);
 
-				TSyncState* OutSyncState = Buffers.Sync.WriteKeyframe(OutputKeyframe);
+				TSyncState* OutSyncState = Buffers.Sync.WriteFrame(OutputFrame);
 				check(OutSyncState);
 
 				{
-					TScopedSimulationTick UpdateScope(Ticker, OutputKeyframe, InputCmd->GetFrameDeltaTime());
+					TScopedSimulationTick UpdateScope(Ticker, OutputFrame, InputCmd->GetFrameDeltaTime());
 					
 					Simulation->SimulationTick( 
 						{ InputCmd->GetFrameDeltaTime(), Ticker },
 						{ *InputCmd, *InSyncState, *InAuxState },
-						{ *OutSyncState, Buffers.Aux.LazyWriter(OutputKeyframe) } );
+						{ *OutSyncState, Buffers.Aux.LazyWriter(OutputFrame) } );
 													
 				}
 
 				if (DebugState)
 				{
-					DebugState->ProcessedKeyframes.Add(InputKeyframe);
+					DebugState->ProcessedFrames.Add(InputFrame);
 				}
 			}
 			else
@@ -214,8 +214,8 @@ public:
 		// Finish debug state buffer recording (what the server processed each frame)
 		if (DebugState)
 		{
-			DebugState->PendingKeyframe = Ticker.PendingKeyframe;
-			DebugState->HeadKeyframe = Buffers.Input.HeadKeyframe();
+			DebugState->PendingFrame = Ticker.PendingFrame;
+			DebugState->HeadFrame = Buffers.Input.HeadFrame();
 			DebugState->RemainingAllowedSimulationTimeSeconds = (float)Ticker.GetRemainingAllowedSimulationTime().ToRealTimeSeconds();
 		}
 
@@ -365,14 +365,14 @@ public:
 		RepProxy_Autonomous.bDependentSimulationNeedsReconcile = true;
 	}
 
-	void BeginRollback(const FNetworkSimTime& RollbackDeltaTime, const int32 ParentKeyframe) final override
+	void BeginRollback(const FNetworkSimTime& RollbackDeltaTime, const int32 ParentFrame) final override
 	{
-		RepProxy_Simulated.template DependentRollbackBegin<TSimulation, TDriver>(Simulation, Driver, Buffers, Ticker, RollbackDeltaTime, ParentKeyframe);
+		RepProxy_Simulated.template DependentRollbackBegin<TSimulation, TDriver>(Simulation, Driver, Buffers, Ticker, RollbackDeltaTime, ParentFrame);
 	}
 
-	void StepRollback(const FNetworkSimTime& Step, const int32 ParentKeyframe, const bool bFinalStep) final override
+	void StepRollback(const FNetworkSimTime& Step, const int32 ParentFrame, const bool bFinalStep) final override
 	{
-		RepProxy_Simulated.template DependentRollbackStep<TSimulation, TDriver>(Simulation, Driver, Buffers, Ticker, Step, ParentKeyframe, bFinalStep);
+		RepProxy_Simulated.template DependentRollbackStep<TSimulation, TDriver>(Simulation, Driver, Buffers, Ticker, Step, ParentFrame, bFinalStep);
 	}
 
 	void ClearAllDependentSimulations()
@@ -439,7 +439,7 @@ public:
 
 #if NETSIM_MODEL_DEBUG
 	TDebugBuffer* GetLocalDebugBuffer() {	return &Buffers.Debug; }
-	TDebugState* GetNextLocalDebugStateWrite() { return Buffers.Debug.WriteKeyframe( Buffers.Debug.HeadKeyframe() + 1 ); }
+	TDebugState* GetNextLocalDebugStateWrite() { return Buffers.Debug.WriteFrame( Buffers.Debug.HeadFrame() + 1 ); }
 	TNetworkSimBufferContainer<TBufferTypes>* GetHistoricBuffers(bool bCreate=false)
 	{
 		if (HistoricBuffers.IsValid() == false && bCreate) { HistoricBuffers.Reset(new TNetworkSimBufferContainer<TBufferTypes>()); }
diff --git a/Engine/Plugins/Runtime/NetworkPrediction/Source/NetworkPrediction/Public/NetworkSimulationModelBuffer.h b/Engine/Plugins/Runtime/NetworkPrediction/Source/NetworkPrediction/Public/NetworkSimulationModelBuffer.h
--- a/Engine/Plugins/Runtime/NetworkPrediction/Source/NetworkPrediction/Public/NetworkSimulationModelBuffer.h
+++ b/Engine/Plugins/Runtime/NetworkPrediction/Source/NetworkPrediction/Public/NetworkSimulationModelBuffer.h
@@ -2,12 +2,12 @@
 
 #pragma once
 
-// Generic iterator for iterating over each keyframe from tail to head
+// Generic iterator for iterating over each frame from tail to head
 template<typename TBuffer, typename ElementType>
 struct TNetworkSimBufferIterator
 {
 	TNetworkSimBufferIterator(TBuffer& InBuffer) : Buffer(InBuffer)	{ CurrentPos = Buffer.IterStartPos(); }
-	int32 Keyframe() const { return Buffer.IterKeyframe(CurrentPos); }
+	int32 Frame() const { return Buffer.IterFrame(CurrentPos); }
 	ElementType* Element() const { return Buffer.IterElement(CurrentPos); }
 
 	void operator++() { CurrentPos++; }
@@ -23,9 +23,9 @@ template<typename ElementType>
 struct TNetSimLazyWriterFunc
 {
 	template<typename TBuffer>
-	TNetSimLazyWriterFunc(TBuffer& Buffer, int32 PendingKeyframe)
+	TNetSimLazyWriterFunc(TBuffer& Buffer, int32 PendingFrame)
 	{
-		GetFunc = [&Buffer, PendingKeyframe]() { return Buffer.WriteKeyframeInitializedFromHead(PendingKeyframe); };
+		GetFunc = [&Buffer, PendingFrame]() { return Buffer.WriteFrameInitializedFromHead(PendingFrame); };
 	}
 
 	ElementType* Get() const
@@ -59,12 +59,12 @@ struct TNetworkSimBufferBase
 {
 	FString GetBasicDebugStr() const
 	{
-		return FString::Printf(TEXT("Elements: [%d/%d]. Keyframes: [%d-%d]"), ((T*)this)->Num(), ((T*)this)->Max(), ((T*)this)->TailKeyframe(), ((T*)this)->HeadKeyframe());
+		return FString::Printf(TEXT("Elements: [%d/%d]. Frames: [%d-%d]"), ((T*)this)->Num(), ((T*)this)->Max(), ((T*)this)->TailFrame(), ((T*)this)->HeadFrame());
 	}
 
-	bool IsValidKeyframe(int32 Keyframe) const 
+	bool IsValidFrame(int32 Frame) const 
 	{ 
-		return Keyframe >= ((T*)this)->TailKeyframe() && Keyframe <= ((T*)this)->HeadKeyframe(); 
+		return Frame >= ((T*)this)->TailFrame() && Frame <= ((T*)this)->HeadFrame(); 
 	}
 
 	// Copies each element of SourceBuffer into this buffer. This could break continuity so for TNetworkSimContiguousBuffer, it is possible to lose elements by doing this.
@@ -73,35 +73,35 @@ struct TNetworkSimBufferBase
 	{
 		for (auto It = SourceBuffer.CreateConstIterator(); It; ++It)
 		{
-			*((T*)this)->WriteKeyframe( It.Keyframe() ) = *It.Element();
+			*((T*)this)->WriteFrame( It.Frame() ) = *It.Element();
 		}
 	}
 
 protected:
 
-	// Creates a new keyframe, but
-	//	-If keyframe already exists, returns existing
-	//	-If keyframe > head, contents of head are copied into new frame
-	//	-If keyframe < tail, contents are cleared to default value
+	// Creates a new frame, but
+	//	-If frame already exists, returns existing
+	//	-If frame > head, contents of head are copied into new frame
+	//	-If frame < tail, contents are cleared to default value
 	// (this function is sketchy and poorly named but is proving to be useful in a few places)
 	template<typename ElementType>
-	ElementType* WriteKeyframeInitializedFromHeadImpl(int32 Keyframe)
+	ElementType* WriteFrameInitializedFromHeadImpl(int32 Frame)
 	{
-		if (((T*)this)->HeadKeyframe() > Keyframe)
+		if (((T*)this)->HeadFrame() > Frame)
 		{
 			ElementType* HeadElementPtr = ((T*)this)->HeadElement();
-			ElementType* NewElement = ((T*)this)->WriteKeyframe(Keyframe);
+			ElementType* NewElement = ((T*)this)->WriteFrame(Frame);
 			*NewElement = *HeadElementPtr;
 			return NewElement;
 		}
-		else if (Keyframe < ((T*)this)->TailKeyframe())
+		else if (Frame < ((T*)this)->TailFrame())
 		{
-			ElementType* NewElement = ((T*)this)->WriteKeyframe(Keyframe);
+			ElementType* NewElement = ((T*)this)->WriteFrame(Frame);
 			*NewElement = ElementType();
 			return NewElement;
 		}
 		
-		return ((T*)this)->WriteKeyframe(Keyframe);
+		return ((T*)this)->WriteFrame(Frame);
 	}
 };
 
@@ -114,11 +114,11 @@ struct TNetworkSimContiguousBuffer : public TNetworkSimBufferBase<TNetworkSimCon
 		Data.SetNum(NumElements);
 	}
 
-	T* operator[](int32 Keyframe) {	return const_cast<T*>(GetImpl(Keyframe)); }
-	const T* operator[](int32 Keyframe) const { return GetImpl(Keyframe); }
+	T* operator[](int32 Frame) {	return const_cast<T*>(GetImpl(Frame)); }
+	const T* operator[](int32 Frame) const { return GetImpl(Frame); }
 
-	T* Get(int32 Keyframe) { return const_cast<T*>(GetImpl(Keyframe)); }
-	const T* Get(int32 Keyframe) const { return GetImpl(Keyframe); }
+	T* Get(int32 Frame) { return const_cast<T*>(GetImpl(Frame)); }
+	const T* Get(int32 Frame) const { return GetImpl(Frame); }
 
 	const T* HeadElement() const { return HeadElementImpl(); }
 	T* HeadElement() { return const_cast<T*>(HeadElementImpl()); }
@@ -126,19 +126,19 @@ struct TNetworkSimContiguousBuffer : public TNetworkSimBufferBase<TNetworkSimCon
 	const T* TailElement() const { return TailElementImpl(); }
 	T* TailElement() { return const_cast<T*>(TailElementImpl()); }
 
-	int32 HeadKeyframe() const { return Head; }
-	int32 TailKeyframe() const { return GetTail(); }
+	int32 HeadFrame() const { return Head; }
+	int32 TailFrame() const { return GetTail(); }
 
 	int32 Num() const { return Head == INDEX_NONE ? 0 : (Head - GetTail() + 1); }
 	int32 Max() const { return Data.Max(); }
 	int32 GetDirtyCount() const { return DirtyCount; }
 
-	T* WriteKeyframe(int32 Keyframe)
+	T* WriteFrame(int32 Frame)
 	{
-		check(Keyframe >= 0);
+		check(Frame >= 0);
 
 		const int32 Tail = GetTail();
-		if (Head == INDEX_NONE || Keyframe < Tail || Keyframe > Head+1)
+		if (Head == INDEX_NONE || Frame < Tail || Frame > Head+1)
 		{
 			// Writing outside the current range (+1) results in a full wipe of valid contents
 			NumValidElements = 1;
@@ -146,22 +146,22 @@ struct TNetworkSimContiguousBuffer : public TNetworkSimBufferBase<TNetworkSimCon
 		else
 		{
 			// Writing inside the current range (+1) preserves valid elements up to the size of our buffer
-			NumValidElements = FMath::Min<int32>(Keyframe - Tail + 1, Data.Num());
+			NumValidElements = FMath::Min<int32>(Frame - Tail + 1, Data.Num());
 		}
 		
-		Head = Keyframe;
+		Head = Frame;
 		++DirtyCount;
-		return &Data[Keyframe % Data.Num()];
+		return &Data[Frame % Data.Num()];
 	}
 	
-	T* WriteKeyframeInitializedFromHead(int32 Keyframe)
+	T* WriteFrameInitializedFromHead(int32 Frame)
 	{
-		return this->template WriteKeyframeInitializedFromHeadImpl<ElementType>(Keyframe);
+		return this->template WriteFrameInitializedFromHeadImpl<ElementType>(Frame);
 	}
 
-	TNetSimLazyWriterFunc<ElementType> LazyWriter(int32 Keyframe)
+	TNetSimLazyWriterFunc<ElementType> LazyWriter(int32 Frame)
 	{
-		return TNetSimLazyWriterFunc<ElementType>(*this, Keyframe);
+		return TNetSimLazyWriterFunc<ElementType>(*this, Frame);
 	}
 
 	using IteratorType = TNetworkSimBufferIterator<TNetworkSimContiguousBuffer<ElementType, NumElements>, ElementType>;
@@ -172,15 +172,15 @@ struct TNetworkSimContiguousBuffer : public TNetworkSimBufferBase<TNetworkSimCon
 
 private:
 
-	const T* GetImpl(int32 Keyframe) const
+	const T* GetImpl(int32 Frame) const
 	{
-		const int32 RelativeToHead = Keyframe - Head;
+		const int32 RelativeToHead = Frame - Head;
 		if (RelativeToHead > 0 || RelativeToHead <= -NumValidElements)
 		{
 			return nullptr;
 		}
 
-		return &Data[Keyframe % Data.Num()];
+		return &Data[Frame % Data.Num()];
 	}
 
 	const T* HeadElementImpl() const
@@ -214,12 +214,12 @@ private:
 	friend struct TNetworkSimBufferIterator;
 	int32 IterStartPos() const { return GetTail(); }
 	int32 IterEndPos() const { return Head; }
-	int32 IterKeyframe(int32 Pos) const { return Pos; }
+	int32 IterFrame(int32 Pos) const { return Pos; }
 	T* IterElement(int32 Pos) { return &Data[Pos % Data.Num()]; }
 	const T* IterElement(int32 Pos) const { return &Data[Pos % Data.Num()]; }
 };
 
-// Keyframe: arbitrary identifier for data. Not contiguous or controlled by us (always passed in)
+// Frame: arbitrary identifier for data. Not contiguous or controlled by us (always passed in)
 // Position/Pos: increasing counter for position in array. Mod with Data.Num to get Index.
 // Index/Idx: actual index into Data Array
 template<typename T, int32 NumElements=32>
@@ -231,11 +231,11 @@ struct TNetworkSimSparseBuffer : public TNetworkSimBufferBase<TNetworkSimSparseB
 		Data.SetNum(NumElements);
 	}
 
-	T* operator[](int32 Keyframe) {	return const_cast<T*>(GetImpl(Keyframe)); }
-	const T* operator[](int32 Keyframe) const { return GetImpl(Keyframe); }
+	T* operator[](int32 Frame) {	return const_cast<T*>(GetImpl(Frame)); }
+	const T* operator[](int32 Frame) const { return GetImpl(Frame); }
 
-	T* Get(int32 Keyframe) { return const_cast<T*>(GetImpl(Keyframe)); }
-	const T* Get(int32 Keyframe) const { return GetImpl(Keyframe); }
+	T* Get(int32 Frame) { return const_cast<T*>(GetImpl(Frame)); }
+	const T* Get(int32 Frame) const { return GetImpl(Frame); }
 
 	const T* HeadElement() const { return HeadElementImpl(); }
 	T* HeadElement() { return const_cast<T*>(HeadElementImpl()); }
@@ -243,32 +243,32 @@ struct TNetworkSimSparseBuffer : public TNetworkSimBufferBase<TNetworkSimSparseB
 	const T* TailElement() const { return TailElementImpl(); }
 	T* TailElement() { return const_cast<T*>(TailElementImpl()); }
 
-	int32 HeadKeyframe() const { return Data[HeadPos % Data.Num()].Keyframe; }
-	int32 TailKeyframe() const { return Data[GetTailPos() % Data.Num()].Keyframe; }
+	int32 HeadFrame() const { return Data[HeadPos % Data.Num()].Frame; }
+	int32 TailFrame() const { return Data[GetTailPos() % Data.Num()].Frame; }
 
 	int32 Num() const { return HeadPos == INDEX_NONE ? 0 : (HeadPos - GetTailPos() + 1); }
 	int32 Max() const { return Data.Max(); }
 	int32 GetDirtyCount() const { return DirtyCount; }
 
-	// Returns the element @ keyframe for writing. The contents of this element are unknown (could be stale content!). Note the element returned is immediately considered "valid" by Num(), Iterators, etc!
-	T* WriteKeyframe(int32 Keyframe)
+	// Returns the element @ frame for writing. The contents of this element are unknown (could be stale content!). Note the element returned is immediately considered "valid" by Num(), Iterators, etc!
+	T* WriteFrame(int32 Frame)
 	{
-		check(Keyframe >= 0);
+		check(Frame >= 0);
 		++DirtyCount;
 
 		int32 Pos = HeadPos;
 		if (HeadPos != INDEX_NONE)
 		{
-			// If we have elements, find where this Keyframe would go
+			// If we have elements, find where this Frame would go
 			const int32 Tail = GetTailPos();
 			for (; Pos >= Tail; --Pos)
 			{
 				TInternal& InternalData = Data[Pos % Data.Num()];
-				if (InternalData.Keyframe == Keyframe)
+				if (InternalData.Frame == Frame)
 				{
 					return &InternalData.Element;
 				}
-				if (InternalData.Keyframe <= Keyframe)
+				if (InternalData.Frame <= Frame)
 				{
 					break;
 				}
@@ -278,18 +278,18 @@ struct TNetworkSimSparseBuffer : public TNetworkSimBufferBase<TNetworkSimSparseB
 		// Write frame to next pos
 		HeadPos = Pos+1;
 		TInternal& WriteData = Data[HeadPos % Data.Num()];
-		WriteData.Keyframe = Keyframe;
+		WriteData.Frame = Frame;
 		return &WriteData.Element;
 	}
 
-	ElementType* WriteKeyframeInitializedFromHead(int32 Keyframe)
+	ElementType* WriteFrameInitializedFromHead(int32 Frame)
 	{
-		return this->template WriteKeyframeInitializedFromHeadImpl<ElementType>(Keyframe);
+		return this->template WriteFrameInitializedFromHeadImpl<ElementType>(Frame);
 	}
 
-	TNetSimLazyWriterFunc<ElementType> LazyWriter(int32 Keyframe)
+	TNetSimLazyWriterFunc<ElementType> LazyWriter(int32 Frame)
 	{
-		return TNetSimLazyWriterFunc<ElementType>(*this, Keyframe);
+		return TNetSimLazyWriterFunc<ElementType>(*this, Frame);
 	}
 
 	using IteratorType = TNetworkSimBufferIterator<TNetworkSimSparseBuffer<ElementType, NumElements>, ElementType>;
@@ -302,13 +302,13 @@ private:
 
 	struct TInternal
 	{
-		int32 Keyframe = -1;
+		int32 Frame = -1;
 		T Element;
 	};
 
-	const T* GetImpl(int32 Keyframe) const
+	const T* GetImpl(int32 Frame) const
 	{
-		const int32 Pos = GetPosForKeyframe(Keyframe);
+		const int32 Pos = GetPosForFrame(Frame);
 		if (Pos != INDEX_NONE)
 		{		
 			return &Data[Pos % Data.Num()].Element;
@@ -336,12 +336,12 @@ private:
 		return nullptr;
 	}
 	
-	int32 GetPosForKeyframe(int32 Keyframe) const
+	int32 GetPosForFrame(int32 Frame) const
 	{
 		const int32 TailPos = GetTailPos();
 		for (int32 Pos=HeadPos; Pos >= TailPos; --Pos)
 		{
-			if (Data[Pos % Data.Num()].Keyframe <= Keyframe)
+			if (Data[Pos % Data.Num()].Frame <= Frame)
 			{
 				return Pos;
 			}
@@ -360,7 +360,7 @@ private:
 	friend struct TNetworkSimBufferIterator;
 	int32 IterStartPos() const { return GetTailPos(); }
 	int32 IterEndPos() const { return HeadPos; }
-	int32 IterKeyframe(int32 Pos) const { return Data[Pos % Data.Num()].Keyframe; }
+	int32 IterFrame(int32 Pos) const { return Data[Pos % Data.Num()].Frame; }
 	T* IterElement(int32 Pos) { return &Data[Pos % Data.Num()].Element; }
 	const T* IterElement(int32 Pos) const { return &Data[Pos % Data.Num()].Element; }
 };
\ No newline at end of file
diff --git a/Engine/Plugins/Runtime/NetworkPrediction/Source/NetworkPrediction/Public/NetworkSimulationModelDebugger.h b/Engine/Plugins/Runtime/NetworkPrediction/Source/NetworkPrediction/Public/NetworkSimulationModelDebugger.h
--- a/Engine/Plugins/Runtime/NetworkPrediction/Source/NetworkPrediction/Public/NetworkSimulationModelDebugger.h
+++ b/Engine/Plugins/Runtime/NetworkPrediction/Source/NetworkPrediction/Public/NetworkSimulationModelDebugger.h
@@ -21,7 +21,7 @@ DEFINE_LOG_CATEGORY_STATIC(LogNetworkSimDebug, Log, All);
 
 namespace NetworkSimulationModelDebugCVars
 {
-	NETSIM_DEVCVAR_SHIPCONST_INT(DrawKeyframes, 1, "nsm.debug.DrawKeyFrames", "Draws keyframe data (text) in debug graphs");
+	NETSIM_DEVCVAR_SHIPCONST_INT(DrawFrames, 1, "nsm.debug.DrawFrames", "Draws frame data (text) in debug graphs");
 	NETSIM_DEVCVAR_SHIPCONST_INT(DrawNetworkSendLines, 1, "nsm.debug.DrawNetworkSendLines", "Draws lines representing network traffic in debugger");
 	NETSIM_DEVCVAR_SHIPCONST_INT(GatherServerSidePIE, 1, "nsm.debug.GatherServerSide", "Whenever we gather debug info from a client side actor, also gather server side equivelent. Only works in PIE.");
 }
@@ -145,7 +145,7 @@ struct NETWORKPREDICTION_API FNetworkSimulationModelDebuggerManager: public FTic
 			C->DrawItem(*Item.Get());
 		}
 
-		if (NetworkSimulationModelDebugCVars::DrawKeyframes() > 0)
+		if (NetworkSimulationModelDebugCVars::DrawFrames() > 0)
 		{
 			for (auto& Item : CanvasItems[1])
 			{
@@ -207,7 +207,7 @@ struct NETWORKPREDICTION_API FNetworkSimulationModelDebuggerManager: public FTic
 		FStandardLoggingParameters LocalParameters = Parameters;
 		LocalParameters.Ar = &StrOut;
 
-		auto* Element = Buffer[LocalParameters.Keyframe];
+		auto* Element = Buffer[LocalParameters.Frame];
 		if (Element)
 		{
 			Element->Log(LocalParameters);
@@ -361,7 +361,7 @@ struct TNetworkSimulationModelDebugger : public INetworkSimulationModelDebugger
 			FVector2D RecvPosition = FVector2D::ZeroVector;
 		};
 
-		TMap<int32, FScreenPositions> Keyframes;
+		TMap<int32, FScreenPositions> Frames;
 	};
 
 	template<typename TBuffer>
@@ -413,13 +413,13 @@ struct TNetworkSimulationModelDebugger : public INetworkSimulationModelDebugger
 				Out.EmitQuad(FVector2D( ScreenX, ScreenY), FVector2D( FixedWidth, FrameHeight), Color);
 				Out.EmitText(FVector2D( ScreenX, ScreenY), FColor::Black, FString::Printf(TEXT("%.2f"), (DebugState->LocalDeltaTimeSeconds * 1000.f)));
 
-				// Processed InputcmdsKeyframes (above baseline)
+				// Processed InputcmdsFrames (above baseline)
 				float ClientX = ScreenX;
 				float ClientY = ScreenY - Pad;
 
-				for (int32 Keyframe : DebugState->ProcessedKeyframes)
+				for (int32 Frame : DebugState->ProcessedFrames)
 				{
-					auto* Cmd = InputBuffer[Keyframe];
+					auto* Cmd = InputBuffer[Frame];
 					if (Cmd)
 					{
 						float ClientSizeX = FixedWidth;
@@ -427,33 +427,33 @@ struct TNetworkSimulationModelDebugger : public INetworkSimulationModelDebugger
 
 						FVector2D ScreenPos(ClientX, ClientY - ClientSizeY);
 						Out.EmitQuad(ScreenPos, FVector2D( ClientSizeX, ClientSizeY), FColor::Blue);
-						Out.EmitText(ScreenPos, FColor::White, LexToString(Keyframe));
+						Out.EmitText(ScreenPos, FColor::White, LexToString(Frame));
 						ClientY -= (ClientSizeY + Pad);
 					}
 				}
 
 				// Unprocessed InputCmds (above processed)				
-				for (int32 Keyframe = DebugState->PendingKeyframe; Keyframe <= DebugState->HeadKeyframe; ++Keyframe)
+				for (int32 Frame = DebugState->PendingFrame; Frame <= DebugState->HeadFrame; ++Frame)
 				{
-					if (auto* Cmd = InputBuffer[Keyframe])
+					if (auto* Cmd = InputBuffer[Frame])
 					{
 						float ClientSizeX = FixedWidth;
 						float ClientSizeY = SecondsToPixelsY * Cmd->GetFrameDeltaTime().ToRealTimeSeconds();
 						FVector2D ScreenPos(ClientX, ClientY - ClientSizeY);
 						Out.EmitQuad(ScreenPos, FVector2D( ClientSizeX, ClientSizeY), FColor::Red);
-						Out.EmitText(ScreenPos, FColor::White, LexToString(Keyframe));
+						Out.EmitText(ScreenPos, FColor::White, LexToString(Frame));
 						ClientY -= (ClientSizeY + Pad);
 					}
 				}
 
-				// Cache Screen Positions based on keyframe
-				RecvCache.Keyframes.FindOrAdd(DebugState->LastReceivedInputKeyframe).SetRecv(FVector2D(ScreenX, BaseLineYPos));
+				// Cache Screen Positions based on frame
+				RecvCache.Frames.FindOrAdd(DebugState->LastReceivedInputFrame).SetRecv(FVector2D(ScreenX, BaseLineYPos));
 				
 				// Advance 
 				ScreenX += FixedWidth + Pad;
 
 				// Send cache
-				SendCache.Keyframes.FindOrAdd(DebugState->LastSentInputKeyframe).SetSent(FVector2D(ScreenX, BaseLineYPos));
+				SendCache.Frames.FindOrAdd(DebugState->LastSentInputFrame).SetSent(FVector2D(ScreenX, BaseLineYPos));
 			}
 
 
@@ -500,17 +500,17 @@ struct TNetworkSimulationModelDebugger : public INetworkSimulationModelDebugger
 			float ColumnTime = 0.f;
 
 			auto* DebugState = It.Element();
-			for (int32 Keyframe : DebugState->ProcessedKeyframes)
+			for (int32 Frame : DebugState->ProcessedFrames)
 			{
-				auto* Cmd = InputBuffer[Keyframe];
+				auto* Cmd = InputBuffer[Frame];
 				if (Cmd)
 				{
 					ColumnTime += (float)Cmd->GetFrameDeltaTime().ToRealTimeSeconds();
 				}
 			}
-			for (int32 Keyframe = DebugState->PendingKeyframe; Keyframe <= DebugState->HeadKeyframe; ++Keyframe)
+			for (int32 Frame = DebugState->PendingFrame; Frame <= DebugState->HeadFrame; ++Frame)
 			{
-				if (auto* Cmd = InputBuffer[Keyframe])
+				if (auto* Cmd = InputBuffer[Frame])
 				{
 					ColumnTime += (float)Cmd->GetFrameDeltaTime().ToRealTimeSeconds();
 				}
@@ -534,27 +534,27 @@ struct TNetworkSimulationModelDebugger : public INetworkSimulationModelDebugger
 		// ------------------------------------------------------------------------------------------------------------------------------------------------
 
 		Out.Emit(FString::Printf(TEXT("%s - %s"), *Owner->GetName(), *UEnum::GetValueAsString(TEXT("Engine.ENetRole"), Owner->GetLocalRole())), FColor::Yellow);
-		Out.Emit(FString::Printf(TEXT("PendingKeyframe: %d (%d Buffered)"), NetworkSim->Ticker.PendingKeyframe, NetworkSim->Buffers.Input.HeadKeyframe() - NetworkSim->Ticker.PendingKeyframe));
+		Out.Emit(FString::Printf(TEXT("PendingFrame: %d (%d Buffered)"), NetworkSim->Ticker.PendingFrame, NetworkSim->Buffers.Input.HeadFrame() - NetworkSim->Ticker.PendingFrame));
 				
 		if (Owner->GetLocalRole() == ROLE_AutonomousProxy)
 		{			
 			FColor Color = FColor::White;
 			const bool FaultDetected = NetworkSim->RepProxy_Autonomous.IsReconcileFaultDetected();
 
-			const int32 LastSerializedKeyframe = NetworkSim->RepProxy_Autonomous.GetLastSerializedKeyframe();
+			const int32 LastSerializedFrame = NetworkSim->RepProxy_Autonomous.GetLastSerializedFrame();
 
 			// Calc how much predicted time we have processed. Note that we use the motionstate buffer to iterate but the MS is on the input cmd. (if we are buffering cmds, don't want to count them)
 			
 			TSimTime PredictedMS;
-			for (int32 PredKeyframe = LastSerializedKeyframe+1; PredKeyframe <= NetworkSim->Buffers.Sync.HeadKeyframe(); ++PredKeyframe)
+			for (int32 PredFrame = LastSerializedFrame+1; PredFrame <= NetworkSim->Buffers.Sync.HeadFrame(); ++PredFrame)
 			{
-				if (auto* Cmd = NetworkSim->Buffers.Input[PredKeyframe])
+				if (auto* Cmd = NetworkSim->Buffers.Input[PredFrame])
 				{
 					PredictedMS += Cmd->GetFrameDeltaTime();
 				}
 			}
 
-			FString ConfirmedFrameStr = FString::Printf(TEXT("LastConfirmedFrame: %d. Prediction: %d Frames, %s MS"), LastSerializedKeyframe, NetworkSim->Buffers.Sync.HeadKeyframe() - LastSerializedKeyframe, *PredictedMS.ToString());
+			FString ConfirmedFrameStr = FString::Printf(TEXT("LastConfirmedFrame: %d. Prediction: %d Frames, %s MS"), LastSerializedFrame, NetworkSim->Buffers.Sync.HeadFrame() - LastSerializedFrame, *PredictedMS.ToString());
 			if (FaultDetected)
 			{
 				ConfirmedFrameStr += TEXT(" RECONCILE FAULT DETECTED!");
@@ -567,7 +567,7 @@ struct TNetworkSimulationModelDebugger : public INetworkSimulationModelDebugger
 				*NetworkSim->RepProxy_Autonomous.GetLastSerializedSimTime().ToString(), *(NetworkSim->Ticker.GetTotalProcessedSimulationTime() - NetworkSim->RepProxy_Autonomous.GetLastSerializedSimTime()).ToString());
 			Out.Emit(*SimulationTimeString, Color);
 
-			FString AllowedSimulationTimeString = FString::Printf(TEXT("Allowed Simulation Time: %s. Keyframe: %d/%d/%d"), *NetworkSim->Ticker.GetRemainingAllowedSimulationTime().ToString(), NetworkSim->Ticker.MaxAllowedKeyframe, NetworkSim->Ticker.PendingKeyframe, NetworkSim->Buffers.Input.HeadKeyframe());
+			FString AllowedSimulationTimeString = FString::Printf(TEXT("Allowed Simulation Time: %s. Frame: %d/%d/%d"), *NetworkSim->Ticker.GetRemainingAllowedSimulationTime().ToString(), NetworkSim->Ticker.MaxAllowedFrame, NetworkSim->Ticker.PendingFrame, NetworkSim->Buffers.Input.HeadFrame());
 			Out.Emit(*AllowedSimulationTimeString, Color);
 		}
 		else if (Owner->GetLocalRole() == ROLE_SimulatedProxy)
@@ -583,7 +583,7 @@ struct TNetworkSimulationModelDebugger : public INetworkSimulationModelDebugger
 			Out.Emit(FString::Printf(TEXT("//////////////// %s ///////////////"), *BufferName), FColor::Yellow);
 			Out.Emit(FString::Printf(TEXT("%s"), *Buffer.GetBasicDebugStr()));		
 			Out.Emit();
-			Out.EmitElement(Buffer, FStandardLoggingParameters(nullptr, EStandardLoggingContext::Full, Buffer.HeadKeyframe()));
+			Out.EmitElement(Buffer, FStandardLoggingParameters(nullptr, EStandardLoggingContext::Full, Buffer.HeadFrame()));
 		};
 
 		EmitBuffer(TEXT("InputBuffer"), NetworkSim->Buffers.Input);
@@ -622,29 +622,29 @@ struct TNetworkSimulationModelDebugger : public INetworkSimulationModelDebugger
 			// Network Send/Recv lines
 			if (NetworkSimulationModelDebugCVars::DrawNetworkSendLines() > 0)
 			{
-				for (auto& It : ServerToClientCache.Keyframes)
+				for (auto& It : ServerToClientCache.Frames)
 				{
-					const int32 KeyFrame = It.Key;
+					const int32 Frame = It.Key;
 					auto& Positions = It.Value;
-					if (KeyFrame != 0 && Positions.RecvPosition != FVector2D::ZeroVector && Positions.SentPosition != FVector2D::ZeroVector)
+					if (Frame != 0 && Positions.RecvPosition != FVector2D::ZeroVector && Positions.SentPosition != FVector2D::ZeroVector)
 					{
 						Out.EmitLine(Positions.SentPosition, Positions.RecvPosition, FColor::Purple);
 
 						FVector2D TextPos = Positions.SentPosition + (0.25f * (Positions.RecvPosition - Positions.SentPosition));
-						Out.EmitText(TextPos, FColor::Purple, LexToString(KeyFrame));
+						Out.EmitText(TextPos, FColor::Purple, LexToString(Frame));
 					}
 				}
 
-				for (auto& It : ClientToServerCache.Keyframes)
+				for (auto& It : ClientToServerCache.Frames)
 				{
-					const int32 KeyFrame = It.Key;
+					const int32 Frame = It.Key;
 					auto& Positions = It.Value;
-					if (KeyFrame != 0 && Positions.RecvPosition != FVector2D::ZeroVector && Positions.SentPosition != FVector2D::ZeroVector)
+					if (Frame != 0 && Positions.RecvPosition != FVector2D::ZeroVector && Positions.SentPosition != FVector2D::ZeroVector)
 					{
 						Out.EmitLine(Positions.SentPosition, Positions.RecvPosition, FColor::Orange);
 
 						FVector2D TextPos = Positions.SentPosition + (0.25f * (Positions.RecvPosition - Positions.SentPosition));
-						Out.EmitText(TextPos, FColor::Orange, LexToString(KeyFrame));
+						Out.EmitText(TextPos, FColor::Orange, LexToString(Frame));
 					}
 				}
 			}
@@ -663,9 +663,9 @@ struct TNetworkSimulationModelDebugger : public INetworkSimulationModelDebugger
 
 		if (auto* LatestSync = NetworkSim->Buffers.Sync.HeadElement())
 		{
-			FVisualLoggingParameters VLogParams(EVisualLoggingContext::LastPredicted, NetworkSim->Buffers.Sync.HeadKeyframe(), EVisualLoggingLifetime::Transient);
-			int32 SyncKeyframe = NetworkSim->Buffers.Sync.HeadKeyframe();
-			NetworkSim->Driver->VisualLog(NetworkSim->Buffers.Input[SyncKeyframe], LatestSync, NetworkSim->Buffers.Aux[SyncKeyframe], VLogParams);
+			FVisualLoggingParameters VLogParams(EVisualLoggingContext::LastPredicted, NetworkSim->Buffers.Sync.HeadFrame(), EVisualLoggingLifetime::Transient);
+			int32 SyncFrame = NetworkSim->Buffers.Sync.HeadFrame();
+			NetworkSim->Driver->VisualLog(NetworkSim->Buffers.Input[SyncFrame], LatestSync, NetworkSim->Buffers.Aux[SyncFrame], VLogParams);
 		}
 
 		FStuff ServerPIEStuff = GetServerPIEStuff();
@@ -673,29 +673,29 @@ struct TNetworkSimulationModelDebugger : public INetworkSimulationModelDebugger
 		{
 			if (auto* ServerLatestSync = ServerPIEStuff.NetworkSim->Buffers.Sync.HeadElement())
 			{
-				FVisualLoggingParameters VLogParams(EVisualLoggingContext::CurrentServerPIE, ServerPIEStuff.NetworkSim->Buffers.Sync.HeadKeyframe(), EVisualLoggingLifetime::Transient);
-				int32 SyncKeyframe = ServerPIEStuff.NetworkSim->Buffers.Sync.HeadKeyframe();
-				NetworkSim->Driver->VisualLog(ServerPIEStuff.NetworkSim->Buffers.Input[SyncKeyframe], ServerLatestSync, ServerPIEStuff.NetworkSim->Buffers.Aux[SyncKeyframe], VLogParams);
+				FVisualLoggingParameters VLogParams(EVisualLoggingContext::CurrentServerPIE, ServerPIEStuff.NetworkSim->Buffers.Sync.HeadFrame(), EVisualLoggingLifetime::Transient);
+				int32 SyncFrame = ServerPIEStuff.NetworkSim->Buffers.Sync.HeadFrame();
+				NetworkSim->Driver->VisualLog(ServerPIEStuff.NetworkSim->Buffers.Input[SyncFrame], ServerLatestSync, ServerPIEStuff.NetworkSim->Buffers.Aux[SyncFrame], VLogParams);
 			}
 		}
 
 		if (Owner->GetLocalRole() == ROLE_AutonomousProxy)
 		{
-			for (int32 Keyframe = NetworkSim->RepProxy_Autonomous.GetLastSerializedKeyframe(); Keyframe < NetworkSim->Buffers.Sync.HeadKeyframe(); ++Keyframe)
+			for (int32 Frame = NetworkSim->RepProxy_Autonomous.GetLastSerializedFrame(); Frame < NetworkSim->Buffers.Sync.HeadFrame(); ++Frame)
 			{
-				if (auto* SyncState = NetworkSim->Buffers.Sync[Keyframe])
+				if (auto* SyncState = NetworkSim->Buffers.Sync[Frame])
 				{
-					const EVisualLoggingContext Context = (Keyframe == NetworkSim->RepProxy_Autonomous.GetLastSerializedKeyframe()) ? EVisualLoggingContext::LastConfirmed : EVisualLoggingContext::OtherPredicted;
-					FVisualLoggingParameters VLogParams(Context, NetworkSim->Buffers.Sync.HeadKeyframe(), EVisualLoggingLifetime::Transient);
-					NetworkSim->Driver->VisualLog(NetworkSim->Buffers.Input[Keyframe], SyncState, NetworkSim->Buffers.Aux[Keyframe], VLogParams);
+					const EVisualLoggingContext Context = (Frame == NetworkSim->RepProxy_Autonomous.GetLastSerializedFrame()) ? EVisualLoggingContext::LastConfirmed : EVisualLoggingContext::OtherPredicted;
+					FVisualLoggingParameters VLogParams(Context, NetworkSim->Buffers.Sync.HeadFrame(), EVisualLoggingLifetime::Transient);
+					NetworkSim->Driver->VisualLog(NetworkSim->Buffers.Input[Frame], SyncState, NetworkSim->Buffers.Aux[Frame], VLogParams);
 				}
 			}
 		}
 		else if (Owner->GetLocalRole() == ROLE_SimulatedProxy)
 		{
-			FVisualLoggingParameters VLogParams(EVisualLoggingContext::LastConfirmed, NetworkSim->Buffers.Sync.HeadKeyframe(), EVisualLoggingLifetime::Transient);
-			int32 HeadKeyframe = NetworkSim->Buffers.Sync.HeadKeyframe();
-			NetworkSim->Driver->VisualLog(NetworkSim->Buffers.Input[HeadKeyframe], NetworkSim->Buffers.Sync[HeadKeyframe], NetworkSim->Buffers.Aux[HeadKeyframe], VLogParams);
+			FVisualLoggingParameters VLogParams(EVisualLoggingContext::LastConfirmed, NetworkSim->Buffers.Sync.HeadFrame(), EVisualLoggingLifetime::Transient);
+			int32 HeadFrame = NetworkSim->Buffers.Sync.HeadFrame();
+			NetworkSim->Driver->VisualLog(NetworkSim->Buffers.Input[HeadFrame], NetworkSim->Buffers.Sync[HeadFrame], NetworkSim->Buffers.Aux[HeadFrame], VLogParams);
 			
 			if (NetworkSim->GetSimulatedUpdateMode() != ESimulatedUpdateMode::Interpolate)
 			{
diff --git a/Engine/Plugins/Runtime/NetworkPrediction/Source/NetworkPrediction/Public/NetworkSimulationModelInterpolator.h b/Engine/Plugins/Runtime/NetworkPrediction/Source/NetworkPrediction/Public/NetworkSimulationModelInterpolator.h
--- a/Engine/Plugins/Runtime/NetworkPrediction/Source/NetworkPrediction/Public/NetworkSimulationModelInterpolator.h
+++ b/Engine/Plugins/Runtime/NetworkPrediction/Source/NetworkPrediction/Public/NetworkSimulationModelInterpolator.h
@@ -66,7 +66,7 @@ struct TInterpolator
 		if (InterpolationTime <= 0.f)
 		{
 			InterpolationTime = SimulationTimeBuffer.TailElement()->ToRealTimeSeconds();
-			InterpolationKeyframe = SimulationTimeBuffer.TailKeyframe();
+			InterpolationFrame = SimulationTimeBuffer.TailFrame();
 			InterpolationState = *Buffers.Sync.TailElement();
 		}
 
@@ -135,7 +135,7 @@ struct TInterpolator
 			}
 		}
 
-		// Find "To" keyframe
+		// Find "To" frame
 		const TSyncState* ToState = nullptr;
 		const TAuxState* AuxState = nullptr;
 		FRealTime ToTime = 0.f;
@@ -145,10 +145,10 @@ struct TInterpolator
 			FSimTime ElementSimTime = *It.Element();
 			if (NewInterpolationTime <= ElementSimTime.ToRealTimeSeconds())
 			{
-				InterpolationKeyframe = It.Keyframe();
+				InterpolationFrame = It.Frame();
 				ToTime = ElementSimTime.ToRealTimeSeconds();
-				ToState = Buffers.Sync[It.Keyframe()];
-				AuxState = Buffers.Aux[It.Keyframe()];
+				ToState = Buffers.Sync[It.Frame()];
+				AuxState = Buffers.Aux[It.Frame()];
 				break;
 			}
 		}
@@ -197,28 +197,28 @@ struct TInterpolator
 					const TSyncState* DebugTail = Buffers.Sync.TailElement();
 					const TSyncState* DebugHead = Buffers.Sync.HeadElement();
 
-					auto VLogHelper = [&](int32 Keyframe, EVisualLoggingContext Context)
+					auto VLogHelper = [&](int32 Frame, EVisualLoggingContext Context)
 					{
-						FVisualLoggingParameters VLogParams(Context, Keyframe, EVisualLoggingLifetime::Transient);
-						Driver->VisualLog(Buffers.Input[Keyframe], Buffers.Sync[Keyframe], Buffers.Aux[Keyframe], VLogParams);
+						FVisualLoggingParameters VLogParams(Context, Frame, EVisualLoggingLifetime::Transient);
+						Driver->VisualLog(Buffers.Input[Frame], Buffers.Sync[Frame], Buffers.Aux[Frame], VLogParams);
 					};
 
-					VLogHelper(Buffers.Sync.TailKeyframe(), EVisualLoggingContext::InterpolationBufferTail);
-					VLogHelper(Buffers.Sync.HeadKeyframe(), EVisualLoggingContext::InterpolationBufferHead);
+					VLogHelper(Buffers.Sync.TailFrame(), EVisualLoggingContext::InterpolationBufferTail);
+					VLogHelper(Buffers.Sync.HeadFrame(), EVisualLoggingContext::InterpolationBufferHead);
 
 					{
-						FVisualLoggingParameters VLogParams(EVisualLoggingContext::InterpolationFrom, InterpolationKeyframe-1, EVisualLoggingLifetime::Transient);
-						Driver->VisualLog(Buffers.Input[InterpolationKeyframe-1], &InterpolationState, Buffers.Aux[InterpolationKeyframe-1], VLogParams);
+						FVisualLoggingParameters VLogParams(EVisualLoggingContext::InterpolationFrom, InterpolationFrame-1, EVisualLoggingLifetime::Transient);
+						Driver->VisualLog(Buffers.Input[InterpolationFrame-1], &InterpolationState, Buffers.Aux[InterpolationFrame-1], VLogParams);
 					}
 
 					{
-						FVisualLoggingParameters VLogParams(EVisualLoggingContext::InterpolationTo, InterpolationKeyframe, EVisualLoggingLifetime::Transient);
-						Driver->VisualLog(Buffers.Input[InterpolationKeyframe], &InterpolationState, Buffers.Aux[InterpolationKeyframe], VLogParams);
+						FVisualLoggingParameters VLogParams(EVisualLoggingContext::InterpolationTo, InterpolationFrame, EVisualLoggingLifetime::Transient);
+						Driver->VisualLog(Buffers.Input[InterpolationFrame], &InterpolationState, Buffers.Aux[InterpolationFrame], VLogParams);
 					}
 
 					{
-						FVisualLoggingParameters VLogParams(LoggingContext, InterpolationKeyframe, EVisualLoggingLifetime::Transient);
-						Driver->VisualLog(Buffers.Input[InterpolationKeyframe], &NewInterpolatedState, Buffers.Aux[InterpolationKeyframe], VLogParams);
+						FVisualLoggingParameters VLogParams(LoggingContext, InterpolationFrame, EVisualLoggingLifetime::Transient);
+						Driver->VisualLog(Buffers.Input[InterpolationFrame], &NewInterpolatedState, Buffers.Aux[InterpolationFrame], VLogParams);
 					}
 				}
 
@@ -231,7 +231,7 @@ struct TInterpolator
 private:
 
 	FRealTime InterpolationTime = 0.f; // SimTime we are currently interpolating at
-	int32 InterpolationKeyframe = INDEX_NONE; // Keyframe we are currently/last interpolated at
+	int32 InterpolationFrame = INDEX_NONE; // Frame we are currently/last interpolated at
 
 	TSyncState InterpolationState;
 
diff --git a/Engine/Plugins/Runtime/NetworkPrediction/Source/NetworkPrediction/Public/NetworkSimulationModelReplicators.h b/Engine/Plugins/Runtime/NetworkPrediction/Source/NetworkPrediction/Public/NetworkSimulationModelReplicators.h
--- a/Engine/Plugins/Runtime/NetworkPrediction/Source/NetworkPrediction/Public/NetworkSimulationModelReplicators.h
+++ b/Engine/Plugins/Runtime/NetworkPrediction/Source/NetworkPrediction/Public/NetworkSimulationModelReplicators.h
@@ -19,35 +19,35 @@ namespace NetworkSimulationModelCVars
 	NETSIM_DEVCVAR_SHIPCONST_INT(ForceReconcileSingle, 0, "ns.ForceReconcileSingle", "Forces a since reconcile to happen on the next frame");
 }
 
-#ifndef NETSIM_NETCONSTANT_NUM_BITS_KEYFRAME
-	#define NETSIM_NETCONSTANT_NUM_BITS_KEYFRAME 8	// Allows you to override this setting via UBT, but access via cleaner FActorMotionNetworkingConstants::NUM_BITS_KEYFRAME
+#ifndef NETSIM_NETCONSTANT_NUM_BITS_FRAME
+	#define NETSIM_NETCONSTANT_NUM_BITS_FRAME 8	// Allows you to override this setting via UBT, but access via cleaner FActorMotionNetworkingConstants::NUM_BITS_FRAME
 #endif
 
 struct FNetworkSimulationSerialization
 {
 	// How many bits we use to encode the key frame number for buffers.
-	// Client Keyframes are stored locally as 32 bit integers, but we use a smaller # of bits to NetSerialize.
+	// Client Frames are stored locally as 32 bit integers, but we use a smaller # of bits to NetSerialize.
 	// Frames are only relatively relevant: the absolute value doesn't really matter. We just need to detect newer/older.
-	enum { NUM_BITS_KEYFRAME = NETSIM_NETCONSTANT_NUM_BITS_KEYFRAME };		
+	enum { NUM_BITS_FRAME = NETSIM_NETCONSTANT_NUM_BITS_FRAME };		
 
 	// Abs max value we encode into the bit writer
-	enum { MAX_KEYFRAME_WRITE = 1 << NUM_BITS_KEYFRAME };
+	enum { MAX_FRAME_WRITE = 1 << NUM_BITS_FRAME };
 
-	// This is the threshold at which we would wrap around and incorrectly assign a keyframe on the receiving side.
-	// E.g, If there are KEYFRAME_ERROR_THRESHOLD keyframes that do not make it across from sender->receiver, the
+	// This is the threshold at which we would wrap around and incorrectly assign a frame on the receiving side.
+	// E.g, If there are FRAME_ERROR_THRESHOLD frames that do not make it across from sender->receiver, the
 	// receiver will have incorrect local values. With 8 bits, this works out to be 128 frames or about 2 seconds at 60fps.
-	enum { KEYFRAME_ERROR_THRESHOLD = MAX_KEYFRAME_WRITE / 2};
+	enum { FRAME_ERROR_THRESHOLD = MAX_FRAME_WRITE / 2};
 
-	// Helper to serialize the int32 HeadKeyframe. Returns the unpacked value (this will be same as input in the save path)
-	static int32 SerializeKeyframe(FArchive& Ar, int32 LocalHeadKeyframe)
+	// Helper to serialize the int32 HeadFrame. Returns the unpacked value (this will be same as input in the save path)
+	static int32 SerializeFrame(FArchive& Ar, int32 LocalHeadFrame)
 	{
 		if (Ar.IsSaving())
 		{
-			((FNetBitWriter&)Ar).WriteIntWrapped( LocalHeadKeyframe, MAX_KEYFRAME_WRITE );
-			return LocalHeadKeyframe;
+			((FNetBitWriter&)Ar).WriteIntWrapped( LocalHeadFrame, MAX_FRAME_WRITE );
+			return LocalHeadFrame;
 		}
 		
-		return MakeRelative(((FNetBitReader&)Ar).ReadInt( MAX_KEYFRAME_WRITE ), LocalHeadKeyframe, MAX_KEYFRAME_WRITE );
+		return MakeRelative(((FNetBitReader&)Ar).ReadInt( MAX_FRAME_WRITE ), LocalHeadFrame, MAX_FRAME_WRITE );
 	}
 };
 
@@ -147,7 +147,7 @@ struct TReplicator_SimTime<TBufferTypes, TTickSettings, TBase, false> : public T
 
 // Replicates a sequence of elements. i.e., "the last MaxNumElements".
 // On the receiving side, we merge the sequence into whatever we have locally
-// Keyframes are synchronized. SimTime is also serialized by default (change by changing TBase)
+// Frames are synchronized. SimTime is also serialized by default (change by changing TBase)
 template<typename TBufferTypes, typename TTickSettings, ENetworkSimBufferTypeId BufferId, int32 MaxNumElements=3, typename TBase=TReplicator_SimTime<TBufferTypes, TTickSettings>>
 struct TReplicator_Sequence : public TBase
 {
@@ -165,39 +165,39 @@ struct TReplicator_Sequence : public TBase
 		uint8 SerializedNumElements = FMath::Min<uint8>(MaxNumElements, Buffer.Num());
 		Ar << SerializedNumElements;
 
-		const int32 HeadKeyframe = FNetworkSimulationSerialization::SerializeKeyframe(Ar, Buffer.HeadKeyframe());
-		const int32 StartingKeyframe = FMath::Max(0, HeadKeyframe - SerializedNumElements + 1);
+		const int32 HeadFrame = FNetworkSimulationSerialization::SerializeFrame(Ar, Buffer.HeadFrame());
+		const int32 StartingFrame = FMath::Max(0, HeadFrame - SerializedNumElements + 1);
 
 		if (Ar.IsLoading())
 		{
-			const int32 PrevHead = Buffer.HeadKeyframe();
-			if (PrevHead < StartingKeyframe && PrevHead >= 0)
+			const int32 PrevHead = Buffer.HeadFrame();
+			if (PrevHead < StartingFrame && PrevHead >= 0)
 			{
 				// There is a gap in the stream. In some cases, we want this to be a "fault" and bubble up. We may want to synthesize state or maybe we just skip ahead.
-				UE_LOG(LogNetworkSim, Warning, TEXT("Fault: gap in received %s buffer. PrevHead: %d. Received: %d-%d. Reseting previous buffer contents"), *LexToString(BufferId), PrevHead, StartingKeyframe, HeadKeyframe);
+				UE_LOG(LogNetworkSim, Warning, TEXT("Fault: gap in received %s buffer. PrevHead: %d. Received: %d-%d. Reseting previous buffer contents"), *LexToString(BufferId), PrevHead, StartingFrame, HeadFrame);
 			}
 		}
 
-		for (int32 Keyframe = StartingKeyframe; Keyframe <= HeadKeyframe; ++Keyframe)
+		for (int32 Frame = StartingFrame; Frame <= HeadFrame; ++Frame)
 		{
 			// This, as is, is bad. The intention is that these functions serialize multiple items in some delta compressed fashion.
 			// As is, we are just serializing the elements individually.
-			auto* Cmd = Ar.IsLoading() ? Buffer.WriteKeyframe(Keyframe) : Buffer[Keyframe];
+			auto* Cmd = Ar.IsLoading() ? Buffer.WriteFrame(Frame) : Buffer[Frame];
 			Cmd->NetSerialize(P);
 		}
 
-		LastSerializedKeyframe = HeadKeyframe;
+		LastSerializedFrame = HeadFrame;
 	}
 
-	int32 GetLastSerializedKeyframe() const { return LastSerializedKeyframe; }
+	int32 GetLastSerializedFrame() const { return LastSerializedFrame; }
 
 protected:
 
-	int32 LastSerializedKeyframe = 0;
+	int32 LastSerializedFrame = 0;
 };
 
 // Replicates only the latest single element in the selected buffer.
-// Keyframe is not synchronized: the new element is just added to head.
+// Frame is not synchronized: the new element is just added to head.
 // SimTime is serialized by default (change by changing TBase)
 template<typename TBufferTypes, typename TTickSettings, ENetworkSimBufferTypeId BufferId, typename TBase=TReplicator_SimTime<TBufferTypes, TTickSettings>>
 struct TReplicator_Single : public TBase
@@ -232,25 +232,25 @@ struct TReplicator_Single : public TBase
 	}
 };
 
-/** Helper that writes a new input cmd to the input buffer, at given Keyframe (usually the sim's PendingKeyframe). If keyframe doesn't exist, ProduceInput is called on the driver. */
+/** Helper that writes a new input cmd to the input buffer, at given Frame (usually the sim's PendingFrame). If frame doesn't exist, ProduceInput is called on the driver. */
 template<typename TDriver, typename TBufferTypes>
-void GenerateLocalInputCmdAtKeyframe(TDriver* Driver, TNetworkSimBufferContainer<TBufferTypes>& Buffers, const FNetworkSimTime& DeltaSimTime, int32 Keyframe)
+void GenerateLocalInputCmdAtFrame(TDriver* Driver, TNetworkSimBufferContainer<TBufferTypes>& Buffers, const FNetworkSimTime& DeltaSimTime, int32 Frame)
 {
 	using TInputCmd = typename TBufferTypes::TInputCmd;
-	if (TInputCmd* InputCmd = Buffers.Input[Keyframe])
+	if (TInputCmd* InputCmd = Buffers.Input[Frame])
 	{
 		InputCmd->SetFrameDeltaTime(DeltaSimTime);
 	}
 	else
 	{
-		InputCmd = Buffers.Input.WriteKeyframe(Keyframe);
+		InputCmd = Buffers.Input.WriteFrame(Frame);
 		*InputCmd = TInputCmd();
 		InputCmd->SetFrameDeltaTime(DeltaSimTime);
 		Driver->ProduceInput(DeltaSimTime, *InputCmd);
 	}
 }
 
-/** Helper to generate a local input cmd if we have simulation time to spend and advance the simulation's MaxAllowedKeyframe so that it can be processed. */
+/** Helper to generate a local input cmd if we have simulation time to spend and advance the simulation's MaxAllowedFrame so that it can be processed. */
 template<typename TDriver, typename TBufferTypes, typename TTickSettings>
 void TryGenerateLocalInput(TDriver* Driver, TNetworkSimBufferContainer<TBufferTypes>& Buffers, TSimulationTicker<TTickSettings>& Ticker)
 {
@@ -259,8 +259,8 @@ void TryGenerateLocalInput(TDriver* Driver, TNetworkSimBufferContainer<TBufferTy
 	FNetworkSimTime DeltaSimTime = Ticker.GetRemainingAllowedSimulationTime();
 	if (DeltaSimTime.IsPositive())
 	{
-		GenerateLocalInputCmdAtKeyframe(Driver, Buffers, DeltaSimTime, Ticker.PendingKeyframe);
-		Ticker.MaxAllowedKeyframe = Ticker.PendingKeyframe;
+		GenerateLocalInputCmdAtFrame(Driver, Buffers, DeltaSimTime, Ticker.PendingFrame);
+		Ticker.MaxAllowedFrame = Ticker.PendingFrame;
 	}
 }
 
@@ -280,13 +280,13 @@ struct TReplicator_Server : public TBase
 		// After receiving input, server may process up to the latest received frames.
 		// (If we needed to buffer input server side for whatever reason, we would do it here)
 		// (also note that we will implicitly guard against speed hacks in the core update loop by not processing cmds past what we have been "allowed")
-		Ticker.MaxAllowedKeyframe = Buffers.Input.HeadKeyframe();
+		Ticker.MaxAllowedFrame = Buffers.Input.HeadFrame();
 
 		// Check for gaps in commands
-		if ( Ticker.PendingKeyframe < Buffers.Input.TailKeyframe() )
+		if ( Ticker.PendingFrame < Buffers.Input.TailFrame() )
 		{
-			UE_LOG(LogNetworkSim, Warning, TEXT("TReplicator_Server::Reconcile missing inputcmds. PendingKeyframe: %d. %s. This can happen via packet loss"), Ticker.PendingKeyframe, *Buffers.Input.GetBasicDebugStr());
-			Ticker.PendingKeyframe = Buffers.Input.TailKeyframe();
+			UE_LOG(LogNetworkSim, Warning, TEXT("TReplicator_Server::Reconcile missing inputcmds. PendingFrame: %d. %s. This can happen via packet loss"), Ticker.PendingFrame, *Buffers.Input.GetBasicDebugStr());
+			Ticker.PendingFrame = Buffers.Input.TailFrame();
 		}
 	}
 
@@ -363,7 +363,7 @@ struct TReplicator_Simulated : public TBase
 		}
 		else
 		{
-			check(Ticker.SimulationTimeBuffer.HeadKeyframe() == Buffers.Sync.HeadKeyframe());
+			check(Ticker.SimulationTimeBuffer.HeadFrame() == Buffers.Sync.HeadFrame());
 			check(Ticker.GetTotalProcessedSimulationTime() <= Ticker.GetTotalAllowedSimulationTime());
 
 			// Cache off our "starting" time before possibly overwriting it. We will use this in Reconcile to catch back up in some cases.
@@ -373,25 +373,25 @@ struct TReplicator_Simulated : public TBase
 			}
 
 			// Find out where this should go in the local buffer based on the serialized time
-			int32 DestinationKeyframe = INDEX_NONE;
+			int32 DestinationFrame = INDEX_NONE;
 			if (LastSerializedSimulationTime > Ticker.GetTotalProcessedSimulationTime())
 			{
 				// We are getting new state that is ahead of what we have locally, so it can safety go right to head
-				DestinationKeyframe = Ticker.SimulationTimeBuffer.HeadKeyframe()+1;
+				DestinationFrame = Ticker.SimulationTimeBuffer.HeadFrame()+1;
 			}
 			else
 			{
 				// We are getting state that is behind what we have locally
-				for (int32 Keyframe = Ticker.SimulationTimeBuffer.HeadKeyframe(); Keyframe >= Ticker.SimulationTimeBuffer.TailKeyframe(); --Keyframe)
+				for (int32 Frame = Ticker.SimulationTimeBuffer.HeadFrame(); Frame >= Ticker.SimulationTimeBuffer.TailFrame(); --Frame)
 				{
-					if (LastSerializedSimulationTime > *Ticker.SimulationTimeBuffer[Keyframe])
+					if (LastSerializedSimulationTime > *Ticker.SimulationTimeBuffer[Frame])
 					{
-						DestinationKeyframe = Keyframe+1;
+						DestinationFrame = Frame+1;
 						break;
 					}
 				}
 
-				if (DestinationKeyframe == INDEX_NONE)
+				if (DestinationFrame == INDEX_NONE)
 				{
 					FNetworkSimTime TotalTimeAhead = Ticker.GetTotalProcessedSimulationTime() - LastSerializedSimulationTime;
 					FNetworkSimTime SerializeDelta = LastSerializedSimulationTime - PrevLastSerializedSimulationTime;
@@ -399,19 +399,19 @@ struct TReplicator_Simulated : public TBase
 					// We are way far ahead of the server... we will need to clear out sync buffers, take what they gave us, and catch up in reconcile
 					//UE_LOG(LogNetworkSim, Warning, TEXT("!!! TReplicator_Simulated. Large gap detected. SerializedTime: %s. Buffer time: [%s-%s]. %d Elements. DeltaFromHead: %s. DeltaSerialize: %s"), *LastSerializedSimulationTime.ToString(), 
 					//	*Ticker.SimulationTimeBuffer.GetElementFromTail(0)->ToString(), *Ticker.SimulationTimeBuffer.GetElementFromHead(0)->ToString(), Ticker.SimulationTimeBuffer.GetNumValidElements(), *TotalTimeAhead.ToString(), *SerializeDelta.ToString());
-					DestinationKeyframe = Ticker.SimulationTimeBuffer.HeadKeyframe()+2; // (Skip ahead 2 to force a break in continuity)
+					DestinationFrame = Ticker.SimulationTimeBuffer.HeadFrame()+2; // (Skip ahead 2 to force a break in continuity)
 				}
 			}
 
-			check(DestinationKeyframe != INDEX_NONE);
+			check(DestinationFrame != INDEX_NONE);
 
 			// "Finalize" our buffers and time keeping such that we serialize the latest state from the server in the right spot
-			SyncState = Buffers.Sync.WriteKeyframe(DestinationKeyframe);
-			AuxState = Buffers.Aux.WriteKeyframe(DestinationKeyframe);
-			Buffers.Input.WriteKeyframe(DestinationKeyframe);
+			SyncState = Buffers.Sync.WriteFrame(DestinationFrame);
+			AuxState = Buffers.Aux.WriteFrame(DestinationFrame);
+			Buffers.Input.WriteFrame(DestinationFrame);
 
 			// Update tick info
-			Ticker.SetTotalProcessedSimulationTime(LastSerializedSimulationTime, DestinationKeyframe);
+			Ticker.SetTotalProcessedSimulationTime(LastSerializedSimulationTime, DestinationFrame);
 			if (Ticker.GetTotalAllowedSimulationTime() < LastSerializedSimulationTime)
 			{
 				Ticker.SetTotalAllowedSimulationTime(LastSerializedSimulationTime);
@@ -419,8 +419,8 @@ struct TReplicator_Simulated : public TBase
 
 			check(Ticker.GetTotalProcessedSimulationTime() <= Ticker.GetTotalAllowedSimulationTime());
 
-			Ticker.PendingKeyframe = DestinationKeyframe;	// We are about to serialize state to DestinationKeyframe which will be "unprocessed" (has not been used to generate a new frame)
-			Ticker.MaxAllowedKeyframe = DestinationKeyframe-1; // Do not process PendingKeyframe on our account. ::PreSimTick will advance this based on our interpolation/extrapolation settings
+			Ticker.PendingFrame = DestinationFrame;	// We are about to serialize state to DestinationFrame which will be "unprocessed" (has not been used to generate a new frame)
+			Ticker.MaxAllowedFrame = DestinationFrame-1; // Do not process PendingFrame on our account. ::PreSimTick will advance this based on our interpolation/extrapolation settings
 		}
 
 		check(SyncState && AuxState);
@@ -448,9 +448,9 @@ struct TReplicator_Simulated : public TBase
 		{
 			// Simulated Reconcile requires the input buffer to be kept up to date with the Sync buffer
 			// Generate a new, fake, command since we just added a new sync state to head
-			while (Buffers.Input.HeadKeyframe() < Buffers.Sync.HeadKeyframe())
+			while (Buffers.Input.HeadFrame() < Buffers.Sync.HeadFrame())
 			{
-				TInputCmd* Next = Buffers.Input.WriteKeyframeInitializedFromHead(Buffers.Input.HeadKeyframe()+1);
+				TInputCmd* Next = Buffers.Input.WriteFrameInitializedFromHead(Buffers.Input.HeadFrame()+1);
 			}
 
 			// Do we have time to make up? We may have extrapolated ahead of the server (totally fine - can happen with small amount of latency variance)
@@ -503,31 +503,31 @@ struct TReplicator_Simulated : public TBase
 	}
 
 	template<typename TSimulation, typename TDriver>
-	void DependentRollbackBegin(TSimulation* Simulation, TDriver* Driver, TNetworkSimBufferContainer<TBufferTypes>& Buffers, TSimulationTicker<TTickSettings>& Ticker, const FNetworkSimTime& RollbackDeltaTime, const int32 ParentKeyframe)
+	void DependentRollbackBegin(TSimulation* Simulation, TDriver* Driver, TNetworkSimBufferContainer<TBufferTypes>& Buffers, TSimulationTicker<TTickSettings>& Ticker, const FNetworkSimTime& RollbackDeltaTime, const int32 ParentFrame)
 	{
 		// For now, we make the assumption that our last serialized state and time match the parent simulation.
 		// This would not always be the case with low frequency simulated proxies. But could be handled by replicating the simulations together (at the replication level)
-		const int32 NewHeadKeyframe = Buffers.Sync.HeadKeyframe()+1;
+		const int32 NewHeadFrame = Buffers.Sync.HeadFrame()+1;
 		
-		Ticker.SetTotalProcessedSimulationTime(LastSerializedSimulationTime, NewHeadKeyframe);
+		Ticker.SetTotalProcessedSimulationTime(LastSerializedSimulationTime, NewHeadFrame);
 		Ticker.SetTotalAllowedSimulationTime(LastSerializedSimulationTime);
-		Ticker.PendingKeyframe = NewHeadKeyframe;
-		Ticker.MaxAllowedKeyframe = NewHeadKeyframe;
+		Ticker.PendingFrame = NewHeadFrame;
+		Ticker.MaxAllowedFrame = NewHeadFrame;
 
-		*Buffers.Sync.WriteKeyframe(NewHeadKeyframe) = LastSerializedSyncState;
-		*Buffers.Input.WriteKeyframe(NewHeadKeyframe) = TInputCmd();
+		*Buffers.Sync.WriteFrame(NewHeadFrame) = LastSerializedSyncState;
+		*Buffers.Input.WriteFrame(NewHeadFrame) = TInputCmd();
 
 		Driver->FinalizeFrame(LastSerializedSyncState, LastSerializedAuxState);
 
 		if (NETSIM_MODEL_DEBUG)
 		{
-			FVisualLoggingParameters VLogParams(EVisualLoggingContext::FirstMispredicted, ParentKeyframe, EVisualLoggingLifetime::Persistent, TEXT("DependentRollbackBegin"));
+			FVisualLoggingParameters VLogParams(EVisualLoggingContext::FirstMispredicted, ParentFrame, EVisualLoggingLifetime::Persistent, TEXT("DependentRollbackBegin"));
 			Driver->VisualLog(nullptr, &LastSerializedSyncState, &LastSerializedAuxState, VLogParams);
 		}
 	}
 
 	template<typename TSimulation, typename TDriver>
-	void DependentRollbackStep(TSimulation* Simulation, TDriver* Driver, TNetworkSimBufferContainer<TBufferTypes>& Buffers, TSimulationTicker<TTickSettings>& Ticker, const FNetworkSimTime& StepTime, const int32 ParentKeyframe, const bool bFinalStep)
+	void DependentRollbackStep(TSimulation* Simulation, TDriver* Driver, TNetworkSimBufferContainer<TBufferTypes>& Buffers, TSimulationTicker<TTickSettings>& Ticker, const FNetworkSimTime& StepTime, const int32 ParentFrame, const bool bFinalStep)
 	{
 		Ticker.SetTotalAllowedSimulationTime( Ticker.GetTotalAllowedSimulationTime() + StepTime );
 
@@ -538,7 +538,7 @@ struct TReplicator_Simulated : public TBase
 		
 		if (NETSIM_MODEL_DEBUG)
 		{
-			FVisualLoggingParameters VLogParams(bFinalStep ? EVisualLoggingContext::LastMispredicted : EVisualLoggingContext::OtherMispredicted, ParentKeyframe, EVisualLoggingLifetime::Persistent, TEXT("DependentRollbackStep"));
+			FVisualLoggingParameters VLogParams(bFinalStep ? EVisualLoggingContext::LastMispredicted : EVisualLoggingContext::OtherMispredicted, ParentFrame, EVisualLoggingLifetime::Persistent, TEXT("DependentRollbackStep"));
 			Driver->VisualLog(nullptr, SyncState, Buffers.Aux.HeadElement(), VLogParams);
 		}
 	}
@@ -549,32 +549,32 @@ private:
 	void SimulationExtrapolation(TSimulation* Simulation, TDriver* Driver, TNetworkSimBufferContainer<TBufferTypes>& Buffers, TSimulationTicker<TTickSettings>& Ticker, const FNetworkSimTime DeltaSimTime)
 	{
 		// We have extrapolated ahead of the server. The latest network update is now "in the past" from what we rendered last frame.
-		// We will insert a new keyframe to make up the difference from the last known state to where we want to be in the now.
+		// We will insert a new frame to make up the difference from the last known state to where we want to be in the now.
 
-		ensure(Buffers.Input.HeadKeyframe() == Buffers.Sync.HeadKeyframe());
+		ensure(Buffers.Input.HeadFrame() == Buffers.Sync.HeadFrame());
 
-		const int32 InputKeyframe = Buffers.Input.HeadKeyframe();
-		const int32 OutputKeyframe = InputKeyframe + 1;
+		const int32 InputFrame = Buffers.Input.HeadFrame();
+		const int32 OutputFrame = InputFrame + 1;
 
 		// Create fake cmd				
-		TInputCmd* NewCmd = Buffers.Input.WriteKeyframeInitializedFromHead(OutputKeyframe);
+		TInputCmd* NewCmd = Buffers.Input.WriteFrameInitializedFromHead(OutputFrame);
 		NewCmd->SetFrameDeltaTime(DeltaSimTime);	
 		
 		// Create new sync state to write to
-		TSyncState* PrevSyncState = Buffers.Sync[InputKeyframe];
-		TSyncState* NextSyncState = Buffers.Sync.WriteKeyframe(OutputKeyframe);
-		TAuxState* AuxState = Buffers.Aux[InputKeyframe];
+		TSyncState* PrevSyncState = Buffers.Sync[InputFrame];
+		TSyncState* NextSyncState = Buffers.Sync.WriteFrame(OutputFrame);
+		TAuxState* AuxState = Buffers.Aux[InputFrame];
 
 		// Do the actual update
 		{
-			TScopedSimulationTick UpdateScope(Ticker, OutputKeyframe, NewCmd->GetFrameDeltaTime());
+			TScopedSimulationTick UpdateScope(Ticker, OutputFrame, NewCmd->GetFrameDeltaTime());
 			Simulation->SimulationTick( 
 				{ NewCmd->GetFrameDeltaTime(), Ticker },
 				{ *NewCmd, *PrevSyncState, *AuxState },
-				{ *NextSyncState, Buffers.Aux.LazyWriter(OutputKeyframe) } );
+				{ *NextSyncState, Buffers.Aux.LazyWriter(OutputFrame) } );
 		}
 
-		Ticker.MaxAllowedKeyframe = OutputKeyframe;
+		Ticker.MaxAllowedFrame = OutputFrame;
 	}
 	
 	FNetworkSimTime ReconcileSimulationTime;
@@ -592,7 +592,7 @@ struct TReplicator_Autonomous : public TBase
 	using TSyncState = typename TBufferTypes::TSyncState;
 	using TAuxState = typename TBufferTypes::TAuxState;
 
-	int32 GetLastSerializedKeyframe() const { return SerializedKeyframe; }
+	int32 GetLastSerializedFrame() const { return SerializedFrame; }
 	bool IsReconcileFaultDetected() const { return bReconcileFaultDetected; }
 	const FNetworkSimTime& GetLastSerializedSimTime() const { return SerializedTime; }
 
@@ -611,7 +611,7 @@ struct TReplicator_Autonomous : public TBase
 	{
 		FArchive& Ar = P.Ar;
 		
-		SerializedKeyframe = FNetworkSimulationSerialization::SerializeKeyframe(Ar, Buffers.Sync.HeadKeyframe());
+		SerializedFrame = FNetworkSimulationSerialization::SerializeFrame(Ar, Buffers.Sync.HeadFrame());
 		
 		SerializedTime = Ticker.GetTotalProcessedSimulationTime();
 		SerializedTime.NetSerialize(P.Ar);
@@ -634,19 +634,19 @@ struct TReplicator_Autonomous : public TBase
 			bPendingReconciliation = false;
 
 			// The state the client predicted that corresponds to the state the server just serialized to us
-			if (TSyncState* ClientExistingState = Buffers.Sync[SerializedKeyframe])
+			if (TSyncState* ClientExistingState = Buffers.Sync[SerializedFrame])
 			{
 				// TODO: AuxState->ShouldReconcile
 				if (ClientExistingState->ShouldReconcile(SerializedSyncState) || (NetworkSimulationModelCVars::ForceReconcile() > 0) || (NetworkSimulationModelCVars::ForceReconcileSingle() > 0))
 				{
 					NetworkSimulationModelCVars::SetForceReconcileSingle(0);
-					UE_CLOG(!Buffers.Input.IsValidKeyframe(SerializedKeyframe-1), LogNetworkSim, Error, TEXT("::NetSerialize: Client InputBuffer does not contain data for frame %d. {%s} {%s}"), SerializedKeyframe, *Buffers.Input.GetBasicDebugStr(), *Buffers.Sync.GetBasicDebugStr());
+					UE_CLOG(!Buffers.Input.IsValidFrame(SerializedFrame-1), LogNetworkSim, Error, TEXT("::NetSerialize: Client InputBuffer does not contain data for frame %d. {%s} {%s}"), SerializedFrame, *Buffers.Input.GetBasicDebugStr(), *Buffers.Sync.GetBasicDebugStr());
 					bPendingReconciliation =  true;
 				}
 			}
 			else
 			{
-				if (SerializedKeyframe < Buffers.Sync.TailKeyframe())
+				if (SerializedFrame < Buffers.Sync.TailFrame())
 				{
 					// Case 1: the serialized state is older than what we've kept in our buffer. A bigger buffer would solve this! (at the price of more resimulated frames to recover when this happens)
 					// This is a reconcile fault and we just need to chill. We'll stop sending user commands until the cmds in flight flush through the system and we catch back up.
@@ -667,20 +667,20 @@ struct TReplicator_Autonomous : public TBase
 	template<typename TSimulation, typename TDriver>
 	void Reconcile(TSimulation* Simulation, TDriver* Driver, TNetworkSimBufferContainer<TBufferTypes>& Buffers, TSimulationTicker<TTickSettings>& Ticker)
 	{
-		if (bPendingReconciliation == false && (bDependentSimulationNeedsReconcile == false || SerializedKeyframe == INDEX_NONE))
+		if (bPendingReconciliation == false && (bDependentSimulationNeedsReconcile == false || SerializedFrame == INDEX_NONE))
 		{
 			return;
 		}
 		bPendingReconciliation = false;
 		bDependentSimulationNeedsReconcile = false;
 
-		TSyncState* ClientSyncState = Buffers.Sync[SerializedKeyframe];
+		TSyncState* ClientSyncState = Buffers.Sync[SerializedFrame];
 		const bool bDoVisualLog = NetworkSimulationModelCVars::EnableLocalPrediction() > 0 && NETSIM_MODEL_DEBUG; // don't visual log if we have prediction disabled
 		
 		if (bDoVisualLog)
 		{
-			FVisualLoggingParameters VLogParameters(EVisualLoggingContext::LastConfirmed, SerializedKeyframe, EVisualLoggingLifetime::Persistent, TEXT("Serialized State"));
-			Driver->VisualLog(Buffers.Input[SerializedKeyframe], &SerializedSyncState, &SerializedAuxState, VLogParameters);
+			FVisualLoggingParameters VLogParameters(EVisualLoggingContext::LastConfirmed, SerializedFrame, EVisualLoggingLifetime::Persistent, TEXT("Serialized State"));
+			Driver->VisualLog(Buffers.Input[SerializedFrame], &SerializedSyncState, &SerializedAuxState, VLogParameters);
 		}
 
 		if (ClientSyncState)
@@ -688,14 +688,14 @@ struct TReplicator_Autonomous : public TBase
 			// Existing ClientSyncState, log it before overwriting it
 			if (bDoVisualLog)
 			{
-				FVisualLoggingParameters VLogParameters(EVisualLoggingContext::FirstMispredicted, SerializedKeyframe, EVisualLoggingLifetime::Persistent, TEXT("Mispredicted State"));
-				Driver->VisualLog(Buffers.Input[SerializedKeyframe], ClientSyncState, Buffers.Aux[SerializedKeyframe], VLogParameters);
+				FVisualLoggingParameters VLogParameters(EVisualLoggingContext::FirstMispredicted, SerializedFrame, EVisualLoggingLifetime::Persistent, TEXT("Mispredicted State"));
+				Driver->VisualLog(Buffers.Input[SerializedFrame], ClientSyncState, Buffers.Aux[SerializedFrame], VLogParameters);
 			}
 		}
 		else
 		{
 			// No existing state, so create add it explicitly
-			ClientSyncState = Buffers.Sync.WriteKeyframe( SerializedKeyframe );
+			ClientSyncState = Buffers.Sync.WriteFrame( SerializedFrame );
 		}
 
 		// Set client's sync state to the server version
@@ -703,14 +703,14 @@ struct TReplicator_Autonomous : public TBase
 		*ClientSyncState = SerializedSyncState;
 
 		// Set Client's aux state to the server version
-		*Buffers.Aux.WriteKeyframe(SerializedKeyframe) = SerializedAuxState;
+		*Buffers.Aux.WriteFrame(SerializedFrame) = SerializedAuxState;
 
 		const FNetworkSimTime RollbackDeltaTime = SerializedTime - Ticker.GetTotalProcessedSimulationTime();
 
 		// Set the canonical simulation time to what we received (we will advance it as we resimulate)
-		Ticker.SetTotalProcessedSimulationTime(SerializedTime, SerializedKeyframe);
-		Ticker.PendingKeyframe = SerializedKeyframe;
-		//Ticker.MaxAllowedKeyframe = FMath::Max(Ticker.MaxAllowedKeyframe, Ticker.PendingKeyframe); // Make sure this doesn't lag behind. This is the only place we should need to do this.
+		Ticker.SetTotalProcessedSimulationTime(SerializedTime, SerializedFrame);
+		Ticker.PendingFrame = SerializedFrame;
+		//Ticker.MaxAllowedFrame = FMath::Max(Ticker.MaxAllowedFrame, Ticker.PendingFrame); // Make sure this doesn't lag behind. This is the only place we should need to do this.
 
 		if (NetworkSimulationModelCVars::EnableLocalPrediction() == 0)
 		{
@@ -721,20 +721,20 @@ struct TReplicator_Autonomous : public TBase
 		// Tell dependent simulations to rollback
 		for (INetworkSimulationModel* DependentSim : DependentSimulations)
 		{
-			DependentSim->BeginRollback(RollbackDeltaTime, SerializedKeyframe);
+			DependentSim->BeginRollback(RollbackDeltaTime, SerializedFrame);
 		}
 		
 		// Resimulate all user commands 
-		const int32 LastKeyframeToProcess = Ticker.MaxAllowedKeyframe;
-		for (int32 Keyframe = SerializedKeyframe; Keyframe <= LastKeyframeToProcess; ++Keyframe)
+		const int32 LastFrameToProcess = Ticker.MaxAllowedFrame;
+		for (int32 Frame = SerializedFrame; Frame <= LastFrameToProcess; ++Frame)
 		{
-			const int32 OutputKeyframe = Keyframe+1;
+			const int32 OutputFrame = Frame+1;
 
-			// Keyframe is the frame we are resimulating right now.
-			TInputCmd* ResimulateCmd  = Buffers.Input[Keyframe];
-			TAuxState* AuxState = Buffers.Aux[Keyframe];
-			TSyncState* PrevSyncState = Buffers.Sync[Keyframe];
-			TSyncState* NextSyncState = Buffers.Sync.WriteKeyframe(OutputKeyframe);
+			// Frame is the frame we are resimulating right now.
+			TInputCmd* ResimulateCmd  = Buffers.Input[Frame];
+			TAuxState* AuxState = Buffers.Aux[Frame];
+			TSyncState* PrevSyncState = Buffers.Sync[Frame];
+			TSyncState* NextSyncState = Buffers.Sync.WriteFrame(OutputFrame);
 			
 			check(ResimulateCmd);
 			check(PrevSyncState);
@@ -744,30 +744,30 @@ struct TReplicator_Autonomous : public TBase
 			// Log out the Mispredicted state that we are about to overwrite.
 			if (NETSIM_MODEL_DEBUG)
 			{
-				FVisualLoggingParameters VLogParameters(Keyframe == LastKeyframeToProcess ? EVisualLoggingContext::LastMispredicted : EVisualLoggingContext::OtherMispredicted, Keyframe, EVisualLoggingLifetime::Persistent, TEXT("Resimulate Step: mispredicted"));
-				Driver->VisualLog(Buffers.Input[OutputKeyframe], Buffers.Sync[OutputKeyframe], Buffers.Aux[OutputKeyframe], VLogParameters);
+				FVisualLoggingParameters VLogParameters(Frame == LastFrameToProcess ? EVisualLoggingContext::LastMispredicted : EVisualLoggingContext::OtherMispredicted, Frame, EVisualLoggingLifetime::Persistent, TEXT("Resimulate Step: mispredicted"));
+				Driver->VisualLog(Buffers.Input[OutputFrame], Buffers.Sync[OutputFrame], Buffers.Aux[OutputFrame], VLogParameters);
 			}
 
 			// Do the actual update
 			{
-				TScopedSimulationTick UpdateScope(Ticker, OutputKeyframe, ResimulateCmd->GetFrameDeltaTime());
+				TScopedSimulationTick UpdateScope(Ticker, OutputFrame, ResimulateCmd->GetFrameDeltaTime());
 				Simulation->SimulationTick( 
 					{ ResimulateCmd->GetFrameDeltaTime(), Ticker },
 					{ *ResimulateCmd, *PrevSyncState, *AuxState },
-					{ *NextSyncState, Buffers.Aux.LazyWriter(OutputKeyframe) } );
+					{ *NextSyncState, Buffers.Aux.LazyWriter(OutputFrame) } );
 			}
 
 			// Log out the newly predicted state that we got.
 			if (NETSIM_MODEL_DEBUG)
 			{			
-				FVisualLoggingParameters VLogParameters(Keyframe == LastKeyframeToProcess ? EVisualLoggingContext::LastPredicted : EVisualLoggingContext::OtherPredicted, Keyframe, EVisualLoggingLifetime::Persistent, TEXT("Resimulate Step: repredicted"));
-				Driver->VisualLog(Buffers.Input[OutputKeyframe], Buffers.Sync[OutputKeyframe], Buffers.Aux[OutputKeyframe], VLogParameters);
+				FVisualLoggingParameters VLogParameters(Frame == LastFrameToProcess ? EVisualLoggingContext::LastPredicted : EVisualLoggingContext::OtherPredicted, Frame, EVisualLoggingLifetime::Persistent, TEXT("Resimulate Step: repredicted"));
+				Driver->VisualLog(Buffers.Input[OutputFrame], Buffers.Sync[OutputFrame], Buffers.Aux[OutputFrame], VLogParameters);
 			}
 
 			// Tell dependent simulations to advance
 			for (INetworkSimulationModel* DependentSim : DependentSimulations)
 			{
-				DependentSim->StepRollback(ResimulateCmd->GetFrameDeltaTime(), Keyframe, (Keyframe == LastKeyframeToProcess));
+				DependentSim->StepRollback(ResimulateCmd->GetFrameDeltaTime(), Frame, (Frame == LastFrameToProcess));
 			}
 		}
 	}
@@ -801,7 +801,7 @@ struct TReplicator_Autonomous : public TBase
 				NonPredictedInputTimeAccumulator.Accumulate(NonPredictedInputTime, TickParameters.LocalDeltaTimeSeconds);
 				if (NonPredictedInputTime.IsPositive())
 				{
-					GenerateLocalInputCmdAtKeyframe(Driver, Buffers, NonPredictedInputTime, Buffers.Input.HeadKeyframe() + 1);
+					GenerateLocalInputCmdAtFrame(Driver, Buffers, NonPredictedInputTime, Buffers.Input.HeadFrame() + 1);
 				}
 			}
 		}
@@ -812,7 +812,7 @@ private:
 	TSyncState SerializedSyncState;
 	TAuxState SerializedAuxState;
 	FNetworkSimTime SerializedTime;
-	int32 SerializedKeyframe = -1;
+	int32 SerializedFrame = -1;
 
 	bool bPendingReconciliation = false;	// Reconciliation is pending: we need to reconcile state from the server that differs from the locally predicted state
 	bool bReconcileFaultDetected = false;	// A fault was detected: we received state from the server that we are unable to reconcile with locally predicted state
@@ -842,24 +842,24 @@ struct TReplicator_Debug : public TBase
 		uint8 SerializedNumElements = FMath::Min<uint8>(MaxNumElements, Buffer.Num());
 		Ar << SerializedNumElements;
 
-		const int32 HeadKeyframe = FNetworkSimulationSerialization::SerializeKeyframe(Ar, Buffer.HeadKeyframe());
-		const int32 StartingKeyframe = FMath::Max(0, HeadKeyframe - SerializedNumElements + 1);
+		const int32 HeadFrame = FNetworkSimulationSerialization::SerializeFrame(Ar, Buffer.HeadFrame());
+		const int32 StartingFrame = FMath::Max(0, HeadFrame - SerializedNumElements + 1);
 
 		if (Ar.IsLoading())
 		{
-			const int32 PrevHead = Buffer.HeadKeyframe();
-			if (PrevHead < StartingKeyframe && PrevHead >= 0)
+			const int32 PrevHead = Buffer.HeadFrame();
+			if (PrevHead < StartingFrame && PrevHead >= 0)
 			{
 				// There is a gap in the stream. In some cases, we want this to be a "fault" and bubble up. We may want to synthesize state or maybe we just skip ahead.
-				UE_LOG(LogNetworkSim, Warning, TEXT("Fault: gap in received Debug buffer. PrevHead: %d. Received: %d-%d. Reseting previous buffer contents"), PrevHead, StartingKeyframe, HeadKeyframe);
+				UE_LOG(LogNetworkSim, Warning, TEXT("Fault: gap in received Debug buffer. PrevHead: %d. Received: %d-%d. Reseting previous buffer contents"), PrevHead, StartingFrame, HeadFrame);
 			}
 		}
 
-		for (int32 Keyframe = StartingKeyframe; Keyframe <= HeadKeyframe; ++Keyframe)
+		for (int32 Frame = StartingFrame; Frame <= HeadFrame; ++Frame)
 		{
 			// This, as is, is bad. The intention is that these functions serialize multiple items in some delta compressed fashion.
 			// As is, we are just serializing the elements individually.
-			auto* Cmd = Ar.IsLoading() ? Buffer.WriteKeyframe(Keyframe) : Buffer[Keyframe];
+			auto* Cmd = Ar.IsLoading() ? Buffer.WriteFrame(Frame) : Buffer[Frame];
 			Cmd->NetSerialize(P);
 		}
 	}
diff --git a/Engine/Plugins/Runtime/NetworkPrediction/Source/NetworkPrediction/Public/NetworkSimulationModelTypes.h b/Engine/Plugins/Runtime/NetworkPrediction/Source/NetworkPrediction/Public/NetworkSimulationModelTypes.h
--- a/Engine/Plugins/Runtime/NetworkPrediction/Source/NetworkPrediction/Public/NetworkSimulationModelTypes.h
+++ b/Engine/Plugins/Runtime/NetworkPrediction/Source/NetworkPrediction/Public/NetworkSimulationModelTypes.h
@@ -434,25 +434,25 @@ private:
 };
 
 // ----------------------------------------------------------------------------------------------------------------------------------------------
-//	FSimulationTickState: Holds active state for simulation ticking. We track two things: keyframes and time.
+//	FSimulationTickState: Holds active state for simulation ticking. We track two things: frames and time.
 //
-//	PendingKeyframe is the next keyframe we will process: the input/sync/aux state @ PendingKeyframe will be run through ::SimulationTick and produce the 
-//	next frame's (PendingKeyframe+1) Sync and possibly Aux state (if it changes). "Out of band" modifications to the sync/aux state should happen
-//	to PendingKeyframe (e.g, before it is processed. Once a frame has been processed, we won't run it through ::SimulationTick again!).
+//	PendingFrame is the next frame we will process: the input/sync/aux state @ PendingFrame will be run through ::SimulationTick and produce the 
+//	next frame's (PendingFrame+1) Sync and possibly Aux state (if it changes). "Out of band" modifications to the sync/aux state should happen
+//	to PendingFrame (e.g, before it is processed. Once a frame has been processed, we won't run it through ::SimulationTick again!).
 //
-//	MaxAllowedKeyframe is a keyframe based limiter on simulation updates. This must be incremented to allow the simulation to advance.
+//	MaxAllowedFrame is a frame based limiter on simulation updates. This must be incremented to allow the simulation to advance.
 //
 //	Time is also tracked. We keep running total for how much the sim has advanced and how much it is allowed to advance. There is also a historic buffer of
 //	simulation time in SimulationTimeBuffer.
 //
-//	Consider that Keyframes are essentially client dependent and gaps can happen due to packet loss, etc. Time will always be continuous though.
+//	Consider that Frames are essentially client dependent and gaps can happen due to packet loss, etc. Time will always be continuous though.
 //	
 // ----------------------------------------------------------------------------------------------------------------------------------------------
 
 struct FSimulationTickState
 {
-	int32 PendingKeyframe = 0;
-	int32 MaxAllowedKeyframe = -1;
+	int32 PendingFrame = 0;
+	int32 MaxAllowedFrame = -1;
 	bool bUpdateInProgress = false;
 	
 	FNetworkSimTime GetTotalProcessedSimulationTime() const 
@@ -460,16 +460,16 @@ struct FSimulationTickState
 		return TotalProcessedSimulationTime; 
 	}
 
-	void SetTotalProcessedSimulationTime(const FNetworkSimTime& SimTime, int32 Keyframe)
+	void SetTotalProcessedSimulationTime(const FNetworkSimTime& SimTime, int32 Frame)
 	{
 		TotalProcessedSimulationTime = SimTime;
-		*SimulationTimeBuffer.WriteKeyframe(Keyframe) = SimTime;
+		*SimulationTimeBuffer.WriteFrame(Frame) = SimTime;
 	}
 
-	void IncrementTotalProcessedSimulationTime(const FNetworkSimTime& DeltaSimTime, int32 Keyframe)
+	void IncrementTotalProcessedSimulationTime(const FNetworkSimTime& DeltaSimTime, int32 Frame)
 	{
 		TotalProcessedSimulationTime += DeltaSimTime;
-		*SimulationTimeBuffer.WriteKeyframe(Keyframe) = TotalProcessedSimulationTime;
+		*SimulationTimeBuffer.WriteFrame(Frame) = TotalProcessedSimulationTime;
 	}
 		
 	// Historic tracking of simulation time. This allows us to timestamp sync data as its produced
@@ -492,6 +492,7 @@ protected:
 	FNetworkSimTime TotalProcessedSimulationTime;	// How much time we've actually processed. The only way to increment this is to process user commands or receive authoritative state from the network.
 };
 
+// "Ticker" that actually allows us to give the simulation time. This struct will generally not be passed around outside of the core TNetworkedSimulationModel/Replicators
 template<typename TickSettings=TNetworkSimTickSettings<>>
 struct TSimulationTicker : public FSimulationTickState
 {
@@ -516,23 +517,23 @@ private:
 };
 
 // Scoped helper to be used right before entering a call to the sim's ::SimulationTick function.
-// Important to note that this advanced the PendingKeyFrame to the output Keyframe. So that any writes that occur to the buffers will go to the output frame.
+// Important to note that this advances the PendingFrame to the output Frame. So that any writes that occur to the buffers during this scope will go to the output frame.
 struct TScopedSimulationTick
 {
-	TScopedSimulationTick(FSimulationTickState& InTicker, const int32& InOutputKeyframe, const FNetworkSimTime& InDeltaSimTime)
-		: Ticker(InTicker), OutputKeyframe(InOutputKeyframe), DeltaSimTime(InDeltaSimTime)
+	TScopedSimulationTick(FSimulationTickState& InTicker, const int32& InOutputFrame, const FNetworkSimTime& InDeltaSimTime)
+		: Ticker(InTicker), OutputFrame(InOutputFrame), DeltaSimTime(InDeltaSimTime)
 	{
 		check(Ticker.bUpdateInProgress == false);
-		Ticker.PendingKeyframe = OutputKeyframe;
+		Ticker.PendingFrame = OutputFrame;
 		Ticker.bUpdateInProgress = true;
 	}
 	~TScopedSimulationTick()
 	{
-		Ticker.IncrementTotalProcessedSimulationTime(DeltaSimTime, OutputKeyframe);
+		Ticker.IncrementTotalProcessedSimulationTime(DeltaSimTime, OutputFrame);
 		Ticker.bUpdateInProgress = false;
 	}
 	FSimulationTickState& Ticker;
-	const int32& OutputKeyframe;
+	const int32& OutputFrame;
 	const FNetworkSimTime& DeltaSimTime;
 };
 
@@ -542,7 +543,7 @@ struct TScopedSimulationTick
 
 // Accessor conditionally gives access to the current (pending) Sync/Aux state to outside code.
 // Reads are always allowed
-// Writes are conditional. Authority can always write to the pending keyframe. Non authority requires the netsim to be currently processing an ::SimulationTick.
+// Writes are conditional. Authority can always write to the pending frame. Non authority requires the netsim to be currently processing an ::SimulationTick.
 // If you aren't inside an ::SimulationTick call, it is really not safe to predict state changes. It is safest and simplest to just not predict these changes.
 //
 // Explanation: During the scope of an ::SimulationTick call, we know exactly 'when' we are relative to what the server is processing. If the predicting client wants
@@ -565,7 +566,7 @@ struct TNetworkSimStateAccessor
 			// Gross: find the buffer our TState is in. This allows these accessors to be declared as class member variables without knowing the exact net sim it will
 			// be pulling from. (For example, you may want to templatize your network sim model).
 			auto& Buffer = TBufferGetterHelper_ByState<TNetworkSimBufferContainer<typename TNetworkSimModel::TBufferTypes>, TState>::GetBuffer(NetSimModel->Buffers);
-			OutState = bWrite ? Buffer.WriteKeyframeInitializedFromHead(NetSimModel->Ticker.PendingKeyframe) : Buffer[NetSimModel->Ticker.PendingKeyframe];
+			OutState = bWrite ? Buffer.WriteFrameInitializedFromHead(NetSimModel->Ticker.PendingFrame) : Buffer[NetSimModel->Ticker.PendingFrame];
 			OutSafe = NetSimModel->Ticker.bUpdateInProgress;
 		};
 	}
@@ -575,7 +576,7 @@ struct TNetworkSimStateAccessor
 		GetStateFunc = nullptr;
 	}
 
-	/** Gets the current (PendingKeyframe) state for reading. This is not expected to fail outside of startup/shutdown edge cases */
+	/** Gets the current (PendingFrame) state for reading. This is not expected to fail outside of startup/shutdown edge cases */
 	const TState* GetStateRead() const
 	{
 		TState* State = nullptr;
@@ -587,7 +588,7 @@ struct TNetworkSimStateAccessor
 		return State;
 	}
 
-	/** Gets the current (PendingKeyframe) state for writing. This is expected to fail outside of the core update loop when bHasAuthority=false. (E.g, it is not safe to predict writes) */
+	/** Gets the current (PendingFrame) state for writing. This is expected to fail outside of the core update loop when bHasAuthority=false. (E.g, it is not safe to predict writes) */
 	TState* GetStateWrite(bool bHasAuthority) const
 	{
 		TState* State = nullptr;
@@ -607,7 +608,7 @@ private:
 //TODO
 struct FNetworkSimEventAccessor
 {
-	int32 GetPendingKeyframe() { return -1; }
+	int32 GetPendingFrame() { return -1; }
 
 	struct FFrameEvents
 	{
diff --git a/Engine/Plugins/Runtime/NetworkPredictionExtras/Source/NetworkPredictionExtras/Private/MockNetworkSimulation.cpp b/Engine/Plugins/Runtime/NetworkPredictionExtras/Source/NetworkPredictionExtras/Private/MockNetworkSimulation.cpp
--- a/Engine/Plugins/Runtime/NetworkPredictionExtras/Source/NetworkPredictionExtras/Private/MockNetworkSimulation.cpp
+++ b/Engine/Plugins/Runtime/NetworkPredictionExtras/Source/NetworkPredictionExtras/Private/MockNetworkSimulation.cpp
@@ -165,7 +165,7 @@ void UMockNetworkSimulationComponent::VisualLog(const FMockInputCmd* Input, cons
 	AActor* Owner = GetOwner();
 	FTransform Transform = Owner->GetActorTransform();
 
-	LocalParameters.DebugString += FString::Printf(TEXT(" [%d] Total: %.4f"), SystemParameters.Keyframe, Sync->Total);
+	LocalParameters.DebugString += FString::Printf(TEXT(" [%d] Total: %.4f"), SystemParameters.Frame, Sync->Total);
 
 	FVisualLoggingHelpers::VisualLogActor(Owner, Transform, LocalParameters);
 }
diff --git a/Engine/Plugins/Runtime/NetworkPredictionExtras/Source/NetworkPredictionExtras/Private/NetworkPredictionExtrasFlyingPawn.cpp b/Engine/Plugins/Runtime/NetworkPredictionExtras/Source/NetworkPredictionExtras/Private/NetworkPredictionExtrasFlyingPawn.cpp
--- a/Engine/Plugins/Runtime/NetworkPredictionExtras/Source/NetworkPredictionExtras/Private/NetworkPredictionExtrasFlyingPawn.cpp
+++ b/Engine/Plugins/Runtime/NetworkPredictionExtras/Source/NetworkPredictionExtras/Private/NetworkPredictionExtrasFlyingPawn.cpp
@@ -252,13 +252,13 @@ FlyingMovement::FAuxState* ANetworkPredictionExtrasFlyingPawn::GetAuxStateWrite(
 	return nullptr;
 }
 
-int32 ANetworkPredictionExtrasFlyingPawn::GetPendingKeyframe() const
+int32 ANetworkPredictionExtrasFlyingPawn::GetPendingFrame() const
 {
 	// TODO
 	/*
 	if (ensure(FlyingMovementComponent))
 	{
-		return FlyingMovementComponent->MovementSyncState->GetPendingKeyframe();
+		return FlyingMovementComponent->MovementSyncState->GetPendingFrame();
 	}
 	*/
 	return 0;
diff --git a/Engine/Plugins/Runtime/NetworkPredictionExtras/Source/NetworkPredictionExtras/Public/MockNetworkSimulation.h b/Engine/Plugins/Runtime/NetworkPredictionExtras/Source/NetworkPredictionExtras/Public/MockNetworkSimulation.h
--- a/Engine/Plugins/Runtime/NetworkPredictionExtras/Source/NetworkPredictionExtras/Public/MockNetworkSimulation.h
+++ b/Engine/Plugins/Runtime/NetworkPredictionExtras/Source/NetworkPredictionExtras/Public/MockNetworkSimulation.h
@@ -120,7 +120,7 @@ struct FMockAuxState
 	{
 		if (Params.Context == EStandardLoggingContext::HeaderOnly)
 		{
-			Params.Ar->Logf(TEXT(" %d "), Params.Keyframe);
+			Params.Ar->Logf(TEXT(" %d "), Params.Frame);
 		}
 		else if (Params.Context == EStandardLoggingContext::Full)
 		{
diff --git a/Engine/Plugins/Runtime/NetworkPredictionExtras/Source/NetworkPredictionExtras/Public/NetworkPredictionExtrasFlyingPawn.h b/Engine/Plugins/Runtime/NetworkPredictionExtras/Source/NetworkPredictionExtras/Public/NetworkPredictionExtrasFlyingPawn.h
--- a/Engine/Plugins/Runtime/NetworkPredictionExtras/Source/NetworkPredictionExtras/Public/NetworkPredictionExtrasFlyingPawn.h
+++ b/Engine/Plugins/Runtime/NetworkPredictionExtras/Source/NetworkPredictionExtras/Public/NetworkPredictionExtrasFlyingPawn.h
@@ -78,7 +78,7 @@ public:
 
 	// Only intended for debugging in test map examples. Not really intended to be useful for general game code use.
 	UFUNCTION(BlueprintCallable, Category="Gameplay")
-	int32 GetPendingKeyframe() const;
+	int32 GetPendingFrame() const;
 
 protected:
 
