NetworkPredictionPlugin:
-Bulk rename of .cpp/.h files: NetworkSimulation* -> NetworkedSimulation*
-Reasoning: "NetworkSimulation" can imply "a simulation of a computer network" when we really mean "a networked gameplay simulation". This is more consistent with the class/type names used by the system as well.

#rnx

diff --git a/Engine/Plugins/Runtime/NetworkPrediction/Source/NetworkPrediction/Private/Movement/ParametricMovement.cpp b/Engine/Plugins/Runtime/NetworkPrediction/Source/NetworkPrediction/Private/Movement/ParametricMovement.cpp
--- a/Engine/Plugins/Runtime/NetworkPrediction/Source/NetworkPrediction/Private/Movement/ParametricMovement.cpp
+++ b/Engine/Plugins/Runtime/NetworkPrediction/Source/NetworkPrediction/Private/Movement/ParametricMovement.cpp
@@ -2,7 +2,7 @@
 
 #include "Movement/ParametricMovement.h"
 #include "HAL/IConsoleManager.h"
-#include "NetworkSimulationGlobalManager.h"
+#include "NetworkedSimulationGlobalManager.h"
 #include "Engine/LocalPlayer.h"
 #include "GameFramework/PlayerController.h"
 
diff --git a/Engine/Plugins/Runtime/NetworkPrediction/Source/NetworkPrediction/Private/NetworkPredictionComponent.cpp b/Engine/Plugins/Runtime/NetworkPrediction/Source/NetworkPrediction/Private/NetworkPredictionComponent.cpp
--- a/Engine/Plugins/Runtime/NetworkPrediction/Source/NetworkPrediction/Private/NetworkPredictionComponent.cpp
+++ b/Engine/Plugins/Runtime/NetworkPrediction/Source/NetworkPrediction/Private/NetworkPredictionComponent.cpp
@@ -4,7 +4,8 @@
 #include "Net/UnrealNetwork.h"
 #include "GameFramework/Actor.h"
 #include "GameFramework/Pawn.h"
-#include "NetworkSimulationGlobalManager.h"
+#include "NetworkedSimulationGlobalManager.h"
+#include "Engine/World.h"
 
 UNetworkPredictionComponent::UNetworkPredictionComponent()
 {
diff --git a/Engine/Plugins/Runtime/NetworkPrediction/Source/NetworkPrediction/Private/NetworkPredictionTests.cpp b/Engine/Plugins/Runtime/NetworkPrediction/Source/NetworkPrediction/Private/NetworkPredictionTests.cpp
--- a/Engine/Plugins/Runtime/NetworkPrediction/Source/NetworkPrediction/Private/NetworkPredictionTests.cpp
+++ b/Engine/Plugins/Runtime/NetworkPrediction/Source/NetworkPrediction/Private/NetworkPredictionTests.cpp
@@ -1,6 +1,6 @@
 // Copyright 1998-2019 Epic Games, Inc. All Rights Reserved.
 
-#include "NetworkSimulationModel.h"
+#include "NetworkedSimulationModel.h"
 
 // ------------------------------------------------------------------------------------------------------------
 //	Testing Fixed vs Variable tick rate struct declration
diff --git a/Engine/Plugins/Runtime/NetworkPrediction/Source/NetworkPrediction/Private/NetworkSimulationGlobalManager.cpp b/Engine/Plugins/Runtime/NetworkPrediction/Source/NetworkPrediction/Private/NetworkSimulationGlobalManager.cpp
deleted file mode 100644
diff --git a/Engine/Plugins/Runtime/NetworkPrediction/Source/NetworkPrediction/Private/NetworkSimulationModelDebugger.cpp b/Engine/Plugins/Runtime/NetworkPrediction/Source/NetworkPrediction/Private/NetworkSimulationModelDebugger.cpp
deleted file mode 100644
diff --git a/Engine/Plugins/Runtime/NetworkPrediction/Source/NetworkPrediction/Private/NetworkedSimulationGlobalManager.cpp b/Engine/Plugins/Runtime/NetworkPrediction/Source/NetworkPrediction/Private/NetworkedSimulationGlobalManager.cpp
new file mode 100644
--- /dev/null
+++ b/Engine/Plugins/Runtime/NetworkPrediction/Source/NetworkPrediction/Private/NetworkedSimulationGlobalManager.cpp
@@ -0,0 +1,81 @@
+// Copyright 1998-2019 Epic Games, Inc. All Rights Reserved.
+
+#include "NetworkedSimulationGlobalManager.h"
+#include "Engine/World.h"
+
+DEFINE_LOG_CATEGORY_STATIC(NetworkSimulationGlobalManager, Log, All);
+
+UNetworkSimulationGlobalManager::UNetworkSimulationGlobalManager()
+{
+
+}
+
+void UNetworkSimulationGlobalManager::Initialize(FSubsystemCollectionBase& Collection)
+{
+	UWorld* World = GetWorld();
+	check(World);
+
+	PostTickDispatchHandle = World->OnPostTickDispatch().AddUObject(this, &UNetworkSimulationGlobalManager::ReconcileSimulationsPostNetworkUpdate);
+	PreWorldActorTickHandle = FWorldDelegates::OnWorldPreActorTick.AddUObject(this, &UNetworkSimulationGlobalManager::BeginNewSimulationFrame);
+}
+
+void UNetworkSimulationGlobalManager::Deinitialize()
+{
+	if (UWorld* World = GetWorld())
+	{
+		World->OnPostTickDispatch().Remove(PostTickDispatchHandle);
+	}
+
+	FWorldDelegates::OnWorldPreActorTick.Remove(PreWorldActorTickHandle);
+}
+
+void UNetworkSimulationGlobalManager::ReconcileSimulationsPostNetworkUpdate()
+{
+	for (auto It = SimulationGroupMap.CreateIterator(); It; ++It)
+	{
+		for (auto& SimItem : It.Value().Simulations)
+		{
+			SimItem.Simulation->Reconcile(SimItem.OwningActor->Role);
+		}
+	}
+}
+
+void UNetworkSimulationGlobalManager::BeginNewSimulationFrame(UWorld* InWorld, ELevelTick InLevelTick, float InDeltaSeconds)
+{
+	if (InWorld != GetWorld() || !InWorld->HasBegunPlay())
+	{
+		return;
+	}
+
+	FNetSimTickParameters TickParameters(InDeltaSeconds);
+
+	for (auto It = SimulationGroupMap.CreateIterator(); It; ++It)
+	{
+		for (auto& SimItem : It.Value().Simulations)
+		{
+			TickParameters.InitFromActor(SimItem.OwningActor);
+			SimItem.Simulation->Tick(TickParameters);
+		}
+	}
+
+	TickServerRPCDelegate.Broadcast(InDeltaSeconds);
+}
+
+void UNetworkSimulationGlobalManager::RegisterModel(INetworkSimulationModel* Model, AActor* OwningActor)
+{
+	SimulationGroupMap.FindOrAdd(Model->GetSimulationGroupName()).Simulations.Emplace(Model, OwningActor);
+}
+
+void UNetworkSimulationGlobalManager::UnregisterModel(INetworkSimulationModel* Model, AActor* OwningActor)
+{
+	auto& SimulationList = SimulationGroupMap.FindOrAdd(Model->GetSimulationGroupName()).Simulations;
+	int32 idx = SimulationList.IndexOfByKey(Model);
+	if (idx != INDEX_NONE)
+	{
+		SimulationList.RemoveAtSwap(idx, 1, false);
+	}
+	else
+	{
+		UE_LOG(NetworkSimulationGlobalManager, Warning, TEXT("::Unregister model called on model %s but could not be found in manager."), *Model->GetSimulationGroupName().ToString());
+	}
+}
\ No newline at end of file
diff --git a/Engine/Plugins/Runtime/NetworkPrediction/Source/NetworkPrediction/Private/NetworkedSimulationModelDebugger.cpp b/Engine/Plugins/Runtime/NetworkPrediction/Source/NetworkPrediction/Private/NetworkedSimulationModelDebugger.cpp
new file mode 100644
--- /dev/null
+++ b/Engine/Plugins/Runtime/NetworkPrediction/Source/NetworkPrediction/Private/NetworkedSimulationModelDebugger.cpp
@@ -0,0 +1,159 @@
+// Copyright 1998-2019 Epic Games, Inc. All Rights Reserved.
+
+#include "NetworkedSimulationModelDebugger.h"
+#include "NetworkPredictionTypes.h"
+#include "Engine/World.h"
+#include "Engine/LocalPlayer.h"
+#include "GameFramework/Pawn.h"
+#include "GameFramework/PlayerController.h"
+#include "Misc/NetworkGuid.h"
+#include "Engine/NetConnection.h"
+#include "Engine/NetDriver.h"
+#include "EngineUtils.h"
+#include "ProfilingDebugging/ScopedTimers.h"
+
+FNetworkSimulationModelDebuggerManager& FNetworkSimulationModelDebuggerManager::Get()
+{
+	static FNetworkSimulationModelDebuggerManager Manager;
+	return Manager;
+}
+
+UObject* FindReplicatedObjectOnPIEServer(const UObject* ClientObject)
+{
+	if (ClientObject == nullptr)
+	{
+		return nullptr;
+	}
+
+	UObject* ServerObject = nullptr;
+
+#if WITH_EDITOR
+	if (UWorld* World = ClientObject->GetWorld())
+	{
+		if (UNetDriver* NetDriver = World->GetNetDriver())
+		{
+			if (UNetConnection* NetConnection = NetDriver->ServerConnection)
+			{
+				FNetworkGUID NetGUID = NetConnection->PackageMap->GetNetGUIDFromObject(ClientObject);
+
+				// Find the PIE server world
+				for (TObjectIterator<UWorld> It; It; ++It)
+				{
+					if (It->WorldType == EWorldType::PIE && It->GetNetMode() != NM_Client)
+					{
+						if (UNetDriver* ServerNetDriver = It->GetNetDriver())
+						{
+							if (ServerNetDriver->ClientConnections.Num() > 0)
+							{
+								UPackageMap* ServerPackageMap = ServerNetDriver->ClientConnections[0]->PackageMap;
+								ServerObject = ServerPackageMap->GetObjectFromNetGUID(NetGUID, true);
+								break;
+							}
+						}
+					}
+				}
+			}
+		}
+	}
+#endif
+
+	return ServerObject;
+}
+
+// ------------------------------------------------------------------------------------------------------------------------
+//	Debug functions for toggling debugger on specific actors
+// ------------------------------------------------------------------------------------------------------------------------
+
+// Debug the first locally controlled pawn
+FAutoConsoleCommandWithWorldAndArgs NetworkSimulationModelDebugCmd(TEXT("nms.Debug.LocallyControlledPawn"), TEXT(""),
+	FConsoleCommandWithWorldAndArgsDelegate::CreateLambda([](const TArray< FString >& InArgs, UWorld* World) 
+{
+
+	if (!World || !World->GetFirstLocalPlayerFromController())
+	{
+		return;
+	}
+	ULocalPlayer* Player = World->GetFirstLocalPlayerFromController();
+	if (!Player || !Player->GetPlayerController(World) || !Player->GetPlayerController(World)->GetPawn())
+	{
+		UE_LOG(LogNetworkSimDebug, Error, TEXT("Could not find valid locally controlled pawn."));
+		return;
+	}
+
+	APawn* Pawn = Player->GetPlayerController(World)->GetPawn();
+	FNetworkSimulationModelDebuggerManager::Get().ToggleDebuggerActive(Pawn);
+}));
+
+// Toggles continous updates to debugger for locally controlled player
+FAutoConsoleCommandWithWorldAndArgs NetworkSimulationModelDebugToggleContinousCmd(TEXT("nms.Debug.ToggleContinous"), TEXT(""),
+	FConsoleCommandWithWorldAndArgsDelegate::CreateLambda([](const TArray< FString >& InArgs, UWorld* World) 
+{
+
+	if (!World || !World->GetFirstLocalPlayerFromController())
+	{
+		return;
+	}
+	ULocalPlayer* Player = World->GetFirstLocalPlayerFromController();
+	if (!Player || !Player->GetPlayerController(World) || !Player->GetPlayerController(World)->GetPawn())
+	{
+		UE_LOG(LogNetworkSimDebug, Error, TEXT("Could not find valid locally controlled pawn. "));
+		return;
+	}
+
+	APawn* Pawn = Player->GetPlayerController(World)->GetPawn();
+	FNetworkSimulationModelDebuggerManager::Get().SetDebuggerActive(Pawn, true);
+	FNetworkSimulationModelDebuggerManager::Get().ToggleContinousGather();
+}));
+
+// Debug actors by class filter. 2nd parameter can filter by name
+FAutoConsoleCommandWithWorldAndArgs NetworkSimulationModelDebugClassCmd(TEXT("nms.Debug.Class"), TEXT(""),
+	FConsoleCommandWithWorldAndArgsDelegate::CreateLambda([](const TArray< FString >& Args, UWorld* World) 
+{
+	if (Args.Num() <= 0)
+	{
+		UE_LOG(LogNetworkSimDebug, Display, TEXT("Usage: nms.Debug.Class <Class> <Name>"));
+		return;
+	}
+
+	UClass* Class = FindObject<UClass>(ANY_PACKAGE, *Args[0]);
+	if (Class == nullptr)
+	{
+		UE_LOG(LogNetworkSimDebug, Display, TEXT("Could not find Class: %s. Searching for partial match."), *Args[0]);
+
+		for (TObjectIterator<UClass> ClassIt; ClassIt; ++ClassIt)
+		{
+			if (ClassIt->IsChildOf(AActor::StaticClass()) && ClassIt->GetName().Contains(Args[0]))
+			{
+				Class = *ClassIt;
+				UE_LOG(LogNetworkSimDebug, Display, TEXT("Found: %s"), *Class->GetName());
+				break;
+			}
+
+		}
+
+		if (Class == nullptr)
+		{
+			UE_LOG(LogNetworkSimDebug, Display, TEXT("No Matches"));
+			return;
+		}
+	}
+
+	// Probably need to do distance based sorting or a way to specify which one you want.
+	bool bFound = false;
+	for (TActorIterator<AActor> It(World); It; ++It)
+	{
+		if (It->GetClass()->IsChildOf(Class))
+		{
+			if (Args.Num() < 2 || It->GetName().Contains(Args[1]))
+			{
+				UE_LOG(LogNetworkSimDebug, Display, TEXT("Toggling NetworkSim debugger for %s"), *It->GetName());
+				FNetworkSimulationModelDebuggerManager::Get().ToggleDebuggerActive(*It);
+				bFound = true;
+			}
+		}
+	}
+	if (!bFound)
+	{
+		UE_LOG(LogNetworkSimDebug, Display, TEXT("Unable to find actors matching search criteria. Class: %s"), *GetNameSafe(Class));
+	}
+}));
\ No newline at end of file
diff --git a/Engine/Plugins/Runtime/NetworkPrediction/Source/NetworkPrediction/Public/Movement/FlyingMovement.h b/Engine/Plugins/Runtime/NetworkPrediction/Source/NetworkPrediction/Public/Movement/FlyingMovement.h
--- a/Engine/Plugins/Runtime/NetworkPrediction/Source/NetworkPrediction/Public/Movement/FlyingMovement.h
+++ b/Engine/Plugins/Runtime/NetworkPrediction/Source/NetworkPrediction/Public/Movement/FlyingMovement.h
@@ -13,7 +13,7 @@
 #include "Components/SceneComponent.h"
 #include "Misc/OutputDevice.h"
 #include "Misc/CoreDelegates.h"
-#include "NetworkSimulationModel.h"
+#include "NetworkedSimulationModel.h"
 #include "BaseMovementComponent.h"
 
 #include "FlyingMovement.generated.h"
diff --git a/Engine/Plugins/Runtime/NetworkPrediction/Source/NetworkPrediction/Public/Movement/ParametricMovement.h b/Engine/Plugins/Runtime/NetworkPrediction/Source/NetworkPrediction/Public/Movement/ParametricMovement.h
--- a/Engine/Plugins/Runtime/NetworkPrediction/Source/NetworkPrediction/Public/Movement/ParametricMovement.h
+++ b/Engine/Plugins/Runtime/NetworkPrediction/Source/NetworkPrediction/Public/Movement/ParametricMovement.h
@@ -2,7 +2,7 @@
 
 #pragma once
 
-#include "NetworkSimulationModel.h"
+#include "NetworkedSimulationModel.h"
 #include "BaseMovementComponent.h"
 
 #include "ParametricMovement.generated.h"
diff --git a/Engine/Plugins/Runtime/NetworkPrediction/Source/NetworkPrediction/Public/NetworkSimulationGlobalManager.h b/Engine/Plugins/Runtime/NetworkPrediction/Source/NetworkPrediction/Public/NetworkSimulationGlobalManager.h
deleted file mode 100644
diff --git a/Engine/Plugins/Runtime/NetworkPrediction/Source/NetworkPrediction/Public/NetworkSimulationModel.h b/Engine/Plugins/Runtime/NetworkPrediction/Source/NetworkPrediction/Public/NetworkSimulationModel.h
deleted file mode 100644
diff --git a/Engine/Plugins/Runtime/NetworkPrediction/Source/NetworkPrediction/Public/NetworkSimulationModelBuffer.h b/Engine/Plugins/Runtime/NetworkPrediction/Source/NetworkPrediction/Public/NetworkSimulationModelBuffer.h
deleted file mode 100644
diff --git a/Engine/Plugins/Runtime/NetworkPrediction/Source/NetworkPrediction/Public/NetworkSimulationModelCVars.h b/Engine/Plugins/Runtime/NetworkPrediction/Source/NetworkPrediction/Public/NetworkSimulationModelCVars.h
deleted file mode 100644
diff --git a/Engine/Plugins/Runtime/NetworkPrediction/Source/NetworkPrediction/Public/NetworkSimulationModelDebugger.h b/Engine/Plugins/Runtime/NetworkPrediction/Source/NetworkPrediction/Public/NetworkSimulationModelDebugger.h
deleted file mode 100644
diff --git a/Engine/Plugins/Runtime/NetworkPrediction/Source/NetworkPrediction/Public/NetworkSimulationModelInterpolator.h b/Engine/Plugins/Runtime/NetworkPrediction/Source/NetworkPrediction/Public/NetworkSimulationModelInterpolator.h
deleted file mode 100644
diff --git a/Engine/Plugins/Runtime/NetworkPrediction/Source/NetworkPrediction/Public/NetworkSimulationModelReplicators.h b/Engine/Plugins/Runtime/NetworkPrediction/Source/NetworkPrediction/Public/NetworkSimulationModelReplicators.h
deleted file mode 100644
diff --git a/Engine/Plugins/Runtime/NetworkPrediction/Source/NetworkPrediction/Public/NetworkSimulationModelTypes.h b/Engine/Plugins/Runtime/NetworkPrediction/Source/NetworkPrediction/Public/NetworkSimulationModelTypes.h
deleted file mode 100644
diff --git a/Engine/Plugins/Runtime/NetworkPrediction/Source/NetworkPrediction/Public/NetworkedSimulationGlobalManager.h b/Engine/Plugins/Runtime/NetworkPrediction/Source/NetworkPrediction/Public/NetworkedSimulationGlobalManager.h
new file mode 100644
--- /dev/null
+++ b/Engine/Plugins/Runtime/NetworkPrediction/Source/NetworkPrediction/Public/NetworkedSimulationGlobalManager.h
@@ -0,0 +1,87 @@
+// Copyright 1998-2019 Epic Games, Inc. All Rights Reserved.
+
+#pragma once
+
+#include "CoreMinimal.h"
+#include "UObject/ObjectMacros.h"
+#include "UObject/UObjectGlobals.h"
+#include "UObject/Object.h"
+#include "Subsystems/WorldSubsystem.h"
+#include "Engine/EngineBaseTypes.h"
+#include "NetworkPredictionTypes.h"
+#include "Containers/SortedMap.h"
+
+#include "NetworkedSimulationGlobalManager.generated.h"
+
+// ----------------------------------------------------------------------------------------------------------------------------
+//	The purpose of the global manager is to enforce frame boundaries and ordering amongst all instances of network simulations.
+//	Other than registration it is not something much code should interact with. It mainly manages when network sims reconcile and tick.
+//
+//
+//	Frame Overview:
+//
+//		Network TickFlush (receive network traffic)
+//			NetSerialize is called on individual net sims (in an essentially random order)
+//				RepNotifies, PostNetReceive, etc are called on actors *right after they individually receive traffic* (this is kind of unfortunate but how it is)
+//		
+//		UNetworkSimulationGlobalManager::ReconcilePostNetworkUpdate	 (Reconcile is called in deterministic order on all network sims)
+//			(think: "Get yourself right with the latest update from the server *before* we start a new frame")
+//
+//		UNetworkSimulationGlobalManager::BeginNewSimulationFrame ("Add" "allowed" time into the global "space" which will allow sims to continue to tick)
+//			Simulations are ticked in deterministic order (if registered with UNetworkSimulationGlobalManager)
+//
+//		UE4 Actors Tick
+//
+//		Network Tick Dispatch (send outgoing network traffic)
+//
+//
+// ----------------------------------------------------------------------------------------------------------------------------
+
+class IReplicationProxy;
+
+UCLASS()
+class UNetworkSimulationGlobalManager : public UWorldSubsystem
+{
+	GENERATED_BODY()
+public:
+
+	UNetworkSimulationGlobalManager();
+	
+	// Subsystem Init/Deinit
+	virtual void Initialize(FSubsystemCollectionBase& Collection) override;
+	virtual void Deinitialize() override;
+
+	// Netsim Register/Unregister
+	void RegisterModel(INetworkSimulationModel* Model, AActor* OwningActor);
+	void UnregisterModel(INetworkSimulationModel* Model, AActor* OwningActor);
+
+	// Delegate for auto proxies to register with to send their server RPC
+	DECLARE_MULTICAST_DELEGATE_OneParam(FTickServerRPC, float /*RealTime DeltaSeconds*/)
+	FTickServerRPC TickServerRPCDelegate;
+
+private:
+
+	void ReconcileSimulationsPostNetworkUpdate();
+	void BeginNewSimulationFrame(UWorld* InWorld, ELevelTick InLevelTick, float InDeltaSeconds);
+	
+	// Structure for organizing our simulations into groups to give some ticking order.
+	// This will need to evolve over time a bit and be optimized for larger number of simulations.
+	struct FSimulationGroup
+	{
+		struct FItem
+		{
+			FItem(INetworkSimulationModel* InSim, AActor* InActor) : Simulation(InSim), OwningActor(InActor) { }
+			bool operator == (const INetworkSimulationModel* InSim) const { return Simulation == InSim; }
+
+			INetworkSimulationModel* Simulation;
+			AActor* OwningActor;
+		};
+
+		TArray<FItem> Simulations;
+	};
+
+	TSortedMap<FName, FSimulationGroup, FDefaultAllocator, FNameFastLess> SimulationGroupMap;
+
+	FDelegateHandle PostTickDispatchHandle;
+	FDelegateHandle PreWorldActorTickHandle;
+};
\ No newline at end of file
diff --git a/Engine/Plugins/Runtime/NetworkPrediction/Source/NetworkPrediction/Public/NetworkedSimulationModel.h b/Engine/Plugins/Runtime/NetworkPrediction/Source/NetworkPrediction/Public/NetworkedSimulationModel.h
new file mode 100644
--- /dev/null
+++ b/Engine/Plugins/Runtime/NetworkPrediction/Source/NetworkPrediction/Public/NetworkedSimulationModel.h
@@ -0,0 +1,464 @@
+// Copyright 1998-2019 Epic Games, Inc. All Rights Reserved.
+
+#pragma once
+#include "NetworkPredictionTypes.h"
+#include "NetworkedSimulationModelBuffer.h"
+#include "NetworkedSimulationModelTypes.h"
+#include "NetworkedSimulationModelReplicators.h"
+#include "NetworkedSimulationModelDebugger.h"
+
+// ----------------------------------------------------------------------------------------------------------------------------------------------------------------------
+//	TNetworkedSimulationModel
+//	
+//	* Has all logic for "ticking, advancing buffers, calling Update, calling ServerRPC etc
+//	* Doesn't have anything about update component, movesweep, etc
+//	* Concept of "IDriver" which is the owning object that is driving the network sim. This is the interface to the outside UE4 world.
+//	* Has 4 buffers:
+//		-Input: Generated by a client / not the authority.
+//		-Sync: What we are trying to keep in sync. The state that evolves frame to frame with an Update function.
+//		-Aux: State that is also an input into the simulation but does not intrinsically evolve from to frame. Changes to this state can be trapped/tracked/predicted.
+//		-Debug: Replicated buffer from server->client with server-frame centered debug information. Compiled out of shipping builds.
+//
+//	* How other code interacts with this:
+//		-Network updates will come in through UE4 networking -> FReplicationProxy (on actor/component) -> TNetworkedSimulationModel::RepProxy_* -> Buffers.*
+//		-UNetworkSimulationGlobalManager: responsible for ticking simulation (after recv net traffic, prior to UE4 actor ticking)
+//		-External game code can interact with the system:
+//			-The TNetworkedSimulationModel is mostly public and exposed. It is not recommend to publicly expose to your "user" code (high level scripting, designers, etc).
+//			-TNetworkSimStateAccessor is a helper for safely reading/writing to state within the system.
+//
+// ----------------------------------------------------------------------------------------------------------------------------------------------------------------------
+
+template <
+	typename InTSimulation,								// Final Simulation class
+	typename TUserBufferTypes,							// The user types (input, sync, aux, debug). Note this gets wrapped in TInternalBufferTypes internally.
+	typename InTTickSettings=TNetworkSimTickSettings<>, // Defines global rules about time keeping and ticking
+
+	// Core proxies that dictate how data replicates and how the simulation evolves for the three main roles
+	typename TRepProxyServerRPC =	TReplicator_Server		<TInternalBufferTypes<TUserBufferTypes, InTTickSettings>,	InTTickSettings >,
+	typename TRepProxyAutonomous =	TReplicator_Autonomous	<TInternalBufferTypes<TUserBufferTypes, InTTickSettings>,	InTTickSettings>,
+	typename TRepProxySimulated =	TReplicator_Simulated	<TInternalBufferTypes<TUserBufferTypes, InTTickSettings>,	InTTickSettings>,
+
+	// Defines how replication happens on these special channels, but doesn't dictate how simulation evolves
+	typename TRepProxyReplay =		TReplicator_Sequence	<TInternalBufferTypes<TUserBufferTypes, InTTickSettings>,	InTTickSettings, ENetworkSimBufferTypeId::Sync,  3>,
+	typename TRepProxyDebug =		TReplicator_Debug		<TInternalBufferTypes<TUserBufferTypes, InTTickSettings>,	InTTickSettings>
+>
+class TNetworkedSimulationModel : public INetworkSimulationModel
+{
+public:
+
+	using TSimulation = InTSimulation;
+	using TDriver = TNetworkedSimulationModelDriver<TUserBufferTypes>;
+
+	using TBufferTypes = TInternalBufferTypes<TUserBufferTypes, InTTickSettings>;
+	using TTickSettings = InTTickSettings;
+
+	using TInputCmd = typename TBufferTypes::TInputCmd;
+	using TSyncState = typename TBufferTypes::TSyncState;
+	using TAuxState = typename TBufferTypes::TAuxState;
+	using TDebugState = typename TBufferTypes::TDebugState;
+
+	using TSimTime = FNetworkSimTime;
+	using TRealTime = FNetworkSimTime::FRealTime;
+	
+	TNetworkedSimulationModel(TSimulation* InSimulation, TDriver* InDriver, const TSyncState& InitialSyncState = TSyncState(), const TAuxState& InitialAuxState = TAuxState())
+	{
+		check(InSimulation && InDriver);
+		Simulation = InSimulation;
+		Driver = InDriver;
+		*Buffers.Sync.WriteFrame(0) = InitialSyncState;
+		*Buffers.Aux.WriteFrame(0) = InitialAuxState;
+		Ticker.SetTotalProcessedSimulationTime(FNetworkSimTime(), 0);
+		DO_NETSIM_MODEL_DEBUG(FNetworkSimulationModelDebuggerManager::Get().RegisterNetworkSimulationModel(this, Driver->GetVLogOwner()));
+	}
+
+	virtual ~TNetworkedSimulationModel()
+	{
+		SetParentSimulation(nullptr);
+		ClearAllDependentSimulations();
+	}
+
+	void Tick(const FNetSimTickParameters& Parameters) final override
+	{
+		// Update previous DebugState based on what we (might) have sent *after* our last Tick 
+		// (property replication and ServerRPC get sent after the tick, rather than forcing awkward callback into the NetSim post replication, we can check it here)
+		if (auto* DebugBuffer = GetLocalDebugBuffer())
+		{
+			if (TDebugState* const PrevDebugState = DebugBuffer->Get(DebugBuffer->HeadFrame()))
+			{
+				if (Parameters.Role == ROLE_AutonomousProxy)
+				{
+					PrevDebugState->LastSentInputFrame = RepProxy_ServerRPC.GetLastSerializedFrame();
+				}
+				else if (Parameters.Role == ROLE_Authority)
+				{
+					PrevDebugState->LastSentInputFrame = RepProxy_Autonomous.GetLastSerializedFrame();
+				}
+			}
+		}
+
+		// Current frame debug state
+		TDebugState* const DebugState = GetNextLocalDebugStateWrite();
+		if (DebugState)
+		{
+			*DebugState = TDebugState();
+			DebugState->LocalDeltaTimeSeconds = Parameters.LocalDeltaTimeSeconds;
+			DebugState->LocalGFrameNumber = GFrameNumber;
+			DebugState->ProcessedFrames.Reset();
+			
+			if (Parameters.Role == ROLE_AutonomousProxy)
+			{
+				DebugState->LastReceivedInputFrame = RepProxy_Autonomous.GetLastSerializedFrame();
+			}
+			else if (Parameters.Role == ROLE_Authority)
+			{
+				DebugState->LastReceivedInputFrame = RepProxy_ServerRPC.GetLastSerializedFrame();
+			}
+		}
+
+		// ----------------------------------------------------------------------------------------------------------------
+		//	PreSimTick
+		//	This is the beginning of a new frame. PreSimTick will decide if we should take Parameters.LocalDeltaTimeSeconds
+		//	and advance the simulation or not. It will also generate new local input if necessary.
+		// ----------------------------------------------------------------------------------------------------------------
+		switch (Parameters.Role)
+		{
+			case ROLE_Authority:
+				RepProxy_ServerRPC.template PreSimTick<TDriver>(Driver, Buffers, Ticker, Parameters);
+			break;
+
+			case ROLE_AutonomousProxy:
+				RepProxy_Autonomous.template PreSimTick<TDriver>(Driver, Buffers, Ticker, Parameters);
+			break;
+
+			case ROLE_SimulatedProxy:
+				RepProxy_Simulated.template PreSimTick<TDriver>(Driver, Buffers, Ticker, Parameters);
+			break;
+		}
+
+		// -------------------------------------------------------------------------------------------------------------------------------------------------
+		//												Input Processing & Simulation Update
+		// -------------------------------------------------------------------------------------------------------------------------------------------------
+		while (Ticker.PendingFrame <= Ticker.MaxAllowedFrame)
+		{
+			const int32 InputFrame = Ticker.PendingFrame;
+			const int32 OutputFrame = Ticker.PendingFrame + 1;
+
+			const TInputCmd* InputCmd = Buffers.Input[InputFrame];
+			if (!ensureMsgf(InputCmd, TEXT("No InputCmd available for Frame %d. PendingFrame: %d. MaxAllowedFrame: %d."), InputFrame, Ticker.PendingFrame, Ticker.MaxAllowedFrame))
+			{
+				break;
+			}
+
+			// We have an unprocessed command, do we have enough allotted simulation time to process it?
+			if (Ticker.GetRemainingAllowedSimulationTime() >= InputCmd->GetFrameDeltaTime())
+			{
+				TSyncState* InSyncState = Buffers.Sync[InputFrame];
+				if (InSyncState == nullptr)
+				{
+					// We don't have valid sync state for this frame. This could mean we skipped ahead somehow, such as heavy packet loss where the server doesn't receive all input cmds
+					// We will use the last known sync state, but first write it out to the current InputFrame
+					UE_LOG(LogNetworkSim, Log, TEXT("%s: No sync state found @ frame %d. Using previous head element @ frame %d."), *Driver->GetDebugName(), InputFrame, Buffers.Sync.HeadFrame());
+					InSyncState = Buffers.Sync.WriteFrameInitializedFromHead(InputFrame);
+				}
+
+				const TAuxState* InAuxState = Buffers.Aux[InputFrame];
+				checkf(InAuxState, TEXT("No AuxState available for Frame %d. PendingFrame: %d. MaxAllowedFrame: %d."), InputFrame, Ticker.PendingFrame, Ticker.MaxAllowedFrame);
+
+				TSyncState* OutSyncState = Buffers.Sync.WriteFrame(OutputFrame);
+				check(OutSyncState);
+
+				{
+					TScopedSimulationTick UpdateScope(Ticker, OutputFrame, InputCmd->GetFrameDeltaTime());
+					
+					Simulation->SimulationTick( 
+						{ InputCmd->GetFrameDeltaTime(), Ticker },
+						{ *InputCmd, *InSyncState, *InAuxState },
+						{ *OutSyncState, Buffers.Aux.LazyWriter(OutputFrame) } );
+													
+				}
+
+				if (DebugState)
+				{
+					DebugState->ProcessedFrames.Add(InputFrame);
+				}
+			}
+			else
+			{
+				break;
+			}
+		}
+
+		// -------------------------------------------------------------------------------------------------------------------------------------------------
+		//												Post Sim Tick: finalize the frame
+		// -------------------------------------------------------------------------------------------------------------------------------------------------
+
+		switch (Parameters.Role)
+		{
+			case ROLE_Authority:
+				RepProxy_ServerRPC.template PostSimTick<TDriver>(Driver, Buffers, Ticker, Parameters);
+			break;
+
+			case ROLE_AutonomousProxy:
+				RepProxy_Autonomous.template PostSimTick<TDriver>(Driver, Buffers, Ticker, Parameters);
+			break;
+
+			case ROLE_SimulatedProxy:
+				RepProxy_Simulated.template PostSimTick<TDriver>(Driver, Buffers, Ticker, Parameters);
+			break;
+		}
+
+		// -------------------------------------------------------------------------------------------------------------------------------------------------
+		//														Debug
+		// -------------------------------------------------------------------------------------------------------------------------------------------------
+
+		// Finish debug state buffer recording (what the server processed each frame)
+		if (DebugState)
+		{
+			DebugState->PendingFrame = Ticker.PendingFrame;
+			DebugState->HeadFrame = Buffers.Input.HeadFrame();
+			DebugState->RemainingAllowedSimulationTimeSeconds = (float)Ticker.GetRemainingAllowedSimulationTime().ToRealTimeSeconds();
+		}
+
+		// Historical data recording (longer buffers for historical reference)
+		if (auto* HistoricData = GetHistoricBuffers())
+		{
+			HistoricData->Input.CopyAndMerge(Buffers.Input);
+			HistoricData->Sync.CopyAndMerge(Buffers.Sync);
+			HistoricData->Aux.CopyAndMerge(Buffers.Aux);
+		}
+	}
+
+	virtual void Reconcile(const ENetRole Role) final override
+	{
+		// --------------------------------------------------------------------------------------------------------------------------
+		//	Reconcile
+		//	This will eventually be called outside the Tick loop, only after processing a network bunch
+		//	Reconcile is about "making things right" after a network update. We are not processing "more" simulation yet.
+		// --------------------------------------------------------------------------------------------------------------------------
+		switch (Role)
+		{
+			case ROLE_Authority:
+				RepProxy_ServerRPC.template Reconcile<TSimulation, TDriver>(Simulation, Driver, Buffers, Ticker);
+			break;
+
+			case ROLE_AutonomousProxy:
+				RepProxy_Autonomous.template Reconcile<TSimulation, TDriver>(Simulation, Driver, Buffers, Ticker);
+			break;
+
+			case ROLE_SimulatedProxy:
+				RepProxy_Simulated.template Reconcile<TSimulation, TDriver>(Simulation, Driver, Buffers, Ticker);
+			break;
+		}
+	}
+	
+	void InitializeForNetworkRole(const ENetRole Role, const FNetworkSimulationModelInitParameters& Parameters) final override
+	{
+		// FIXME: buffer sizes are now inlined allocated but we want to support role based buffer sizes
+
+		//Buffers.Input.SetBufferSize(Parameters.InputBufferSize);
+		//Buffers.Sync.SetBufferSize(Parameters.SyncedBufferSize);
+		//Buffers.Aux.SetBufferSize(Parameters.AuxBufferSize); AUXFIXME
+
+		if (GetLocalDebugBuffer())
+		{
+			//GetLocalDebugBuffer()->SetBufferSize(Parameters.DebugBufferSize);
+		}
+
+		if (auto* MyHistoricBuffers = GetHistoricBuffers(true))
+		{
+			//MyHistoricBuffers->Input.SetBufferSize(Parameters.HistoricBufferSize);
+			//MyHistoricBuffers->Sync.SetBufferSize(Parameters.HistoricBufferSize);
+			//MyHistoricBuffers->Aux.SetBufferSize(Parameters.HistoricBufferSize); AUXFIXME
+		}
+
+		//Ticker.InitSimulationTimeBuffer(Parameters.SyncedBufferSize);
+	}
+
+	void NetSerializeProxy(EReplicationProxyTarget Target, const FNetSerializeParams& Params) final override
+	{
+		switch(Target)
+		{
+		case EReplicationProxyTarget::ServerRPC:
+			RepProxy_ServerRPC.NetSerialize(Params, Buffers, Ticker);
+			break;
+		case EReplicationProxyTarget::AutonomousProxy:
+			RepProxy_Autonomous.NetSerialize(Params, Buffers, Ticker);
+			break;
+		case EReplicationProxyTarget::SimulatedProxy:
+			RepProxy_Simulated.NetSerialize(Params, Buffers, Ticker);
+			break;
+		case EReplicationProxyTarget::Replay:
+			RepProxy_Replay.NetSerialize(Params, Buffers, Ticker);
+			break;
+		case EReplicationProxyTarget::Debug:
+#if NETSIM_MODEL_DEBUG
+			RepProxy_Debug.NetSerialize(Params, Buffers, Ticker);
+			break;
+#endif
+		default:
+			checkf(false, TEXT("Unknown: %d"), (int32)Target);
+		};
+	}
+
+	int32 GetProxyDirtyCount(EReplicationProxyTarget Target) final override
+	{
+		switch(Target)
+		{
+		case EReplicationProxyTarget::ServerRPC:
+			return RepProxy_ServerRPC.GetProxyDirtyCount(Buffers);
+		case EReplicationProxyTarget::AutonomousProxy:
+			return RepProxy_Autonomous.GetProxyDirtyCount(Buffers);
+		case EReplicationProxyTarget::SimulatedProxy:
+			return RepProxy_Simulated.GetProxyDirtyCount(Buffers);
+		case EReplicationProxyTarget::Replay:
+			return RepProxy_Replay.GetProxyDirtyCount(Buffers);
+		case EReplicationProxyTarget::Debug:
+#if NETSIM_MODEL_DEBUG
+			return RepProxy_Debug.GetProxyDirtyCount(Buffers);
+#endif
+		default:
+			checkf(false, TEXT("Unknown: %d"), (int32)Target);
+			return 0;
+		};
+	}
+
+	ESimulatedUpdateMode GetSimulatedUpdateMode() const
+	{
+		return RepProxy_Simulated.GetSimulatedUpdateMode();
+	}
+
+	// ------------------------------------------------------------------------------------------------------
+
+	void SetParentSimulation(INetworkSimulationModel* ParentSimulation) final override
+	{
+		if (RepProxy_Simulated.ParentSimulation)
+		{
+			RepProxy_Simulated.ParentSimulation->RemoveDependentSimulation(this);
+		}
+		
+		RepProxy_Simulated.ParentSimulation = ParentSimulation;
+		if (ParentSimulation)
+		{
+			ParentSimulation->AddDependentSimulation(this);
+		}
+	}
+
+	INetworkSimulationModel* GetParentSimulation() const final override
+	{
+		return RepProxy_Simulated.ParentSimulation;
+	}
+
+	void AddDependentSimulation(INetworkSimulationModel* DependentSimulation) final override
+	{
+		check(RepProxy_Autonomous.DependentSimulations.Contains(DependentSimulation) == false);
+		RepProxy_Autonomous.DependentSimulations.Add(DependentSimulation);
+		NotifyDependentSimNeedsReconcile(); // force reconcile on purpose
+	}
+
+	void RemoveDependentSimulation(INetworkSimulationModel* DependentSimulation) final override
+	{
+		RepProxy_Autonomous.DependentSimulations.Remove(DependentSimulation);
+	}
+
+	void NotifyDependentSimNeedsReconcile()
+	{
+		RepProxy_Autonomous.bDependentSimulationNeedsReconcile = true;
+	}
+
+	void BeginRollback(const FNetworkSimTime& RollbackDeltaTime, const int32 ParentFrame) final override
+	{
+		RepProxy_Simulated.template DependentRollbackBegin<TSimulation, TDriver>(Simulation, Driver, Buffers, Ticker, RollbackDeltaTime, ParentFrame);
+	}
+
+	void StepRollback(const FNetworkSimTime& Step, const int32 ParentFrame, const bool bFinalStep) final override
+	{
+		RepProxy_Simulated.template DependentRollbackStep<TSimulation, TDriver>(Simulation, Driver, Buffers, Ticker, Step, ParentFrame, bFinalStep);
+	}
+
+	void ClearAllDependentSimulations()
+	{
+		TArray<INetworkSimulationModel*> LocalList = MoveTemp(RepProxy_Autonomous.DependentSimulations);
+		for (INetworkSimulationModel* DependentSim : LocalList)
+		{
+			DependentSim->SetParentSimulation(nullptr);
+		}
+	}
+
+	// -------------------------------------------------------------------------------------------------------
+
+	TSimulation* Simulation = nullptr;
+	TDriver* Driver = nullptr;
+
+	TSimulationTicker<TTickSettings> Ticker;
+
+	TNetworkSimBufferContainer<TBufferTypes> Buffers;
+
+	TRepProxyServerRPC RepProxy_ServerRPC;
+	TRepProxyAutonomous RepProxy_Autonomous;
+	TRepProxySimulated RepProxy_Simulated;
+	TRepProxyReplay RepProxy_Replay;
+
+	using TInputBuffer = typename TNetworkSimBufferContainer<TBufferTypes>::TInputBuffer;
+	using TSyncBuffer = typename TNetworkSimBufferContainer<TBufferTypes>::TSyncBuffer;
+	using TAuxBuffer = typename TNetworkSimBufferContainer<TBufferTypes>::TAuxBuffer;
+	using TDebugBuffer = typename TNetworkSimBufferContainer<TBufferTypes>::TDebugBuffer;
+
+	// ------------------------------------------------------------------
+	// RPC Sending helper: provides basic send frequency settings for tracking when the Server RPC can be invoked.
+	// Note that the Driver is the one that must call the RPC, that cannot be rolled into this templated structure.
+	// More flexbile/dynamic send rates may be desireable. There is not reason this *has* to be done here, it could
+	// completely be tracked at the driver level, but that will also push more boilerplate to that layer for users.
+	// ------------------------------------------------------------------
+
+	void SetDesiredServerRPCSendFrequency(float DesiredHz) final override { ServerRPCThresholdTimeSeconds = 1.f / DesiredHz; }
+	bool ShouldSendServerRPC(float DeltaTimeSeconds) final override
+	{
+		// Don't allow a large delta time to pollute the accumulator
+		const float CappedDeltaTimeSeconds = FMath::Min<float>(DeltaTimeSeconds, ServerRPCThresholdTimeSeconds);
+		ServerRPCAccumulatedTimeSeconds += DeltaTimeSeconds;
+		if (ServerRPCAccumulatedTimeSeconds >= ServerRPCThresholdTimeSeconds)
+		{
+			ServerRPCAccumulatedTimeSeconds -= ServerRPCThresholdTimeSeconds;
+			return true;
+		}
+
+		return false;
+	}
+
+	// Simulation class should have a static const FName GroupName member
+	FName GetSimulationGroupName() const final override { return TSimulation::GroupName; }
+
+private:
+	float ServerRPCAccumulatedTimeSeconds = 0.f;
+	float ServerRPCThresholdTimeSeconds = 1.f / 999.f; // Default is to send at a max of 999hz. This part of the system needs to be build out more (better handling of super high FPS clients and fixed rate servers)
+
+	// ------------------------------------------------------------------
+	//	Debugging
+	// ------------------------------------------------------------------
+public:
+
+#if NETSIM_MODEL_DEBUG
+	TDebugBuffer* GetLocalDebugBuffer() {	return &Buffers.Debug; }
+	TDebugState* GetNextLocalDebugStateWrite() { return Buffers.Debug.WriteFrame( Buffers.Debug.HeadFrame() + 1 ); }
+	TNetworkSimBufferContainer<TBufferTypes>* GetHistoricBuffers(bool bCreate=false)
+	{
+		if (HistoricBuffers.IsValid() == false && bCreate) { HistoricBuffers.Reset(new TNetworkSimBufferContainer<TBufferTypes>()); }
+		return HistoricBuffers.Get();
+	}
+
+	TDebugBuffer* GetRemoteDebugBuffer() {	return &RepProxy_Debug.ReceivedBuffer; }
+#else
+	TDebugBuffer* GetLocalDebugBuffer() {	return nullptr; }
+	TDebugState* GetNextLocalDebugStateWrite() { return nullptr; }
+	TNetworkSimBufferContainer<TBufferTypes>* GetHistoricBuffers(bool bCreate=false) { return nullptr; }
+	TDebugBuffer* GetRemoteDebugBuffer() {	return nullptr; }
+#endif
+
+private:
+
+#if NETSIM_MODEL_DEBUG
+	TRepProxyDebug RepProxy_Debug;
+	TUniquePtr<TNetworkSimBufferContainer<TBufferTypes>> HistoricBuffers;
+#endif
+
+};
\ No newline at end of file
diff --git a/Engine/Plugins/Runtime/NetworkPrediction/Source/NetworkPrediction/Public/NetworkedSimulationModelBuffer.h b/Engine/Plugins/Runtime/NetworkPrediction/Source/NetworkPrediction/Public/NetworkedSimulationModelBuffer.h
new file mode 100644
--- /dev/null
+++ b/Engine/Plugins/Runtime/NetworkPrediction/Source/NetworkPrediction/Public/NetworkedSimulationModelBuffer.h
@@ -0,0 +1,366 @@
+// Copyright 1998-2019 Epic Games, Inc. All Rights Reserved.
+
+#pragma once
+
+// Generic iterator for iterating over each frame from tail to head
+template<typename TBuffer, typename ElementType>
+struct TNetworkSimBufferIterator
+{
+	TNetworkSimBufferIterator(TBuffer& InBuffer) : Buffer(InBuffer)	{ CurrentPos = Buffer.IterStartPos(); }
+	int32 Frame() const { return Buffer.IterFrame(CurrentPos); }
+	ElementType* Element() const { return Buffer.IterElement(CurrentPos); }
+
+	void operator++() { CurrentPos++; }
+	operator bool()	{ return CurrentPos != INDEX_NONE && CurrentPos <= Buffer.IterEndPos(); }
+
+private:
+	TBuffer& Buffer;
+	int32 CurrentPos = -1;
+};
+
+// Helper struct to encapsulate optional, delayed writing of new element to buffer
+template<typename ElementType>
+struct TNetSimLazyWriterFunc
+{
+	template<typename TBuffer>
+	TNetSimLazyWriterFunc(TBuffer& Buffer, int32 PendingFrame)
+	{
+		GetFunc = [&Buffer, PendingFrame]() { return Buffer.WriteFrameInitializedFromHead(PendingFrame); };
+	}
+
+	ElementType* Get() const
+	{
+		check(GetFunc);
+		return (ElementType*)GetFunc();
+	}
+
+	TFunction<void*()> GetFunc;
+};
+
+// Reference version of LazyWriter. Faster to pass around and do cascading downcasts
+template<typename ElementType>
+struct TNetSimLazyWriter
+{
+	template<typename TLazyWriter>
+	TNetSimLazyWriter(const TLazyWriter& Parent)
+		: GetFunc(Parent.GetFunc) { }
+
+	ElementType* Get() const
+	{
+		return (ElementType*)GetFunc();
+	}
+
+	TFunctionRef<void*()> GetFunc;
+};
+
+// Common implementations shared by TNetworkSimContiguousBuffer/TNetworkSimSparseBuffer
+template<typename T>
+struct TNetworkSimBufferBase
+{
+	FString GetBasicDebugStr() const
+	{
+		return FString::Printf(TEXT("Elements: [%d/%d]. Frames: [%d-%d]"), ((T*)this)->Num(), ((T*)this)->Max(), ((T*)this)->TailFrame(), ((T*)this)->HeadFrame());
+	}
+
+	bool IsValidFrame(int32 Frame) const 
+	{ 
+		return Frame >= ((T*)this)->TailFrame() && Frame <= ((T*)this)->HeadFrame(); 
+	}
+
+	// Copies each element of SourceBuffer into this buffer. This could break continuity so for TNetworkSimContiguousBuffer, it is possible to lose elements by doing this.
+	template<typename TSource>
+	void CopyAndMerge(const TSource& SourceBuffer)
+	{
+		for (auto It = SourceBuffer.CreateConstIterator(); It; ++It)
+		{
+			*((T*)this)->WriteFrame( It.Frame() ) = *It.Element();
+		}
+	}
+
+protected:
+
+	// Creates a new frame, but
+	//	-If frame already exists, returns existing
+	//	-If frame > head, contents of head are copied into new frame
+	//	-If frame < tail, contents are cleared to default value
+	// (this function is sketchy and poorly named but is proving to be useful in a few places)
+	template<typename ElementType>
+	ElementType* WriteFrameInitializedFromHeadImpl(int32 Frame)
+	{
+		if (((T*)this)->HeadFrame() > Frame)
+		{
+			ElementType* HeadElementPtr = ((T*)this)->HeadElement();
+			ElementType* NewElement = ((T*)this)->WriteFrame(Frame);
+			*NewElement = *HeadElementPtr;
+			return NewElement;
+		}
+		else if (Frame < ((T*)this)->TailFrame())
+		{
+			ElementType* NewElement = ((T*)this)->WriteFrame(Frame);
+			*NewElement = ElementType();
+			return NewElement;
+		}
+		
+		return ((T*)this)->WriteFrame(Frame);
+	}
+};
+
+template<typename T, int32 NumElements=32>
+struct TNetworkSimContiguousBuffer : public TNetworkSimBufferBase<TNetworkSimContiguousBuffer<T, NumElements>>
+{
+	using ElementType = T;
+	TNetworkSimContiguousBuffer()
+	{
+		Data.SetNum(NumElements);
+	}
+
+	T* operator[](int32 Frame) {	return const_cast<T*>(GetImpl(Frame)); }
+	const T* operator[](int32 Frame) const { return GetImpl(Frame); }
+
+	T* Get(int32 Frame) { return const_cast<T*>(GetImpl(Frame)); }
+	const T* Get(int32 Frame) const { return GetImpl(Frame); }
+
+	const T* HeadElement() const { return HeadElementImpl(); }
+	T* HeadElement() { return const_cast<T*>(HeadElementImpl()); }
+
+	const T* TailElement() const { return TailElementImpl(); }
+	T* TailElement() { return const_cast<T*>(TailElementImpl()); }
+
+	int32 HeadFrame() const { return Head; }
+	int32 TailFrame() const { return GetTail(); }
+
+	int32 Num() const { return Head == INDEX_NONE ? 0 : (Head - GetTail() + 1); }
+	int32 Max() const { return Data.Max(); }
+	int32 GetDirtyCount() const { return DirtyCount; }
+
+	T* WriteFrame(int32 Frame)
+	{
+		check(Frame >= 0);
+
+		const int32 Tail = GetTail();
+		if (Head == INDEX_NONE || Frame < Tail || Frame > Head+1)
+		{
+			// Writing outside the current range (+1) results in a full wipe of valid contents
+			NumValidElements = 1;
+		}
+		else
+		{
+			// Writing inside the current range (+1) preserves valid elements up to the size of our buffer
+			NumValidElements = FMath::Min<int32>(Frame - Tail + 1, Data.Num());
+		}
+		
+		Head = Frame;
+		++DirtyCount;
+		return &Data[Frame % Data.Num()];
+	}
+	
+	T* WriteFrameInitializedFromHead(int32 Frame)
+	{
+		return this->template WriteFrameInitializedFromHeadImpl<ElementType>(Frame);
+	}
+
+	TNetSimLazyWriterFunc<ElementType> LazyWriter(int32 Frame)
+	{
+		return TNetSimLazyWriterFunc<ElementType>(*this, Frame);
+	}
+
+	using IteratorType = TNetworkSimBufferIterator<TNetworkSimContiguousBuffer<ElementType, NumElements>, ElementType>;
+	using ConstIteratorType = TNetworkSimBufferIterator<const TNetworkSimContiguousBuffer<ElementType, NumElements>, const ElementType>;
+
+	IteratorType CreateIterator() { return IteratorType(*this); }
+	ConstIteratorType CreateConstIterator() const { return ConstIteratorType(*this); }
+
+private:
+
+	const T* GetImpl(int32 Frame) const
+	{
+		const int32 RelativeToHead = Frame - Head;
+		if (RelativeToHead > 0 || RelativeToHead <= -NumValidElements)
+		{
+			return nullptr;
+		}
+
+		return &Data[Frame % Data.Num()];
+	}
+
+	const T* HeadElementImpl() const
+	{
+		if (Head != INDEX_NONE)
+		{
+			return &Data[Head % Data.Num()];
+		}
+		return nullptr;
+	}
+
+	const T* TailElementImpl() const
+	{
+		const int32 Tail = GetTail();
+		if (Tail != INDEX_NONE)
+		{
+			return &Data[Tail % Data.Num()];
+		}
+		return nullptr;
+	}
+
+	int32 GetTail() const { return Head == INDEX_NONE ? INDEX_NONE : (Head - NumValidElements + 1); };
+
+	int32 DirtyCount = 0;
+	int32 Head = INDEX_NONE;
+	int32 NumValidElements = 0;
+	TArray<ElementType, TInlineAllocator<NumElements>> Data;
+
+	// Efficient position based accessing through TNetworkSimBufferIterator.
+	template<typename, typename>
+	friend struct TNetworkSimBufferIterator;
+	int32 IterStartPos() const { return GetTail(); }
+	int32 IterEndPos() const { return Head; }
+	int32 IterFrame(int32 Pos) const { return Pos; }
+	T* IterElement(int32 Pos) { return &Data[Pos % Data.Num()]; }
+	const T* IterElement(int32 Pos) const { return &Data[Pos % Data.Num()]; }
+};
+
+// Frame: arbitrary identifier for data. Not contiguous or controlled by us (always passed in)
+// Position/Pos: increasing counter for position in array. Mod with Data.Num to get Index.
+// Index/Idx: actual index into Data Array
+template<typename T, int32 NumElements=32>
+struct TNetworkSimSparseBuffer : public TNetworkSimBufferBase<TNetworkSimSparseBuffer<T, NumElements>>
+{
+	using ElementType = T;
+	TNetworkSimSparseBuffer()
+	{
+		Data.SetNum(NumElements);
+	}
+
+	T* operator[](int32 Frame) {	return const_cast<T*>(GetImpl(Frame)); }
+	const T* operator[](int32 Frame) const { return GetImpl(Frame); }
+
+	T* Get(int32 Frame) { return const_cast<T*>(GetImpl(Frame)); }
+	const T* Get(int32 Frame) const { return GetImpl(Frame); }
+
+	const T* HeadElement() const { return HeadElementImpl(); }
+	T* HeadElement() { return const_cast<T*>(HeadElementImpl()); }
+
+	const T* TailElement() const { return TailElementImpl(); }
+	T* TailElement() { return const_cast<T*>(TailElementImpl()); }
+
+	int32 HeadFrame() const { return Data[HeadPos % Data.Num()].Frame; }
+	int32 TailFrame() const { return Data[GetTailPos() % Data.Num()].Frame; }
+
+	int32 Num() const { return HeadPos == INDEX_NONE ? 0 : (HeadPos - GetTailPos() + 1); }
+	int32 Max() const { return Data.Max(); }
+	int32 GetDirtyCount() const { return DirtyCount; }
+
+	// Returns the element @ frame for writing. The contents of this element are unknown (could be stale content!). Note the element returned is immediately considered "valid" by Num(), Iterators, etc!
+	T* WriteFrame(int32 Frame)
+	{
+		check(Frame >= 0);
+		++DirtyCount;
+
+		int32 Pos = HeadPos;
+		if (HeadPos != INDEX_NONE)
+		{
+			// If we have elements, find where this Frame would go
+			const int32 Tail = GetTailPos();
+			for (; Pos >= Tail; --Pos)
+			{
+				TInternal& InternalData = Data[Pos % Data.Num()];
+				if (InternalData.Frame == Frame)
+				{
+					return &InternalData.Element;
+				}
+				if (InternalData.Frame <= Frame)
+				{
+					break;
+				}
+			}
+		}
+
+		// Write frame to next pos
+		HeadPos = Pos+1;
+		TInternal& WriteData = Data[HeadPos % Data.Num()];
+		WriteData.Frame = Frame;
+		return &WriteData.Element;
+	}
+
+	ElementType* WriteFrameInitializedFromHead(int32 Frame)
+	{
+		return this->template WriteFrameInitializedFromHeadImpl<ElementType>(Frame);
+	}
+
+	TNetSimLazyWriterFunc<ElementType> LazyWriter(int32 Frame)
+	{
+		return TNetSimLazyWriterFunc<ElementType>(*this, Frame);
+	}
+
+	using IteratorType = TNetworkSimBufferIterator<TNetworkSimSparseBuffer<ElementType, NumElements>, ElementType>;
+	using ConstIteratorType = TNetworkSimBufferIterator<const TNetworkSimSparseBuffer<ElementType, NumElements>, const ElementType>;
+
+	IteratorType CreateIterator() { return IteratorType(*this); }
+	ConstIteratorType CreateConstIterator() const { return ConstIteratorType(*this); }
+
+private:
+
+	struct TInternal
+	{
+		int32 Frame = -1;
+		T Element;
+	};
+
+	const T* GetImpl(int32 Frame) const
+	{
+		const int32 Pos = GetPosForFrame(Frame);
+		if (Pos != INDEX_NONE)
+		{		
+			return &Data[Pos % Data.Num()].Element;
+		}
+
+		return nullptr;
+	}
+
+	const T* HeadElementImpl() const
+	{
+		if (HeadPos != INDEX_NONE)
+		{
+			return &Data[HeadPos % Data.Num()].Element;
+		}
+		return nullptr;
+	}
+
+	const T* TailElementImpl() const
+	{
+		const int32 TailPos = GetTailPos();
+		if (TailPos != INDEX_NONE)
+		{
+			return &Data[TailPos % Data.Num()].Element;
+		}
+		return nullptr;
+	}
+	
+	int32 GetPosForFrame(int32 Frame) const
+	{
+		const int32 TailPos = GetTailPos();
+		for (int32 Pos=HeadPos; Pos >= TailPos; --Pos)
+		{
+			if (Data[Pos % Data.Num()].Frame <= Frame)
+			{
+				return Pos;
+			}
+		}
+		return INDEX_NONE;
+	}
+
+	int32 GetTailPos() const { return HeadPos == INDEX_NONE ? INDEX_NONE : FMath::Max<int32>(0, HeadPos - Data.Num()); }
+
+	int32 DirtyCount=0;
+	int32 HeadPos=INDEX_NONE;
+	TArray<TInternal, TInlineAllocator<NumElements>> Data;
+
+	// Efficient position based accessing through TNetworkSimBufferIterator.
+	template<typename, typename>
+	friend struct TNetworkSimBufferIterator;
+	int32 IterStartPos() const { return GetTailPos(); }
+	int32 IterEndPos() const { return HeadPos; }
+	int32 IterFrame(int32 Pos) const { return Data[Pos % Data.Num()].Frame; }
+	T* IterElement(int32 Pos) { return &Data[Pos % Data.Num()].Element; }
+	const T* IterElement(int32 Pos) const { return &Data[Pos % Data.Num()].Element; }
+};
\ No newline at end of file
diff --git a/Engine/Plugins/Runtime/NetworkPrediction/Source/NetworkPrediction/Public/NetworkedSimulationModelCVars.h b/Engine/Plugins/Runtime/NetworkPrediction/Source/NetworkPrediction/Public/NetworkedSimulationModelCVars.h
new file mode 100644
--- /dev/null
+++ b/Engine/Plugins/Runtime/NetworkPrediction/Source/NetworkPrediction/Public/NetworkedSimulationModelCVars.h
@@ -0,0 +1,78 @@
+// Copyright 1998-2019 Epic Games, Inc. All Rights Reserved.
+
+#pragma once
+#include "HAL/IConsoleManager.h"
+
+// ------------------------------------------------------------------------------------------------------------
+//	"Shipping const" cvars: cvars that should compile out to const functions in shipping/test builds
+//		This got a little tricky due to templated usage across different modules. Previous patterns for
+//		this thing didn't work. This implementation requires manual finding of consolve variables so will be 
+//		a bit slower but shouldn't matter since it is compiled out of shipping/test.
+// ------------------------------------------------------------------------------------------------------------
+
+// Whether to treat these cvars as consts
+#define NETSIM_CONST_CVARS (UE_BUILD_SHIPPING || UE_BUILD_TEST)
+
+// This is required because these cvars live in header files that will be included across different compilation units
+// Just using standard FAutoConsoleVariableRef will cause multiple registrations of the same variable
+struct FConditionalAutoConsoleRegister
+{
+	FConditionalAutoConsoleRegister(const TCHAR* VarName, int32 Value, const TCHAR* Help)
+	{
+		if (!IConsoleManager::Get().FindConsoleVariable(VarName, false))
+		{
+			IConsoleManager::Get().RegisterConsoleVariable(VarName, Value, Help, ECVF_Cheat);
+		}
+	}
+
+	FConditionalAutoConsoleRegister(const TCHAR* VarName, float Value, const TCHAR* Help)
+	{
+		if (!IConsoleManager::Get().FindConsoleVariable(VarName, false))
+		{
+			IConsoleManager::Get().RegisterConsoleVariable(VarName, Value, Help, ECVF_Cheat);
+		}
+	}
+};
+
+
+#if NETSIM_CONST_CVARS
+#define NETSIM_DEVCVAR_SHIPCONST_INT(Var,Value,VarName,Help) \
+	inline int32 Var() { return Value; } \
+	inline void Set##Var(int32 V) { }
+#else
+#define NETSIM_DEVCVAR_SHIPCONST_INT(Var,Value,VarName,Help) \
+	static FConditionalAutoConsoleRegister Var##Auto(TEXT(VarName),(int32)Value,TEXT(Help)); \
+	inline int32 Var() \
+	{ \
+		static const auto* Existing = IConsoleManager::Get().FindConsoleVariable(TEXT(VarName), false); \
+		check(Existing); \
+		return Existing->GetInt(); \
+	} \
+	inline void Set##Var(int32 V) \
+	{ \
+		static auto* Existing = IConsoleManager::Get().FindConsoleVariable(TEXT(VarName), false); \
+		check(Existing); \
+		Existing->Set(V, ECVF_SetByConsole); \
+	}
+#endif
+
+#if NETSIM_CONST_CVARS
+#define NETSIM_DEVCVAR_SHIPCONST_FLOAT(Var,Value,VarName,Help) \
+	inline float Var() { return Value; } \
+	inline void Set##Var() { }
+#else
+#define NETSIM_DEVCVAR_SHIPCONST_FLOAT(Var,Value,VarName,Help) \
+	static FConditionalAutoConsoleRegister Var##Auto(TEXT(VarName),(float)Value,TEXT(Help)); \
+	inline float Var() \
+	{ \
+		static const auto* Existing = IConsoleManager::Get().FindConsoleVariable(TEXT(VarName), false); \
+		check(Existing); \
+		return Existing->GetFloat(); \
+	} \
+	inline void Set##Var(float V) \
+	{ \
+		static auto* Existing = IConsoleManager::Get().FindConsoleVariable(TEXT(VarName), false); \
+		check(Existing); \
+		Existing->Set(V, ECVF_SetByConsole); \
+	}
+#endif
\ No newline at end of file
diff --git a/Engine/Plugins/Runtime/NetworkPrediction/Source/NetworkPrediction/Public/NetworkedSimulationModelDebugger.h b/Engine/Plugins/Runtime/NetworkPrediction/Source/NetworkPrediction/Public/NetworkedSimulationModelDebugger.h
new file mode 100644
--- /dev/null
+++ b/Engine/Plugins/Runtime/NetworkPrediction/Source/NetworkPrediction/Public/NetworkedSimulationModelDebugger.h
@@ -0,0 +1,747 @@
+// Copyright 1998-2019 Epic Games, Inc. All Rights Reserved.
+
+#pragma once
+
+#include "Math/Color.h"
+#include "DrawDebugHelpers.h"
+#include "Debug/DebugDrawService.h"
+#include "Engine/Canvas.h"
+#include "Debug/ReporterGraph.h"
+#include "Engine/Engine.h"
+#include "UObject/WeakObjectPtrTemplates.h"
+#include "Tickable.h"
+#include "NetworkPredictionTypes.h"
+#include "CanvasItem.h"
+#include "Containers/Array.h"
+#include "Templates/UniquePtr.h"
+#include "NetworkedSimulationModelCVars.h"
+#include "VisualLogger/VisualLogger.h"
+
+DEFINE_LOG_CATEGORY_STATIC(LogNetworkSimDebug, Log, All);
+
+namespace NetworkSimulationModelDebugCVars
+{
+	NETSIM_DEVCVAR_SHIPCONST_INT(DrawFrames, 1, "nsm.debug.DrawFrames", "Draws frame data (text) in debug graphs");
+	NETSIM_DEVCVAR_SHIPCONST_INT(DrawNetworkSendLines, 1, "nsm.debug.DrawNetworkSendLines", "Draws lines representing network traffic in debugger");
+	NETSIM_DEVCVAR_SHIPCONST_INT(GatherServerSidePIE, 1, "nsm.debug.GatherServerSide", "Whenever we gather debug info from a client side actor, also gather server side equivelent. Only works in PIE.");
+}
+
+struct FNetworkSimulationModelDebuggerManager;
+
+NETWORKPREDICTION_API UObject* FindReplicatedObjectOnPIEServer(const UObject* ClientObject);
+
+// ------------------------------------------------------------------------------------------------------------------------
+//	Debugger support classes
+// ------------------------------------------------------------------------------------------------------------------------
+
+struct INetworkSimulationModelDebugger
+{
+	virtual ~INetworkSimulationModelDebugger() { }
+
+	bool IsActive() { return bActive; }
+	void SetActive(bool InActive) { bActive = InActive; }
+
+	virtual void GatherCurrent(FNetworkSimulationModelDebuggerManager& Out, UCanvas* C) = 0;
+	virtual void Tick( float DeltaTime ) = 0;
+
+protected:
+	bool bActive = false; // Whether you should draw every frame
+};
+
+struct NETWORKPREDICTION_API FNetworkSimulationModelDebuggerManager: public FTickableGameObject, FNoncopyable
+{
+	static FNetworkSimulationModelDebuggerManager& Get();
+
+	~FNetworkSimulationModelDebuggerManager()
+	{
+		if (Graph.IsValid())
+		{
+			Graph->RemoveFromRoot();
+		}
+	}
+
+	FNetworkSimulationModelDebuggerManager()
+	{
+		DrawDebugServicesHandle = UDebugDrawService::Register(TEXT("Game"), FDebugDrawDelegate::CreateRaw(this, &FNetworkSimulationModelDebuggerManager::DrawDebugService));
+		check(DrawDebugServicesHandle.IsValid());
+	}
+
+	// ---------------------------------------------------------------------------------------------------------------------------------------
+	//	Outside API (registration, console commands, draw services, etc)
+	// ---------------------------------------------------------------------------------------------------------------------------------------
+
+	template <typename T>
+	void RegisterNetworkSimulationModel(T* NetworkSim, const AActor* OwningActor);
+
+	void SetDebuggerActive(AActor* OwningActor, bool InActive)
+	{
+		if (INetworkSimulationModelDebugger* Debugger = Find(OwningActor))
+		{
+			Debugger->SetActive(InActive);
+		}
+		ResetCache();
+		Gather(LastCanvas.Get());
+	}
+
+	void ToggleDebuggerActive(AActor* OwningActor)
+	{
+		if (INetworkSimulationModelDebugger* Debugger = Find(OwningActor))
+		{
+			Debugger->SetActive(!Debugger->IsActive());
+		}
+		ResetCache();
+		Gather(LastCanvas.Get());
+	}
+
+	void SetContinousGather(bool InGather)
+	{
+		bContinousGather = InGather;
+		if (!bContinousGather)
+		{
+			Gather(LastCanvas.Get());
+		}
+	}
+
+	void ToggleContinousGather()
+	{
+		SetContinousGather(!bContinousGather);
+	}
+
+	void DrawDebugService(UCanvas* C, APlayerController* PC)
+	{
+		LastCanvas = C;
+		if (bContinousGather)
+		{
+			Gather(C);
+		}
+		
+		FDisplayDebugManager& DisplayDebugManager = C->DisplayDebugManager;
+		DisplayDebugManager.Initialize(C, GEngine->GetSmallFont(), FVector2D(4.0f, 150.0f));
+
+		if (Lines.Num() > 0)
+		{
+			const float TextScale = FMath::Max(C->SizeX / 1920.0f, 1.0f);
+			FCanvasTileItem TextBackgroundTile(FVector2D(0.0f, 120.0f), FVector2D(400.0f, 1800.0f) * TextScale, FColor(0, 0, 0, 100));
+			TextBackgroundTile.BlendMode = SE_BLEND_Translucent;
+			C->DrawItem(TextBackgroundTile);
+		}
+
+		// --------------------------------------------------------
+		//	Lines
+		// --------------------------------------------------------
+
+		for (FDebugLine& Line : Lines)
+		{
+			DisplayDebugManager.SetDrawColor(Line.Color);
+			DisplayDebugManager.DrawString(Line.Str);
+		}
+
+		// --------------------------------------------------------
+		//	Canvas Items (graphs+text)
+		// --------------------------------------------------------
+		
+		for (auto& Item : CanvasItems[0])
+		{
+			C->DrawItem(*Item.Get());
+		}
+
+		if (NetworkSimulationModelDebugCVars::DrawFrames() > 0)
+		{
+			for (auto& Item : CanvasItems[1])
+			{
+				C->DrawItem(*Item.Get());
+			}
+		}
+	}
+
+	virtual void Tick( float DeltaTime )
+	{
+		for (auto It = DebuggerMap.CreateIterator(); It; ++It)
+		{
+			const AActor* Owner = It.Key().Get();
+			if (!Owner)
+			{
+				It.RemoveCurrent();
+				continue;;
+			}
+			
+			if (It.Value()->IsActive())
+			{
+				It.Value()->Tick(DeltaTime);
+			}
+		}
+	}	
+
+	/** return the stat id to use for this tickable **/
+	virtual TStatId GetStatId() const override
+	{
+		RETURN_QUICK_DECLARE_CYCLE_STAT(FNetworkSimulationModelDebuggerManager, STATGROUP_TaskGraphTasks);
+	}
+
+	/** Gathers latest and Logs single frame */
+	void LogSingleFrame(FOutputDevice& Ar)
+	{
+		Gather(LastCanvas.Get());
+		
+		for (FDebugLine& Line : Lines)
+		{
+			Ar.Logf(TEXT("%s"), *Line.Str);
+		}
+	}
+
+	// ---------------------------------------------------------------------------------------------------------------------------------------
+	//	Debugging API used by TNetworkSimulationModelDebugger
+	// ---------------------------------------------------------------------------------------------------------------------------------------
+
+	void Emit(const FString& Str = FString(), FColor Color = FColor::White, float XOffset=0.f, float YOffset=0.f)
+	{
+		Lines.Emplace(FDebugLine{ Str, Color, XOffset, YOffset });
+	}
+
+	template <typename TBuffer>
+	void EmitElement(TBuffer& Buffer, const FStandardLoggingParameters& Parameters)
+	{
+		FStringOutputDevice StrOut;
+		StrOut.SetAutoEmitLineTerminator(true);
+
+		FStandardLoggingParameters LocalParameters = Parameters;
+		LocalParameters.Ar = &StrOut;
+
+		auto* Element = Buffer[LocalParameters.Frame];
+		if (Element)
+		{
+			Element->Log(LocalParameters);
+	
+			TArray<FString> StrLines;
+			StrOut.ParseIntoArrayLines(StrLines, true);
+			for (FString& Str : StrLines)
+			{
+				Emit(Str);
+			}
+		}
+	}
+
+	void EmitQuad(FVector2D ScreenPosition, FVector2D ScreenSize, FColor Color)
+	{
+		FVector2D Quad[4];
+		
+		Quad[0].X = ScreenPosition.X;
+		Quad[0].Y = ScreenPosition.Y;
+
+		Quad[1].X = ScreenPosition.X;
+		Quad[1].Y = ScreenPosition.Y + ScreenSize.Y;
+
+		Quad[2].X = ScreenPosition.X + ScreenSize.X;
+		Quad[2].Y = ScreenPosition.Y + ScreenSize.Y;
+
+		Quad[3].X = ScreenPosition.X + ScreenSize.X;
+		Quad[3].Y = ScreenPosition.Y;
+		
+		CanvasItems[0].Emplace( MakeUnique<FCanvasTriangleItem>(Quad[0], Quad[1], Quad[2], GWhiteTexture) );
+		CanvasItems[0].Last()->SetColor(Color);
+
+		CanvasItems[0].Emplace( MakeUnique<FCanvasTriangleItem>(Quad[2], Quad[3], Quad[0], GWhiteTexture) );
+		CanvasItems[0].Last()->SetColor(Color);
+	}
+
+	void EmitText(FVector2D ScreenPosition, FColor Color, const FString& Str)
+	{
+		CanvasItems[1].Emplace( MakeUnique<FCanvasTextItem>(ScreenPosition, FText::FromString(Str), GEngine->GetTinyFont(), Color) );
+	}
+
+	void EmitLine(FVector2D StartPosition, FVector2D EndPosition, FColor Color, float Thickness=1.f)
+	{
+		CanvasItems[0].Emplace( MakeUnique<FCanvasLineItem>(StartPosition, EndPosition) );
+		CanvasItems[0].Last()->SetColor(Color);
+		((FCanvasLineItem*)CanvasItems[0].Last().Get())->LineThickness = Thickness;
+	}
+
+private:
+
+	INetworkSimulationModelDebugger* Find(const AActor* Actor)
+	{
+		if (!Actor)
+		{
+			return nullptr;
+		}
+
+		INetworkSimulationModelDebugger* Debugger = DebuggerMap.FindRef(TWeakObjectPtr<const AActor>(Actor));
+		if (!Debugger)
+		{
+			UE_LOG(LogNetworkSimDebug, Warning, TEXT("Could not find NetworkSimulationModel associated with %s"), *GetPathNameSafe(Actor));
+		}
+		return Debugger;
+	}
+
+	void Gather(UCanvas* C)
+	{
+		ResetCache();
+
+		for (auto It = DebuggerMap.CreateIterator(); It; ++It)
+		{
+			const AActor* Owner = It.Key().Get();
+			if (!Owner)
+			{
+				It.RemoveCurrent();
+				continue;;
+			}
+			
+			if (It.Value()->IsActive())
+			{
+				It.Value()->GatherCurrent(*this, C);
+				if (NetworkSimulationModelDebugCVars::GatherServerSidePIE() > 0)
+				{
+					if (const AActor* ServerSideActor = Cast<AActor>(FindReplicatedObjectOnPIEServer(Owner)))
+					{
+						if (INetworkSimulationModelDebugger* ServerSideSim = Find(ServerSideActor))
+						{
+							Emit();
+							Emit();
+							ServerSideSim->GatherCurrent(*this, nullptr); // Dont do graphs for server side state
+						}
+					}
+				}
+
+				// Only gather first active debugger (it would be great to have more control over this when debugging multiples)
+				break;
+			}
+		}
+	}
+
+	void ResetCache()
+	{
+		Lines.Reset();
+		
+		CanvasItems[0].Reset();
+		CanvasItems[1].Reset();
+	}
+
+	TMap<TWeakObjectPtr<const AActor>, INetworkSimulationModelDebugger*>	DebuggerMap;
+	bool bContinousGather = true; // Whether you should gather new data every frame
+
+	FDelegateHandle DrawDebugServicesHandle;
+	struct FDebugLine
+	{
+		FString Str;
+		FColor Color;
+		float XOffset;
+		float YOffset;
+	};
+	TArray<FDebugLine> Lines;
+	TArray<TUniquePtr<FCanvasItem>> CanvasItems[2];
+	TWeakObjectPtr<UReporterGraph> Graph;
+	TWeakObjectPtr<UCanvas> LastCanvas;
+};
+
+template <typename TNetSimModel>
+struct TNetworkSimulationModelDebugger : public INetworkSimulationModelDebugger
+{
+	using TSimTime = typename TNetSimModel::TSimTime;
+	using TDebugState = typename TNetSimModel::TDebugState;
+
+	TNetworkSimulationModelDebugger(TNetSimModel* InNetSim, const AActor* OwningActor)
+	{
+		NetworkSim = InNetSim;
+		WeakOwningActor = OwningActor;
+	}
+
+	~TNetworkSimulationModelDebugger()
+	{
+		
+	}
+
+	struct FCachedScreenPositionMap
+	{
+		struct FScreenPositions
+		{
+			void SetSent(const FVector2D& In) { if (SentPosition == FVector2D::ZeroVector) SentPosition = In; }
+			void SetRecv(const FVector2D& In) { if (RecvPosition == FVector2D::ZeroVector) RecvPosition = In; }
+
+			FVector2D SentPosition = FVector2D::ZeroVector;
+			FVector2D RecvPosition = FVector2D::ZeroVector;
+		};
+
+		TMap<int32, FScreenPositions> Frames;
+	};
+
+	template<typename TBuffer>
+	void GatherDebugGraph(FNetworkSimulationModelDebuggerManager& Out, UCanvas* Canvas, TBuffer* DebugBuffer, FRect DrawRect, const float MaxColumnTimeSeconds, const float MaxLocalFrameTime, const FString& Header, FCachedScreenPositionMap& SendCache, FCachedScreenPositionMap& RecvCache)
+	{
+		static float Pad = 2.f;
+		static float BaseLineYPCT = 0.8f;
+
+		auto& InputBuffer = NetworkSim->GetHistoricBuffers() ? NetworkSim->GetHistoricBuffers()->Input : NetworkSim->Buffers.Input;
+
+		if (Canvas && DebugBuffer && DebugBuffer->Num() > 0)
+		{
+			// Outline + Header
+			Out.EmitLine( FVector2D(DrawRect.Min.X, DrawRect.Min.Y), FVector2D(DrawRect.Min.X, DrawRect.Max.Y), FColor::White );
+			Out.EmitLine( FVector2D(DrawRect.Min.X, DrawRect.Min.Y), FVector2D(DrawRect.Max.X, DrawRect.Min.Y), FColor::White );
+			Out.EmitLine( FVector2D(DrawRect.Max.X, DrawRect.Min.Y), FVector2D(DrawRect.Max.X, DrawRect.Max.Y), FColor::White );
+			Out.EmitLine( FVector2D(DrawRect.Min.X, DrawRect.Max.Y), FVector2D(DrawRect.Max.X, DrawRect.Max.Y), FColor::White );
+
+			Out.EmitText( DrawRect.Min, FColor::White, Header );
+
+			// Frame Columns
+			const float BaseLineYPos = DrawRect.Min.Y + (BaseLineYPCT * (DrawRect.Max.Y - DrawRect.Min.Y));
+
+			const float AboveBaseLineSecondsToPixelsY = (BaseLineYPos - DrawRect.Min.Y - Pad) / MaxColumnTimeSeconds;
+			const float BelowBaseLineSecondsToPixelsY = (DrawRect.Max.Y - BaseLineYPos - Pad) / MaxLocalFrameTime;
+
+			const float SecondsToPixelsY = FMath::Min<float>(BelowBaseLineSecondsToPixelsY, AboveBaseLineSecondsToPixelsY);
+
+			Out.EmitLine( FVector2D(DrawRect.Min.X, BaseLineYPos), FVector2D(DrawRect.Max.X, BaseLineYPos), FColor::Black );
+
+
+			FTextSizingParameters TextSizing;
+			TextSizing.DrawFont = GEngine->GetTinyFont();
+			TextSizing.Scaling = FVector2D(1.f,1.f);
+			Canvas->CanvasStringSize(TextSizing, TEXT("00000"));
+
+			const float FixedWidth = TextSizing.DrawXL;
+
+			float ScreenX = DrawRect.Min.X;
+			float ScreenY = BaseLineYPos + Pad;
+
+			for (auto It = DebugBuffer->CreateIterator(); It; ++It)
+			{
+				auto* DebugState = It.Element();
+				const float FrameHeight = SecondsToPixelsY * DebugState->LocalDeltaTimeSeconds;
+				
+				// Green local frame time (below baseline)
+				FColor Color = FColor::Green;
+				Out.EmitQuad(FVector2D( ScreenX, ScreenY), FVector2D( FixedWidth, FrameHeight), Color);
+				Out.EmitText(FVector2D( ScreenX, ScreenY), FColor::Black, FString::Printf(TEXT("%.2f"), (DebugState->LocalDeltaTimeSeconds * 1000.f)));
+
+				// Processed InputcmdsFrames (above baseline)
+				float ClientX = ScreenX;
+				float ClientY = ScreenY - Pad;
+
+				for (int32 Frame : DebugState->ProcessedFrames)
+				{
+					auto* Cmd = InputBuffer[Frame];
+					if (Cmd)
+					{
+						float ClientSizeX = FixedWidth;
+						float ClientSizeY = SecondsToPixelsY * Cmd->GetFrameDeltaTime().ToRealTimeSeconds();
+
+						FVector2D ScreenPos(ClientX, ClientY - ClientSizeY);
+						Out.EmitQuad(ScreenPos, FVector2D( ClientSizeX, ClientSizeY), FColor::Blue);
+						Out.EmitText(ScreenPos, FColor::White, LexToString(Frame));
+						ClientY -= (ClientSizeY + Pad);
+					}
+				}
+
+				// Unprocessed InputCmds (above processed)				
+				for (int32 Frame = DebugState->PendingFrame; Frame <= DebugState->HeadFrame; ++Frame)
+				{
+					if (auto* Cmd = InputBuffer[Frame])
+					{
+						float ClientSizeX = FixedWidth;
+						float ClientSizeY = SecondsToPixelsY * Cmd->GetFrameDeltaTime().ToRealTimeSeconds();
+						FVector2D ScreenPos(ClientX, ClientY - ClientSizeY);
+						Out.EmitQuad(ScreenPos, FVector2D( ClientSizeX, ClientSizeY), FColor::Red);
+						Out.EmitText(ScreenPos, FColor::White, LexToString(Frame));
+						ClientY -= (ClientSizeY + Pad);
+					}
+				}
+
+				// Cache Screen Positions based on frame
+				RecvCache.Frames.FindOrAdd(DebugState->LastReceivedInputFrame).SetRecv(FVector2D(ScreenX, BaseLineYPos));
+				
+				// Advance 
+				ScreenX += FixedWidth + Pad;
+
+				// Send cache
+				SendCache.Frames.FindOrAdd(DebugState->LastSentInputFrame).SetSent(FVector2D(ScreenX, BaseLineYPos));
+			}
+
+
+			// Remaining Simulation Time
+			TOptional<FVector2D> LastLinePos;
+			FVector2D LinePos(DrawRect.Min.X, BaseLineYPos);
+			for (auto It = DebugBuffer->CreateIterator(); It; ++It)
+			{
+				auto* DebugState = It.Element();
+				
+				LinePos.X += FixedWidth + Pad;
+				LinePos.Y = BaseLineYPos - (DebugState->RemainingAllowedSimulationTimeSeconds * SecondsToPixelsY);
+
+				FColor LineColor =  FColor::White;
+				if (LinePos.Y < DrawRect.Min.Y)
+				{
+					LinePos.Y = DrawRect.Min.Y;
+					LineColor = FColor::Red;
+				}
+				if (LinePos.Y > DrawRect.Max.Y)
+				{
+					LinePos.Y = DrawRect.Max.Y;
+					LineColor = FColor::Red;
+				}
+				
+				if (LastLinePos.IsSet())
+				{
+					Out.EmitLine(LastLinePos.GetValue(), LinePos, LineColor, 2.f);
+				}
+
+				LastLinePos = LinePos;
+				
+			}
+		}
+	}
+
+	template<typename TBuffer>
+	void CalcMaxColumnFrameTime(TBuffer* DebugBuffer, float& MaxInputTime, float& MaxLocalFrameTime)
+	{
+		auto& InputBuffer = NetworkSim->GetHistoricBuffers() ? NetworkSim->GetHistoricBuffers()->Input : NetworkSim->Buffers.Input;
+
+		for (auto It = DebugBuffer->CreateIterator(); It; ++It)
+		{
+			float ColumnTime = 0.f;
+
+			auto* DebugState = It.Element();
+			for (int32 Frame : DebugState->ProcessedFrames)
+			{
+				auto* Cmd = InputBuffer[Frame];
+				if (Cmd)
+				{
+					ColumnTime += (float)Cmd->GetFrameDeltaTime().ToRealTimeSeconds();
+				}
+			}
+			for (int32 Frame = DebugState->PendingFrame; Frame <= DebugState->HeadFrame; ++Frame)
+			{
+				if (auto* Cmd = InputBuffer[Frame])
+				{
+					ColumnTime += (float)Cmd->GetFrameDeltaTime().ToRealTimeSeconds();
+				}
+			}
+
+			MaxInputTime = FMath::Max<float>(ColumnTime, MaxInputTime);
+			MaxLocalFrameTime = FMath::Max<float>(DebugState->LocalDeltaTimeSeconds, MaxLocalFrameTime);
+		}
+	}
+
+	void GatherCurrent(FNetworkSimulationModelDebuggerManager& Out, UCanvas* Canvas) override
+	{
+		const AActor* Owner = WeakOwningActor.Get();
+		if (!ensure(Owner))
+		{
+			return;
+		}
+
+		// ------------------------------------------------------------------------------------------------------------------------------------------------
+		//	Lines
+		// ------------------------------------------------------------------------------------------------------------------------------------------------
+
+		Out.Emit(FString::Printf(TEXT("%s - %s"), *Owner->GetName(), *UEnum::GetValueAsString(TEXT("Engine.ENetRole"), Owner->GetLocalRole())), FColor::Yellow);
+		Out.Emit(FString::Printf(TEXT("PendingFrame: %d (%d Buffered)"), NetworkSim->Ticker.PendingFrame, NetworkSim->Buffers.Input.HeadFrame() - NetworkSim->Ticker.PendingFrame));
+				
+		if (Owner->GetLocalRole() == ROLE_AutonomousProxy)
+		{			
+			FColor Color = FColor::White;
+			const bool FaultDetected = NetworkSim->RepProxy_Autonomous.IsReconcileFaultDetected();
+
+			const int32 LastSerializedFrame = NetworkSim->RepProxy_Autonomous.GetLastSerializedFrame();
+
+			// Calc how much predicted time we have processed. Note that we use the motionstate buffer to iterate but the MS is on the input cmd. (if we are buffering cmds, don't want to count them)
+			
+			TSimTime PredictedMS;
+			for (int32 PredFrame = LastSerializedFrame+1; PredFrame <= NetworkSim->Buffers.Sync.HeadFrame(); ++PredFrame)
+			{
+				if (auto* Cmd = NetworkSim->Buffers.Input[PredFrame])
+				{
+					PredictedMS += Cmd->GetFrameDeltaTime();
+				}
+			}
+
+			FString ConfirmedFrameStr = FString::Printf(TEXT("LastConfirmedFrame: %d. Prediction: %d Frames, %s MS"), LastSerializedFrame, NetworkSim->Buffers.Sync.HeadFrame() - LastSerializedFrame, *PredictedMS.ToString());
+			if (FaultDetected)
+			{
+				ConfirmedFrameStr += TEXT(" RECONCILE FAULT DETECTED!");
+				Color = FColor::Red;
+			}
+
+			Out.Emit(*ConfirmedFrameStr, Color);
+
+			FString SimulationTimeString = FString::Printf(TEXT("Local SimulationTime: %s. SerializedSimulationTime: %s. Difference MS: %s"), *NetworkSim->Ticker.GetTotalProcessedSimulationTime().ToString(),
+				*NetworkSim->RepProxy_Autonomous.GetLastSerializedSimTime().ToString(), *(NetworkSim->Ticker.GetTotalProcessedSimulationTime() - NetworkSim->RepProxy_Autonomous.GetLastSerializedSimTime()).ToString());
+			Out.Emit(*SimulationTimeString, Color);
+
+			FString AllowedSimulationTimeString = FString::Printf(TEXT("Allowed Simulation Time: %s. Frame: %d/%d/%d"), *NetworkSim->Ticker.GetRemainingAllowedSimulationTime().ToString(), NetworkSim->Ticker.MaxAllowedFrame, NetworkSim->Ticker.PendingFrame, NetworkSim->Buffers.Input.HeadFrame());
+			Out.Emit(*AllowedSimulationTimeString, Color);
+		}
+		else if (Owner->GetLocalRole() == ROLE_SimulatedProxy)
+		{
+			FColor Color = FColor::White;
+			FString TimeString = FString::Printf(TEXT("Total Processed Simulation Time: %s. Last Serialized Simulation Time: %s. Delta: %s"), *NetworkSim->Ticker.GetTotalProcessedSimulationTime().ToString(), *NetworkSim->RepProxy_Simulated.GetLastSerializedSimulationTime().ToString(), *(NetworkSim->RepProxy_Simulated.GetLastSerializedSimulationTime() - NetworkSim->Ticker.GetTotalProcessedSimulationTime()).ToString());
+			Out.Emit(*TimeString, Color);
+		}
+
+		auto EmitBuffer = [&Out](FString BufferName, auto& Buffer)
+		{
+			Out.Emit();
+			Out.Emit(FString::Printf(TEXT("//////////////// %s ///////////////"), *BufferName), FColor::Yellow);
+			Out.Emit(FString::Printf(TEXT("%s"), *Buffer.GetBasicDebugStr()));		
+			Out.Emit();
+			Out.EmitElement(Buffer, FStandardLoggingParameters(nullptr, EStandardLoggingContext::Full, Buffer.HeadFrame()));
+		};
+
+		EmitBuffer(TEXT("InputBuffer"), NetworkSim->Buffers.Input);
+		EmitBuffer(TEXT("SyncBuffer"), NetworkSim->Buffers.Sync);
+		EmitBuffer(TEXT("AuxBuffer"), NetworkSim->Buffers.Aux);
+
+		// ------------------------------------------------------------------------------------------------------------------------------------------------
+		//	Canvas
+		// ------------------------------------------------------------------------------------------------------------------------------------------------
+
+		if (Canvas)
+		{
+			FRect ServerRect;
+			ServerRect.Min.X = 0.30f * (float)Canvas->SizeX;
+			ServerRect.Max.X = 0.95f * (float)Canvas->SizeX;
+			ServerRect.Min.Y = 0.05f * (float)Canvas->SizeY;
+			ServerRect.Max.Y = 0.45f * (float)Canvas->SizeY;
+
+			FRect ClientRect;
+			ClientRect.Min.X = 0.30f * (float)Canvas->SizeX;
+			ClientRect.Max.X = 0.95f * (float)Canvas->SizeX;
+			ClientRect.Min.Y = 0.55f * (float)Canvas->SizeY;
+			ClientRect.Max.Y = 0.95f * (float)Canvas->SizeY;
+
+			float MaxColumnTime = 1.f/60.f;
+			float MaxLocalFrameTime = 1.f/60.f;
+			CalcMaxColumnFrameTime(NetworkSim->GetRemoteDebugBuffer(), MaxColumnTime, MaxLocalFrameTime);
+			CalcMaxColumnFrameTime(NetworkSim->GetLocalDebugBuffer(), MaxColumnTime, MaxLocalFrameTime);
+
+			FCachedScreenPositionMap ServerToClientCache;
+			FCachedScreenPositionMap ClientToServerCache;
+
+			GatherDebugGraph(Out, Canvas, NetworkSim->GetRemoteDebugBuffer(), ServerRect, MaxColumnTime, MaxLocalFrameTime, TEXT("Server"), ServerToClientCache, ClientToServerCache);
+			GatherDebugGraph(Out, Canvas, NetworkSim->GetLocalDebugBuffer(), ClientRect, MaxColumnTime, MaxLocalFrameTime, TEXT("Client"), ClientToServerCache, ServerToClientCache);
+
+			// Network Send/Recv lines
+			if (NetworkSimulationModelDebugCVars::DrawNetworkSendLines() > 0)
+			{
+				for (auto& It : ServerToClientCache.Frames)
+				{
+					const int32 Frame = It.Key;
+					auto& Positions = It.Value;
+					if (Frame != 0 && Positions.RecvPosition != FVector2D::ZeroVector && Positions.SentPosition != FVector2D::ZeroVector)
+					{
+						Out.EmitLine(Positions.SentPosition, Positions.RecvPosition, FColor::Purple);
+
+						FVector2D TextPos = Positions.SentPosition + (0.25f * (Positions.RecvPosition - Positions.SentPosition));
+						Out.EmitText(TextPos, FColor::Purple, LexToString(Frame));
+					}
+				}
+
+				for (auto& It : ClientToServerCache.Frames)
+				{
+					const int32 Frame = It.Key;
+					auto& Positions = It.Value;
+					if (Frame != 0 && Positions.RecvPosition != FVector2D::ZeroVector && Positions.SentPosition != FVector2D::ZeroVector)
+					{
+						Out.EmitLine(Positions.SentPosition, Positions.RecvPosition, FColor::Orange);
+
+						FVector2D TextPos = Positions.SentPosition + (0.25f * (Positions.RecvPosition - Positions.SentPosition));
+						Out.EmitText(TextPos, FColor::Orange, LexToString(Frame));
+					}
+				}
+			}
+		}
+	}
+
+	void Tick( float DeltaTime ) override
+	{
+		const AActor* Owner = WeakOwningActor.Get();
+		if (!Owner)
+		{
+			return;
+		}
+
+		UWorld* World = Owner->GetWorld();
+
+		if (auto* LatestSync = NetworkSim->Buffers.Sync.HeadElement())
+		{
+			FVisualLoggingParameters VLogParams(EVisualLoggingContext::LastPredicted, NetworkSim->Buffers.Sync.HeadFrame(), EVisualLoggingLifetime::Transient);
+			int32 SyncFrame = NetworkSim->Buffers.Sync.HeadFrame();
+			NetworkSim->Driver->VisualLog(NetworkSim->Buffers.Input[SyncFrame], LatestSync, NetworkSim->Buffers.Aux[SyncFrame], VLogParams);
+		}
+
+		FStuff ServerPIEStuff = GetServerPIEStuff();
+		if (ServerPIEStuff.NetworkSim)
+		{
+			if (auto* ServerLatestSync = ServerPIEStuff.NetworkSim->Buffers.Sync.HeadElement())
+			{
+				FVisualLoggingParameters VLogParams(EVisualLoggingContext::CurrentServerPIE, ServerPIEStuff.NetworkSim->Buffers.Sync.HeadFrame(), EVisualLoggingLifetime::Transient);
+				int32 SyncFrame = ServerPIEStuff.NetworkSim->Buffers.Sync.HeadFrame();
+				NetworkSim->Driver->VisualLog(ServerPIEStuff.NetworkSim->Buffers.Input[SyncFrame], ServerLatestSync, ServerPIEStuff.NetworkSim->Buffers.Aux[SyncFrame], VLogParams);
+			}
+		}
+
+		if (Owner->GetLocalRole() == ROLE_AutonomousProxy)
+		{
+			for (int32 Frame = NetworkSim->RepProxy_Autonomous.GetLastSerializedFrame(); Frame < NetworkSim->Buffers.Sync.HeadFrame(); ++Frame)
+			{
+				if (auto* SyncState = NetworkSim->Buffers.Sync[Frame])
+				{
+					const EVisualLoggingContext Context = (Frame == NetworkSim->RepProxy_Autonomous.GetLastSerializedFrame()) ? EVisualLoggingContext::LastConfirmed : EVisualLoggingContext::OtherPredicted;
+					FVisualLoggingParameters VLogParams(Context, NetworkSim->Buffers.Sync.HeadFrame(), EVisualLoggingLifetime::Transient);
+					NetworkSim->Driver->VisualLog(NetworkSim->Buffers.Input[Frame], SyncState, NetworkSim->Buffers.Aux[Frame], VLogParams);
+				}
+			}
+		}
+		else if (Owner->GetLocalRole() == ROLE_SimulatedProxy)
+		{
+			FVisualLoggingParameters VLogParams(EVisualLoggingContext::LastConfirmed, NetworkSim->Buffers.Sync.HeadFrame(), EVisualLoggingLifetime::Transient);
+			int32 HeadFrame = NetworkSim->Buffers.Sync.HeadFrame();
+			NetworkSim->Driver->VisualLog(NetworkSim->Buffers.Input[HeadFrame], NetworkSim->Buffers.Sync[HeadFrame], NetworkSim->Buffers.Aux[HeadFrame], VLogParams);
+			
+			if (NetworkSim->GetSimulatedUpdateMode() != ESimulatedUpdateMode::Interpolate)
+			{
+				FVector2D ServerSimulationTimeData(Owner->GetWorld()->GetTimeSeconds(), NetworkSim->RepProxy_Simulated.GetLastSerializedSimulationTime().ToRealTimeMS());
+				UE_VLOG_HISTOGRAM(Owner, LogNetworkSimDebug, Log, "Simulated Time Graph", "Serialized Simulation Time", ServerSimulationTimeData);
+
+				FVector2D LocalSimulationTimeData(Owner->GetWorld()->GetTimeSeconds(), NetworkSim->Ticker.GetTotalProcessedSimulationTime().ToRealTimeMS());
+				UE_VLOG_HISTOGRAM(Owner, LogNetworkSimDebug, Log, "Simulated Time Graph", "Local Simulation Time", LocalSimulationTimeData);
+			
+			}
+		}
+	}
+	
+
+	struct FStuff
+	{
+		TNetSimModel* NetworkSim = nullptr;
+	};
+
+	FStuff GetStuff()
+	{
+		return {NetworkSim};
+	}
+
+	TFunction< FStuff() > GetServerPIEStuff;
+
+private:
+	
+	TWeakObjectPtr<const AActor> WeakOwningActor;
+	TNetSimModel* NetworkSim = nullptr;
+};
+
+template <typename T>
+void FNetworkSimulationModelDebuggerManager::RegisterNetworkSimulationModel(T* NetworkSim, const AActor* OwningActor)
+{
+	TNetworkSimulationModelDebugger<T>* Debugger = new TNetworkSimulationModelDebugger<T>(NetworkSim, OwningActor);
+	DebuggerMap.Add( TWeakObjectPtr<const AActor>(OwningActor), Debugger );
+
+	// Gross stuff so that the debugger can find the ServerPIE equiv
+	TWeakObjectPtr<const AActor> WeakOwner(OwningActor);
+	Debugger->GetServerPIEStuff = [WeakOwner, this]()
+	{
+		if (AActor* ServerOwner = Cast<AActor>(FindReplicatedObjectOnPIEServer(WeakOwner.Get())))
+		{
+			return ((TNetworkSimulationModelDebugger<T>*)DebuggerMap.FindRef(TWeakObjectPtr<const AActor>(ServerOwner)))->GetStuff();
+		}
+		return typename TNetworkSimulationModelDebugger<T>::FStuff();
+	};
+}
\ No newline at end of file
diff --git a/Engine/Plugins/Runtime/NetworkPrediction/Source/NetworkPrediction/Public/NetworkedSimulationModelInterpolator.h b/Engine/Plugins/Runtime/NetworkPrediction/Source/NetworkPrediction/Public/NetworkedSimulationModelInterpolator.h
new file mode 100644
--- /dev/null
+++ b/Engine/Plugins/Runtime/NetworkPrediction/Source/NetworkPrediction/Public/NetworkedSimulationModelInterpolator.h
@@ -0,0 +1,248 @@
+// Copyright 1998-2019 Epic Games, Inc. All Rights Reserved.
+
+#pragma once
+
+#include "NetworkedSimulationModelCVars.h"
+#include "VisualLogger/VisualLogger.h"
+#include "GameFramework/Actor.h"
+#include "NetworkPredictionTypes.h"
+#include "Engine/World.h"
+#include "NetworkedSimulationModelTypes.h"
+
+DEFINE_LOG_CATEGORY_STATIC(LogNetInterpolation, Log, All);
+
+namespace NetworkInterpolationDebugCVars
+{
+	NETSIM_DEVCVAR_SHIPCONST_INT(Disable, 0, "ni.Disable", "Disables Network Interpolation");
+	NETSIM_DEVCVAR_SHIPCONST_INT(VLog, 0, "ni.VLog", "Enables Network Interpolation VLog ");
+
+	NETSIM_DEVCVAR_SHIPCONST_FLOAT(WaitSlack, 0.05, "ni.WaitSlack", "How much slack to wait for when waiting");
+
+	NETSIM_DEVCVAR_SHIPCONST_FLOAT(CatchUpThreshold, 0.300, "ni.CatchUpThreshold", "When we start catching up (seconds from head)");
+	NETSIM_DEVCVAR_SHIPCONST_FLOAT(CatchUpGoal, 0.010, "ni.CatchUpGoal", "When we stop cathcing up (seconds from head)");
+	NETSIM_DEVCVAR_SHIPCONST_FLOAT(CatchUpFactor, 1.50, "ni.CatchUpFactor", "Factor we use to catch up");
+}
+
+template<typename TBufferTypes, typename TTickSettings>
+struct TInterpolator
+{
+	using FSimTime = FNetworkSimTime;
+	using FRealTime = FNetworkSimTime::FRealTime;
+	using TSyncState = typename TBufferTypes::TSyncState;
+	using TAuxState = typename TBufferTypes::TAuxState;
+
+	bool bEnableVisualLog = true;
+
+	template<typename TSystemDriver>
+	void PostSimTick(TSystemDriver* Driver, const TNetworkSimBufferContainer<TBufferTypes>& Buffers, const FSimulationTickState& TickInfo, const FNetSimTickParameters& TickParameters)
+	{
+		const bool bDoVLog = NetworkInterpolationDebugCVars::VLog() && bEnableVisualLog;
+		const float DeltaSeconds = TickParameters.LocalDeltaTimeSeconds;
+
+		const AActor* LogOwner = Driver->GetVLogOwner();
+
+		// Interpolation disabled
+		if (NetworkInterpolationDebugCVars::Disable() > 0)
+		{			
+			if (const TSyncState* HeadState = Buffers.Sync.HeadElement())
+			{
+				const TAuxState* AuxState = Buffers.Aux.HeadElement();
+				check(AuxState);
+
+				Driver->FinalizeFrame(*HeadState, *AuxState);
+			}
+			return;
+		}
+		
+		if (TickInfo.SimulationTimeBuffer.Num() <= 1)
+		{
+			// Cant interpolate yet	
+			return;
+		}
+
+		auto& SimulationTimeBuffer = TickInfo.SimulationTimeBuffer;
+
+		// Starting off: start at the tail end
+		if (InterpolationTime <= 0.f)
+		{
+			InterpolationTime = SimulationTimeBuffer.TailElement()->ToRealTimeSeconds();
+			InterpolationFrame = SimulationTimeBuffer.TailFrame();
+			InterpolationState = *Buffers.Sync.TailElement();
+		}
+
+		// Wait if we were too far ahead
+		if (WaitUntilTime > 0.f)
+		{
+			if (WaitUntilTime <= SimulationTimeBuffer.HeadElement()->ToRealTimeSeconds())
+			{
+				// done waiting, we can continue
+				WaitUntilTime = 0.f;
+				UE_VLOG(LogOwner, LogNetInterpolation, Log, TEXT("Done Waiting! Head: %s"), *LexToString(SimulationTimeBuffer.HeadElement()->ToRealTimeSeconds()));
+			}
+			else
+			{
+				if (bDoVLog)
+				{
+					// Still waiting, return
+					UE_VLOG(LogOwner, LogNetInterpolation, Log, TEXT("Still Waiting! %s < %s"), *LexToString(WaitUntilTime), *LexToString(SimulationTimeBuffer.HeadElement()->ToRealTimeSeconds()));
+
+					const FName LocalInterpolationTimeName("Local Interpolation Time");
+					FVector2D LocalTimeVsInterpolationTime(LogOwner->GetWorld()->GetTimeSeconds(), (int32)(InterpolationTime * 1000.f));
+					UE_VLOG_HISTOGRAM(LogOwner, LogNetInterpolation, Log, "ServerSimulationTimeGraph", LocalInterpolationTimeName, LocalTimeVsInterpolationTime);
+				}
+				return;
+			}
+		}
+
+		EVisualLoggingContext LoggingContext = EVisualLoggingContext::InterpolationLatest;
+
+		// Calc new interpolation time
+		FRealTime NewInterpolationTime = InterpolationTime;
+		{
+			FRealTime Step = DeltaSeconds;
+			
+			// Speed up if we are too far behind
+			FRealTime CatchUpThreshold = SimulationTimeBuffer.HeadElement()->ToRealTimeSeconds() - NetworkInterpolationDebugCVars::CatchUpThreshold();
+			if (CatchUpUntilTime <= 0.f && InterpolationTime < CatchUpThreshold)
+			{
+				CatchUpUntilTime = SimulationTimeBuffer.HeadElement()->ToRealTimeSeconds() - NetworkInterpolationDebugCVars::CatchUpGoal();
+			}
+
+			if (CatchUpUntilTime > 0.f)
+			{
+				if (InterpolationTime  < CatchUpUntilTime)
+				{
+					Step *= NetworkInterpolationDebugCVars::CatchUpFactor();
+					LoggingContext = EVisualLoggingContext::InterpolationSpeedUp;
+
+					UE_VLOG(LogOwner, LogNetInterpolation, Log, TEXT("Catching up! %s < %s"), *LexToString(InterpolationTime), *LexToString(CatchUpUntilTime));
+					//!UE_LOG(LogNetInterpolation, Warning, TEXT("Catching up! %s < %s"), *LexToString(InterpolationTime), *LexToString(CatchUpUntilTime));
+				}
+				else
+				{
+					CatchUpUntilTime = 0.f;
+				}
+			}
+
+			NewInterpolationTime += Step;
+
+			// Did this put us too far ahead, and now we need to start waiting?
+			if (NewInterpolationTime > SimulationTimeBuffer.HeadElement()->ToRealTimeSeconds())
+			{
+				UE_VLOG(LogOwner, LogNetInterpolation, Log, TEXT("Too far ahead! Starting to wait! Head: %s"), *LexToString(SimulationTimeBuffer.HeadElement()->ToRealTimeSeconds()));
+				WaitUntilTime = SimulationTimeBuffer.HeadElement()->ToRealTimeSeconds() + NetworkInterpolationDebugCVars::WaitSlack();
+				NewInterpolationTime = SimulationTimeBuffer.HeadElement()->ToRealTimeSeconds();
+			}
+		}
+
+		// Find "To" frame
+		const TSyncState* ToState = nullptr;
+		const TAuxState* AuxState = nullptr;
+		FRealTime ToTime = 0.f;
+
+		for (auto It = SimulationTimeBuffer.CreateConstIterator(); It; ++It)
+		{
+			FSimTime ElementSimTime = *It.Element();
+			if (NewInterpolationTime <= ElementSimTime.ToRealTimeSeconds())
+			{
+				InterpolationFrame = It.Frame();
+				ToTime = ElementSimTime.ToRealTimeSeconds();
+				ToState = Buffers.Sync[It.Frame()];
+				AuxState = Buffers.Aux[It.Frame()];
+				break;
+			}
+		}
+
+		if (ensure(ToState && AuxState))
+		{
+			const FRealTime FromRealTime = InterpolationTime;
+			const FRealTime ToRealTime = ToTime;
+			const FRealTime InterpolationInterval = ToRealTime - FromRealTime;
+		
+			if (ensure(FMath::Abs(InterpolationInterval) > 0.f))
+			{
+				const float InterpolationPCT = (NewInterpolationTime - FromRealTime) / InterpolationInterval;
+				ensureMsgf(InterpolationPCT >= 0.f && InterpolationPCT <= 1.f, TEXT("Calculated InterpolationPCT not in expected range. NewInterpolationTime: %s. From: %s. To: %s"), *LexToString(NewInterpolationTime), *LexToString(FromRealTime), *LexToString(ToRealTime));
+
+				TSyncState NewInterpolatedState;
+				TSyncState::Interpolate(InterpolationState, *ToState, InterpolationPCT, NewInterpolatedState);
+
+				Driver->FinalizeFrame(NewInterpolatedState, *AuxState);
+				
+				if (bDoVLog)
+				{
+					UE_VLOG(LogOwner, LogNetInterpolation, Log, TEXT("%s - %s - %s.  InterpolationPCT: %f"), *LexToString(FromRealTime), *LexToString(NewInterpolationTime), *LexToString(ToRealTime), InterpolationPCT);
+
+					// Graph Interpolation Time vs Buffer Head/Tail times
+					const FName ServerSimulationGraphName("ServerSimulationTimeGraph");
+					const FName ServerSimTimeName("Server Simulation Time");
+					FVector2D LocalTimeVsServerSimTime(LogOwner->GetWorld()->GetTimeSeconds(), (int32)TickInfo.SimulationTimeBuffer.HeadElement()->ToRealTimeMS());
+					UE_VLOG_HISTOGRAM(LogOwner, LogNetInterpolation, Log, ServerSimulationGraphName, ServerSimTimeName, LocalTimeVsServerSimTime);
+
+					const FName BufferTailSimTimeName("Buffer Tail Simulation Time");
+					FVector2D LocalTimeVsBufferTailSim(LogOwner->GetWorld()->GetTimeSeconds(), (int32)TickInfo.SimulationTimeBuffer.TailElement()->ToRealTimeMS());
+					UE_VLOG_HISTOGRAM(LogOwner, LogNetInterpolation, Log, ServerSimulationGraphName, BufferTailSimTimeName, LocalTimeVsBufferTailSim);
+
+					const FName LocalInterpolationTimeName("Local Interpolation Time");
+					FVector2D LocalTimeVsInterpolationTime(LogOwner->GetWorld()->GetTimeSeconds(), (int32)(NewInterpolationTime * 1000.f));
+					UE_VLOG_HISTOGRAM(LogOwner, LogNetInterpolation, Log, ServerSimulationGraphName, LocalInterpolationTimeName, LocalTimeVsInterpolationTime);
+					
+					FVector2D LocalTimeVsCatchUpThreshold(LogOwner->GetWorld()->GetTimeSeconds(), (SimulationTimeBuffer.HeadElement()->ToRealTimeSeconds() - NetworkInterpolationDebugCVars::CatchUpThreshold()) * 1000.f);
+					UE_VLOG_HISTOGRAM(LogOwner, LogNetInterpolation, Log, ServerSimulationGraphName, "Catch Up Threshold", LocalTimeVsCatchUpThreshold);
+
+					FVector2D LocalTimeVsCatchUpGoal(LogOwner->GetWorld()->GetTimeSeconds(), (SimulationTimeBuffer.HeadElement()->ToRealTimeSeconds() - NetworkInterpolationDebugCVars::CatchUpGoal()) * 1000.f);
+					UE_VLOG_HISTOGRAM(LogOwner, LogNetInterpolation, Log, ServerSimulationGraphName, "Catch Up Goal", LocalTimeVsCatchUpGoal);
+
+					// VLog the actual motion states
+					const TSyncState* DebugTail = Buffers.Sync.TailElement();
+					const TSyncState* DebugHead = Buffers.Sync.HeadElement();
+
+					auto VLogHelper = [&](int32 Frame, EVisualLoggingContext Context)
+					{
+						FVisualLoggingParameters VLogParams(Context, Frame, EVisualLoggingLifetime::Transient);
+						Driver->VisualLog(Buffers.Input[Frame], Buffers.Sync[Frame], Buffers.Aux[Frame], VLogParams);
+					};
+
+					VLogHelper(Buffers.Sync.TailFrame(), EVisualLoggingContext::InterpolationBufferTail);
+					VLogHelper(Buffers.Sync.HeadFrame(), EVisualLoggingContext::InterpolationBufferHead);
+
+					{
+						FVisualLoggingParameters VLogParams(EVisualLoggingContext::InterpolationFrom, InterpolationFrame-1, EVisualLoggingLifetime::Transient);
+						Driver->VisualLog(Buffers.Input[InterpolationFrame-1], &InterpolationState, Buffers.Aux[InterpolationFrame-1], VLogParams);
+					}
+
+					{
+						FVisualLoggingParameters VLogParams(EVisualLoggingContext::InterpolationTo, InterpolationFrame, EVisualLoggingLifetime::Transient);
+						Driver->VisualLog(Buffers.Input[InterpolationFrame], &InterpolationState, Buffers.Aux[InterpolationFrame], VLogParams);
+					}
+
+					{
+						FVisualLoggingParameters VLogParams(LoggingContext, InterpolationFrame, EVisualLoggingLifetime::Transient);
+						Driver->VisualLog(Buffers.Input[InterpolationFrame], &NewInterpolatedState, Buffers.Aux[InterpolationFrame], VLogParams);
+					}
+				}
+
+				InterpolationState = NewInterpolatedState;
+				InterpolationTime = NewInterpolationTime;
+			}
+		}
+	}
+
+private:
+
+	FRealTime InterpolationTime = 0.f; // SimTime we are currently interpolating at
+	int32 InterpolationFrame = INDEX_NONE; // Frame we are currently/last interpolated at
+
+	TSyncState InterpolationState;
+
+	FRealTime WaitUntilTime = 0.f;	
+
+	FRealTime CatchUpUntilTime = 0.f;
+
+
+	FRealTime DynamicBufferedTime = 1/60.f; // SimTime we are currently interpolating at
+	FRealTime DynamicBufferedTimeStep = 1/60.f;
+
+	FRealTime	MinBufferedTime = 1/120.f;
+	FRealTime	MaxBufferedTime = 1.f;
+};
\ No newline at end of file
diff --git a/Engine/Plugins/Runtime/NetworkPrediction/Source/NetworkPrediction/Public/NetworkedSimulationModelReplicators.h b/Engine/Plugins/Runtime/NetworkPrediction/Source/NetworkPrediction/Public/NetworkedSimulationModelReplicators.h
new file mode 100644
--- /dev/null
+++ b/Engine/Plugins/Runtime/NetworkPrediction/Source/NetworkPrediction/Public/NetworkedSimulationModelReplicators.h
@@ -0,0 +1,868 @@
+// Copyright 1998-2019 Epic Games, Inc. All Rights Reserved.
+
+#pragma once
+
+#include "Net/UnrealNetwork.h" // For MakeRelative
+#include "NetworkedSimulationModelCVars.h"
+#include "NetworkedSimulationModelInterpolator.h"
+
+// ----------------------------------------------------------------------------------------------------------------------------------------------------------------------
+//	CVars and compile time constants
+// ----------------------------------------------------------------------------------------------------------------------------------------------------------------------
+
+namespace NetworkSimulationModelCVars
+{
+	NETSIM_DEVCVAR_SHIPCONST_INT(EnableLocalPrediction, 1, "ns.EnableLocalPrediction", "Toggle local prediction.");
+	NETSIM_DEVCVAR_SHIPCONST_INT(EnableSimulatedReconcile, 1, "ns.EnableSimulatedReconcile", "Toggle simulated proxy reconciliation.");
+	NETSIM_DEVCVAR_SHIPCONST_INT(EnableSimulatedExtrapolation, 1, "ns.EnableSimulatedExtrapolation", "Toggle simulated proxy extrapolation.");
+	NETSIM_DEVCVAR_SHIPCONST_INT(ForceReconcile, 0, "ns.ForceReconcile", "Forces reconcile even if state does not differ. E.g, force resimulation after every netupdate.");
+	NETSIM_DEVCVAR_SHIPCONST_INT(ForceReconcileSingle, 0, "ns.ForceReconcileSingle", "Forces a since reconcile to happen on the next frame");
+}
+
+#ifndef NETSIM_NETCONSTANT_NUM_BITS_FRAME
+	#define NETSIM_NETCONSTANT_NUM_BITS_FRAME 8	// Allows you to override this setting via UBT, but access via cleaner FActorMotionNetworkingConstants::NUM_BITS_FRAME
+#endif
+
+struct FNetworkSimulationSerialization
+{
+	// How many bits we use to encode the key frame number for buffers.
+	// Client Frames are stored locally as 32 bit integers, but we use a smaller # of bits to NetSerialize.
+	// Frames are only relatively relevant: the absolute value doesn't really matter. We just need to detect newer/older.
+	enum { NUM_BITS_FRAME = NETSIM_NETCONSTANT_NUM_BITS_FRAME };		
+
+	// Abs max value we encode into the bit writer
+	enum { MAX_FRAME_WRITE = 1 << NUM_BITS_FRAME };
+
+	// This is the threshold at which we would wrap around and incorrectly assign a frame on the receiving side.
+	// E.g, If there are FRAME_ERROR_THRESHOLD frames that do not make it across from sender->receiver, the
+	// receiver will have incorrect local values. With 8 bits, this works out to be 128 frames or about 2 seconds at 60fps.
+	enum { FRAME_ERROR_THRESHOLD = MAX_FRAME_WRITE / 2};
+
+	// Helper to serialize the int32 HeadFrame. Returns the unpacked value (this will be same as input in the save path)
+	static int32 SerializeFrame(FArchive& Ar, int32 LocalHeadFrame)
+	{
+		if (Ar.IsSaving())
+		{
+			((FNetBitWriter&)Ar).WriteIntWrapped( LocalHeadFrame, MAX_FRAME_WRITE );
+			return LocalHeadFrame;
+		}
+		
+		return MakeRelative(((FNetBitReader&)Ar).ReadInt( MAX_FRAME_WRITE ), LocalHeadFrame, MAX_FRAME_WRITE );
+	}
+};
+
+// ----------------------------------------------------------------------------------------------------------------------------------------------------------------------
+//	Replicators
+//	The Replicators are the pieces of the TNetworkedSimulationModel that make up the role-specific functionality (Server, Autonomous Client, and Simulated Client).
+//	Mainly they NetSerialize, Reconcile, and have Pre/Post sim tick functions. The TNetworkedSimulationModel is still the core piece that ticks the sim, but the Replicators
+//	basically do everything else in a role-specific way.
+//
+//
+//
+// ----------------------------------------------------------------------------------------------------------------------------------------------------------------------
+
+template<typename TBufferTypes, typename TTickSettings>
+struct TReplicatorEmpty
+{ 
+	// Used for client shadowstate compares. Should just call GetDirtyCount() on the buffer you are replicating
+	int32 GetProxyDirtyCount(TNetworkSimBufferContainer<TBufferTypes>& Buffers) const { return 0; }
+
+	// NetSerialize: just serialize the network data. Don't run simulation steps. Every replicator will be NetSerialized before moving on to Reconcile phase.
+	void NetSerialize(const FNetSerializeParams& P, TNetworkSimBufferContainer<TBufferTypes>& Buffers, const TSimulationTicker<TTickSettings>& Ticker) { }
+
+	// Reconcile: called after everyone has NetSerialized. "Get right with the server": this function is about reconciling what the server told you vs what you have extrapolated or forward predicted locally
+	template<typename TSimulation, typename TDriver>
+	void Reconcile(TSimulation* Simulation, TDriver* Driver, TNetworkSimBufferContainer<TBufferTypes>& Buffers, TSimulationTicker<TTickSettings>& Ticker) { }
+
+	// Called prior to input processing. This function must update Ticker to allow simulation time (from TickParameters) and to possibly get new input.
+	template<typename TDriver>
+	void PreSimTick(TDriver* Driver, TNetworkSimBufferContainer<TBufferTypes>& Buffers, TSimulationTicker<TTickSettings>& Ticker, const FNetSimTickParameters& TickParameters) { }
+
+	// Called after input processing. Should finalize the frame and do any smoothing/interpolation. This function is not allowed to modify the buffers or tick state, or even call the simulation/Update function.
+	template<typename TDriver>
+	void PostSimTick(TDriver* Driver, const TNetworkSimBufferContainer<TBufferTypes>& Buffers, const TSimulationTicker<TTickSettings>& Ticker, const FNetSimTickParameters& TickParameters) { }
+};
+
+// This is the "templated base class" for replicators but is not required (i.e., this not an official interface used by TNetworkedSimulation model. Just a base implementation you can start with)
+template<typename TBufferTypes, typename TTickSettings>
+struct TReplicatorBase
+{
+	using TInputCmd = typename TBufferTypes::TInputCmd;
+
+	// Used for client shadowstate compares. Should just call GetDirtyCount() on the buffer you are replicating
+	int32 GetProxyDirtyCount(TNetworkSimBufferContainer<TBufferTypes>& Buffers) const { return 0; }
+
+	// NetSerialize: just serialize the network data. Don't run simulation steps. Every replicator will be NetSerialized before moving on to Reconcile phase.
+	void NetSerialize(const FNetSerializeParams& P, TNetworkSimBufferContainer<TBufferTypes>& Buffers, const TSimulationTicker<TTickSettings>& Ticker) { }
+
+	// Reconcile: called after everyone has NetSerialized. "Get right with the server": this function is about reconciling what the server told you vs what you have extrapolated or forward predicted locally
+	template<typename TSimulation, typename TDriver>
+	void Reconcile(TSimulation* Simulation, TDriver* Driver, TNetworkSimBufferContainer<TBufferTypes>& Buffers, TSimulationTicker<TTickSettings>& Ticker) { }
+
+	// Called prior to input processing. This function must updated Ticker to allow simulation time (from TickParameters) and to possibly get new input.
+	template<typename TDriver>
+	void PreSimTick(TDriver* Driver, TNetworkSimBufferContainer<TBufferTypes>& Buffers, TSimulationTicker<TTickSettings>& Ticker, const FNetSimTickParameters& TickParameters)
+	{
+		// Accumulate local delta time into Ticker
+		Ticker.GiveSimulationTime(TickParameters.LocalDeltaTimeSeconds);
+	}
+
+	// Called after input processing. Should finalize the frame and do any smoothing/interpolation. This function is not allowed to modify the buffers or tick state, or even call the simulation/Update function.
+	template<typename TDriver>
+	void PostSimTick(TDriver* Driver, const TNetworkSimBufferContainer<TBufferTypes>& Buffers, const TSimulationTicker<TTickSettings>& Ticker, const FNetSimTickParameters& TickParameters)
+	{
+		// Sync to latest frame if there is any
+		if (Buffers.Sync.Num() > 0)
+		{
+			Driver->FinalizeFrame(*Buffers.Sync.HeadElement(), *Buffers.Aux.HeadElement());
+		}
+	}
+};
+
+// -------------------------------------------------------------------------------------------------------
+//	"Reusable" pieces: mainly for replicating data on specified buffers
+// -------------------------------------------------------------------------------------------------------
+
+// Just serializes simulation time. Can be disabled by Enabled templated parameter. Used as base class in other classes below.
+template<typename TBufferTypes, typename TTickSettings, typename TBase=TReplicatorBase<TBufferTypes, TTickSettings>, bool Enabled=true>
+struct TReplicator_SimTime : public TBase
+{
+	void NetSerialize(const FNetSerializeParams& P, TNetworkSimBufferContainer<TBufferTypes>& Buffers, TSimulationTicker<TTickSettings>& Ticker)
+	{
+		TBase::NetSerialize(P, Buffers, Ticker);
+		SerializedTime = Ticker.GetTotalProcessedSimulationTime();
+		SerializedTime.NetSerialize(P.Ar);
+	}
+
+	FNetworkSimTime SerializedTime;
+};
+
+// Enabled=false specialization: do nothing
+template<typename TBufferTypes, typename TTickSettings, typename TBase>
+struct TReplicator_SimTime<TBufferTypes, TTickSettings, TBase, false> : public TBase { };
+
+// -------------------------------------------------------------------------------------------------------
+//	
+// -------------------------------------------------------------------------------------------------------
+
+// Replicates a sequence of elements. i.e., "the last MaxNumElements".
+// On the receiving side, we merge the sequence into whatever we have locally
+// Frames are synchronized. SimTime is also serialized by default (change by changing TBase)
+template<typename TBufferTypes, typename TTickSettings, ENetworkSimBufferTypeId BufferId, int32 MaxNumElements=3, typename TBase=TReplicator_SimTime<TBufferTypes, TTickSettings>>
+struct TReplicator_Sequence : public TBase
+{
+	int32 GetProxyDirtyCount(TNetworkSimBufferContainer<TBufferTypes>& Buffers) const 
+	{
+		return Buffers.template Get<BufferId>().GetDirtyCount() ^ (TBase::GetProxyDirtyCount(Buffers) << 2); 
+	}
+
+	void NetSerialize(const FNetSerializeParams& P, TNetworkSimBufferContainer<TBufferTypes>& Buffers, TSimulationTicker<TTickSettings>& Ticker)
+	{
+		TBase::NetSerialize(P, Buffers, Ticker);
+
+		auto& Buffer = Buffers.template Get<BufferId>();
+		FArchive& Ar = P.Ar;
+		uint8 SerializedNumElements = FMath::Min<uint8>(MaxNumElements, Buffer.Num());
+		Ar << SerializedNumElements;
+
+		const int32 HeadFrame = FNetworkSimulationSerialization::SerializeFrame(Ar, Buffer.HeadFrame());
+		const int32 StartingFrame = FMath::Max(0, HeadFrame - SerializedNumElements + 1);
+
+		if (Ar.IsLoading())
+		{
+			const int32 PrevHead = Buffer.HeadFrame();
+			if (PrevHead < StartingFrame && PrevHead >= 0)
+			{
+				// There is a gap in the stream. In some cases, we want this to be a "fault" and bubble up. We may want to synthesize state or maybe we just skip ahead.
+				UE_LOG(LogNetworkSim, Warning, TEXT("Fault: gap in received %s buffer. PrevHead: %d. Received: %d-%d. Reseting previous buffer contents"), *LexToString(BufferId), PrevHead, StartingFrame, HeadFrame);
+			}
+		}
+
+		for (int32 Frame = StartingFrame; Frame <= HeadFrame; ++Frame)
+		{
+			// This, as is, is bad. The intention is that these functions serialize multiple items in some delta compressed fashion.
+			// As is, we are just serializing the elements individually.
+			auto* Cmd = Ar.IsLoading() ? Buffer.WriteFrame(Frame) : Buffer[Frame];
+			Cmd->NetSerialize(P);
+		}
+
+		LastSerializedFrame = HeadFrame;
+	}
+
+	int32 GetLastSerializedFrame() const { return LastSerializedFrame; }
+
+protected:
+
+	int32 LastSerializedFrame = 0;
+};
+
+// Replicates only the latest single element in the selected buffer.
+// Frame is not synchronized: the new element is just added to head.
+// SimTime is serialized by default (change by changing TBase)
+template<typename TBufferTypes, typename TTickSettings, ENetworkSimBufferTypeId BufferId, typename TBase=TReplicator_SimTime<TBufferTypes, TTickSettings>>
+struct TReplicator_Single : public TBase
+{
+	using TState = typename TBufferTypes::template select_type<BufferId>::type;
+
+	int32 GetProxyDirtyCount(TNetworkSimBufferContainer<TBufferTypes>& Buffers) const
+	{ 
+		return Buffers.template Get<BufferId>().GetDirtyCount() ^ (TBase::GetProxyDirtyCount(Buffers) << 2); ; 
+	}
+
+	void NetSerialize(const FNetSerializeParams& P, TNetworkSimBufferContainer<TBufferTypes>& Buffers, TSimulationTicker<TTickSettings>& Ticker)
+	{
+		TBase::NetSerialize(P, Buffers, Ticker);
+
+		FArchive& Ar = P.Ar;
+
+		auto& Buffer = Buffers.template Get<BufferId>();
+		TState* State = nullptr;
+		
+		if (Ar.IsSaving())
+		{
+			State = Buffer.GetElementFromHead(0);
+			check(State); // We should not be here if the buffer is empty. Want to avoid serializing an "empty" flag at the top here.
+		}
+		else
+		{
+			State = Buffer.GetWriteNext();
+		}
+
+		State->NetSerialize(Ar);
+	}
+};
+
+/** Helper that writes a new input cmd to the input buffer, at given Frame (usually the sim's PendingFrame). If frame doesn't exist, ProduceInput is called on the driver. */
+template<typename TDriver, typename TBufferTypes>
+void GenerateLocalInputCmdAtFrame(TDriver* Driver, TNetworkSimBufferContainer<TBufferTypes>& Buffers, const FNetworkSimTime& DeltaSimTime, int32 Frame)
+{
+	using TInputCmd = typename TBufferTypes::TInputCmd;
+	if (TInputCmd* InputCmd = Buffers.Input[Frame])
+	{
+		InputCmd->SetFrameDeltaTime(DeltaSimTime);
+	}
+	else
+	{
+		InputCmd = Buffers.Input.WriteFrame(Frame);
+		*InputCmd = TInputCmd();
+		InputCmd->SetFrameDeltaTime(DeltaSimTime);
+		Driver->ProduceInput(DeltaSimTime, *InputCmd);
+	}
+}
+
+/** Helper to generate a local input cmd if we have simulation time to spend and advance the simulation's MaxAllowedFrame so that it can be processed. */
+template<typename TDriver, typename TBufferTypes, typename TTickSettings>
+void TryGenerateLocalInput(TDriver* Driver, TNetworkSimBufferContainer<TBufferTypes>& Buffers, TSimulationTicker<TTickSettings>& Ticker)
+{
+	using TInputCmd = typename TBufferTypes::TInputCmd;
+
+	FNetworkSimTime DeltaSimTime = Ticker.GetRemainingAllowedSimulationTime();
+	if (DeltaSimTime.IsPositive())
+	{
+		GenerateLocalInputCmdAtFrame(Driver, Buffers, DeltaSimTime, Ticker.PendingFrame);
+		Ticker.MaxAllowedFrame = Ticker.PendingFrame;
+	}
+}
+
+// -------------------------------------------------------------------------------------------------------
+//	Role based Replicators: these replicators are meant to serve specific roles wrt how the simulation evolves
+// -------------------------------------------------------------------------------------------------------
+
+// Default Replicator for the Server: Replicates the InputBuffer client->server
+template<typename TBufferTypes, typename TTickSettings, typename TBase=TReplicator_Sequence<TBufferTypes, TTickSettings, ENetworkSimBufferTypeId::Input, 3>>
+struct TReplicator_Server : public TBase
+{
+	using TInputCmd = typename TBufferTypes::TInputCmd;
+
+	template<typename TSimulation, typename TDriver>
+	void Reconcile(TSimulation* Simulation, TDriver* Driver, TNetworkSimBufferContainer<TBufferTypes>& Buffers, TSimulationTicker<TTickSettings>& Ticker)
+	{
+		// After receiving input, server may process up to the latest received frames.
+		// (If we needed to buffer input server side for whatever reason, we would do it here)
+		// (also note that we will implicitly guard against speed hacks in the core update loop by not processing cmds past what we have been "allowed")
+		Ticker.MaxAllowedFrame = Buffers.Input.HeadFrame();
+
+		// Check for gaps in commands
+		if ( Ticker.PendingFrame < Buffers.Input.TailFrame() )
+		{
+			UE_LOG(LogNetworkSim, Warning, TEXT("TReplicator_Server::Reconcile missing inputcmds. PendingFrame: %d. %s. This can happen via packet loss"), Ticker.PendingFrame, *Buffers.Input.GetBasicDebugStr());
+			Ticker.PendingFrame = Buffers.Input.TailFrame();
+		}
+	}
+
+	template<typename T, typename TDriver>
+	void PreSimTick(TDriver* Driver, TNetworkSimBufferContainer<TBufferTypes>& Buffers, TSimulationTicker<TTickSettings>& Ticker, const FNetSimTickParameters& TickParameters)
+	{
+		Ticker.GiveSimulationTime(TickParameters.LocalDeltaTimeSeconds);
+		if (TickParameters.bGenerateLocalInputCmds)
+		{
+			TryGenerateLocalInput(Driver, Buffers, Ticker);
+		}
+	}
+};
+
+/** Simulated: "non locally controlled" simulations. We support "Simulation Extrapolation" here (using the sim to fake inputs to advance the sim)  */
+template<typename TBufferTypes, typename TTickSettings, typename TBase=TReplicatorEmpty<TBufferTypes, TTickSettings>>
+struct TReplicator_Simulated : public TBase
+{
+	using TInputCmd = typename TBufferTypes::TInputCmd;
+	using TSyncState = typename TBufferTypes::TSyncState;
+	using TAuxState = typename TBufferTypes::TAuxState;
+
+	// Parent Simulation. If this is set, this simulation will forward predict in sync with this parent sim. The parent sim should be an autonomous proxy driven simulation
+	INetworkSimulationModel* ParentSimulation = nullptr;
+
+	// Instance flag for enabling simulated extrapolation
+	bool bAllowSimulatedExtrapolation = true;
+
+	// Interpolated that will be used if bAllowSimulatedExtrapolation == false && ParentSimulation == nullptr
+	TInterpolator<TBufferTypes, TTickSettings> Interpolator;
+
+	// Last Serialized time and state
+	FNetworkSimTime GetLastSerializedSimulationTime() const { return LastSerializedSimulationTime; }
+	const TSyncState& GetLastSerializedSyncState() const { return LastSerializedSyncState; }
+	
+	ESimulatedUpdateMode GetSimulatedUpdateMode() const
+	{
+		if (ParentSimulation)
+		{
+			return ESimulatedUpdateMode::ForwardPredict;
+		}
+		if (bAllowSimulatedExtrapolation && NetworkSimulationModelCVars::EnableSimulatedExtrapolation())
+		{
+			return ESimulatedUpdateMode::Extrapolate;
+		}
+
+		return ESimulatedUpdateMode::Interpolate;
+	}
+	
+	int32 GetProxyDirtyCount(TNetworkSimBufferContainer<TBufferTypes>& Buffers) const
+	{
+		return Buffers.Sync.GetDirtyCount() ^ (TBase::GetProxyDirtyCount(Buffers) << 2); 
+	}
+
+	void NetSerialize(const FNetSerializeParams& P, TNetworkSimBufferContainer<TBufferTypes>& Buffers, TSimulationTicker<TTickSettings>& Ticker)
+	{
+		FArchive& Ar = P.Ar;
+
+		FNetworkSimTime PrevLastSerializedSimulationTime = LastSerializedSimulationTime;
+
+		// Serialize latest simulation time
+		LastSerializedSimulationTime = Ticker.GetTotalProcessedSimulationTime();
+		LastSerializedSimulationTime.NetSerialize(P.Ar);
+
+		// Serialize latest element
+		TSyncState* SyncState = nullptr;
+		TAuxState* AuxState = nullptr;
+		
+		if (Ar.IsSaving())
+		{
+			SyncState = Buffers.Sync.HeadElement();
+			AuxState = Buffers.Aux.HeadElement();
+			check(SyncState && AuxState	); // We should not be here if the buffer is empty. Want to avoid serializing an "empty" flag at the top here.
+		}
+		else
+		{
+			check(Ticker.SimulationTimeBuffer.HeadFrame() == Buffers.Sync.HeadFrame());
+			check(Ticker.GetTotalProcessedSimulationTime() <= Ticker.GetTotalAllowedSimulationTime());
+
+			// Cache off our "starting" time before possibly overwriting it. We will use this in Reconcile to catch back up in some cases.
+			if (Ticker.GetTotalProcessedSimulationTime() > ReconcileSimulationTime)
+			{
+				ReconcileSimulationTime = Ticker.GetTotalProcessedSimulationTime();
+			}
+
+			// Find out where this should go in the local buffer based on the serialized time
+			int32 DestinationFrame = INDEX_NONE;
+			if (LastSerializedSimulationTime > Ticker.GetTotalProcessedSimulationTime())
+			{
+				// We are getting new state that is ahead of what we have locally, so it can safety go right to head
+				DestinationFrame = Ticker.SimulationTimeBuffer.HeadFrame()+1;
+			}
+			else
+			{
+				// We are getting state that is behind what we have locally
+				for (int32 Frame = Ticker.SimulationTimeBuffer.HeadFrame(); Frame >= Ticker.SimulationTimeBuffer.TailFrame(); --Frame)
+				{
+					if (LastSerializedSimulationTime > *Ticker.SimulationTimeBuffer[Frame])
+					{
+						DestinationFrame = Frame+1;
+						break;
+					}
+				}
+
+				if (DestinationFrame == INDEX_NONE)
+				{
+					FNetworkSimTime TotalTimeAhead = Ticker.GetTotalProcessedSimulationTime() - LastSerializedSimulationTime;
+					FNetworkSimTime SerializeDelta = LastSerializedSimulationTime - PrevLastSerializedSimulationTime;
+
+					// We are way far ahead of the server... we will need to clear out sync buffers, take what they gave us, and catch up in reconcile
+					//UE_LOG(LogNetworkSim, Warning, TEXT("!!! TReplicator_Simulated. Large gap detected. SerializedTime: %s. Buffer time: [%s-%s]. %d Elements. DeltaFromHead: %s. DeltaSerialize: %s"), *LastSerializedSimulationTime.ToString(), 
+					//	*Ticker.SimulationTimeBuffer.GetElementFromTail(0)->ToString(), *Ticker.SimulationTimeBuffer.GetElementFromHead(0)->ToString(), Ticker.SimulationTimeBuffer.GetNumValidElements(), *TotalTimeAhead.ToString(), *SerializeDelta.ToString());
+					DestinationFrame = Ticker.SimulationTimeBuffer.HeadFrame()+2; // (Skip ahead 2 to force a break in continuity)
+				}
+			}
+
+			check(DestinationFrame != INDEX_NONE);
+
+			// "Finalize" our buffers and time keeping such that we serialize the latest state from the server in the right spot
+			SyncState = Buffers.Sync.WriteFrame(DestinationFrame);
+			AuxState = Buffers.Aux.WriteFrame(DestinationFrame);
+			Buffers.Input.WriteFrame(DestinationFrame);
+
+			// Update tick info
+			Ticker.SetTotalProcessedSimulationTime(LastSerializedSimulationTime, DestinationFrame);
+			if (Ticker.GetTotalAllowedSimulationTime() < LastSerializedSimulationTime)
+			{
+				Ticker.SetTotalAllowedSimulationTime(LastSerializedSimulationTime);
+			}
+
+			check(Ticker.GetTotalProcessedSimulationTime() <= Ticker.GetTotalAllowedSimulationTime());
+
+			Ticker.PendingFrame = DestinationFrame;	// We are about to serialize state to DestinationFrame which will be "unprocessed" (has not been used to generate a new frame)
+			Ticker.MaxAllowedFrame = DestinationFrame-1; // Do not process PendingFrame on our account. ::PreSimTick will advance this based on our interpolation/extrapolation settings
+		}
+
+		check(SyncState && AuxState);
+		SyncState->NetSerialize(Ar);
+		AuxState->NetSerialize(Ar);
+
+		if (Ar.IsLoading())
+		{
+			LastSerializedSyncState = *SyncState;
+			LastSerializedAuxState = *AuxState;
+		}
+	}
+
+	template<typename TSimulation, typename TDriver>
+	void Reconcile(TSimulation* Simulation, TDriver* Driver, TNetworkSimBufferContainer<TBufferTypes>& Buffers, TSimulationTicker<TTickSettings>& Ticker)
+	{
+		if (ReconcileSimulationTime.IsPositive() == false)
+		{
+			return;
+		}
+
+		check(Ticker.GetTotalProcessedSimulationTime() <= Ticker.GetTotalAllowedSimulationTime());
+
+		if (bAllowSimulatedExtrapolation && ParentSimulation == nullptr && NetworkSimulationModelCVars::EnableSimulatedExtrapolation() && NetworkSimulationModelCVars::EnableSimulatedReconcile())
+		{
+			// Simulated Reconcile requires the input buffer to be kept up to date with the Sync buffer
+			// Generate a new, fake, command since we just added a new sync state to head
+			while (Buffers.Input.HeadFrame() < Buffers.Sync.HeadFrame())
+			{
+				TInputCmd* Next = Buffers.Input.WriteFrameInitializedFromHead(Buffers.Input.HeadFrame()+1);
+			}
+
+			// Do we have time to make up? We may have extrapolated ahead of the server (totally fine - can happen with small amount of latency variance)
+			FNetworkSimTime DeltaSimTime = ReconcileSimulationTime - Ticker.GetTotalProcessedSimulationTime();
+			if (DeltaSimTime.IsPositive() && NetworkSimulationModelCVars::EnableSimulatedReconcile())
+			{
+				SimulationExtrapolation<TSimulation, TDriver>(Simulation, Driver, Buffers, Ticker, DeltaSimTime);
+			}
+		}
+		
+		check(Ticker.GetTotalProcessedSimulationTime() <= Ticker.GetTotalAllowedSimulationTime());
+		ReconcileSimulationTime.Reset();
+	}
+
+	template<typename TDriver>
+	void PreSimTick(TDriver* Driver, TNetworkSimBufferContainer<TBufferTypes>& Buffers, TSimulationTicker<TTickSettings>& Ticker, const FNetSimTickParameters& TickParameters)
+	{
+		// Tick if we are dependent simulation or extrapolation is enabled
+		if (ParentSimulation || (bAllowSimulatedExtrapolation && NetworkSimulationModelCVars::EnableSimulatedExtrapolation()))
+		{
+			// Don't start this simulation until you've gotten at least one update from the server
+			if (Ticker.GetTotalProcessedSimulationTime().IsPositive())
+			{
+				Ticker.GiveSimulationTime(TickParameters.LocalDeltaTimeSeconds);
+			}
+
+			if (TickParameters.bGenerateLocalInputCmds)
+			{
+				TryGenerateLocalInput(Driver, Buffers, Ticker);
+			}
+		}
+	}
+
+	template<typename TDriver>
+	void PostSimTick(TDriver* Driver, const TNetworkSimBufferContainer<TBufferTypes>& Buffers, const TSimulationTicker<TTickSettings>& Ticker, const FNetSimTickParameters& TickParameters)
+	{
+		if (bAllowSimulatedExtrapolation || ParentSimulation)
+		{
+			// Sync to latest frame if there is any
+			if (Buffers.Sync.Num() > 0)
+			{
+				Driver->FinalizeFrame(*Buffers.Sync.HeadElement(), *Buffers.Aux.HeadElement());
+			}
+			
+		}
+		else
+		{
+			Interpolator.template PostSimTick<TDriver>(Driver, Buffers, Ticker, TickParameters);
+		}
+	}
+
+	template<typename TSimulation, typename TDriver>
+	void DependentRollbackBegin(TSimulation* Simulation, TDriver* Driver, TNetworkSimBufferContainer<TBufferTypes>& Buffers, TSimulationTicker<TTickSettings>& Ticker, const FNetworkSimTime& RollbackDeltaTime, const int32 ParentFrame)
+	{
+		// For now, we make the assumption that our last serialized state and time match the parent simulation.
+		// This would not always be the case with low frequency simulated proxies. But could be handled by replicating the simulations together (at the replication level)
+		const int32 NewHeadFrame = Buffers.Sync.HeadFrame()+1;
+		
+		Ticker.SetTotalProcessedSimulationTime(LastSerializedSimulationTime, NewHeadFrame);
+		Ticker.SetTotalAllowedSimulationTime(LastSerializedSimulationTime);
+		Ticker.PendingFrame = NewHeadFrame;
+		Ticker.MaxAllowedFrame = NewHeadFrame;
+
+		*Buffers.Sync.WriteFrame(NewHeadFrame) = LastSerializedSyncState;
+		*Buffers.Input.WriteFrame(NewHeadFrame) = TInputCmd();
+
+		Driver->FinalizeFrame(LastSerializedSyncState, LastSerializedAuxState);
+
+		if (NETSIM_MODEL_DEBUG)
+		{
+			FVisualLoggingParameters VLogParams(EVisualLoggingContext::FirstMispredicted, ParentFrame, EVisualLoggingLifetime::Persistent, TEXT("DependentRollbackBegin"));
+			Driver->VisualLog(nullptr, &LastSerializedSyncState, &LastSerializedAuxState, VLogParams);
+		}
+	}
+
+	template<typename TSimulation, typename TDriver>
+	void DependentRollbackStep(TSimulation* Simulation, TDriver* Driver, TNetworkSimBufferContainer<TBufferTypes>& Buffers, TSimulationTicker<TTickSettings>& Ticker, const FNetworkSimTime& StepTime, const int32 ParentFrame, const bool bFinalStep)
+	{
+		Ticker.SetTotalAllowedSimulationTime( Ticker.GetTotalAllowedSimulationTime() + StepTime );
+
+		SimulationExtrapolation<TSimulation, TDriver>(Simulation, Driver, Buffers, Ticker, StepTime);
+
+		TSyncState* SyncState = Buffers.Sync.HeadElement();
+		check(SyncState);
+		
+		if (NETSIM_MODEL_DEBUG)
+		{
+			FVisualLoggingParameters VLogParams(bFinalStep ? EVisualLoggingContext::LastMispredicted : EVisualLoggingContext::OtherMispredicted, ParentFrame, EVisualLoggingLifetime::Persistent, TEXT("DependentRollbackStep"));
+			Driver->VisualLog(nullptr, SyncState, Buffers.Aux.HeadElement(), VLogParams);
+		}
+	}
+
+private:
+
+	template<typename TSimulation, typename TDriver>
+	void SimulationExtrapolation(TSimulation* Simulation, TDriver* Driver, TNetworkSimBufferContainer<TBufferTypes>& Buffers, TSimulationTicker<TTickSettings>& Ticker, const FNetworkSimTime DeltaSimTime)
+	{
+		// We have extrapolated ahead of the server. The latest network update is now "in the past" from what we rendered last frame.
+		// We will insert a new frame to make up the difference from the last known state to where we want to be in the now.
+
+		ensure(Buffers.Input.HeadFrame() == Buffers.Sync.HeadFrame());
+
+		const int32 InputFrame = Buffers.Input.HeadFrame();
+		const int32 OutputFrame = InputFrame + 1;
+
+		// Create fake cmd				
+		TInputCmd* NewCmd = Buffers.Input.WriteFrameInitializedFromHead(OutputFrame);
+		NewCmd->SetFrameDeltaTime(DeltaSimTime);	
+		
+		// Create new sync state to write to
+		TSyncState* PrevSyncState = Buffers.Sync[InputFrame];
+		TSyncState* NextSyncState = Buffers.Sync.WriteFrame(OutputFrame);
+		TAuxState* AuxState = Buffers.Aux[InputFrame];
+
+		// Do the actual update
+		{
+			TScopedSimulationTick UpdateScope(Ticker, OutputFrame, NewCmd->GetFrameDeltaTime());
+			Simulation->SimulationTick( 
+				{ NewCmd->GetFrameDeltaTime(), Ticker },
+				{ *NewCmd, *PrevSyncState, *AuxState },
+				{ *NextSyncState, Buffers.Aux.LazyWriter(OutputFrame) } );
+		}
+
+		Ticker.MaxAllowedFrame = OutputFrame;
+	}
+	
+	FNetworkSimTime ReconcileSimulationTime;
+	FNetworkSimTime LastSerializedSimulationTime;
+	TSyncState LastSerializedSyncState;
+	TAuxState LastSerializedAuxState;	// Temp? This should be conditional or optional. We want to support not replicating the aux state to simulated proxies
+};
+
+
+/** Replicates TSyncState and does basic reconciliation. */
+template<typename TBufferTypes, typename TTickSettings, typename TBase=TReplicatorBase<TBufferTypes, TTickSettings>>
+struct TReplicator_Autonomous : public TBase
+{
+	using TInputCmd = typename TBufferTypes::TInputCmd;
+	using TSyncState = typename TBufferTypes::TSyncState;
+	using TAuxState = typename TBufferTypes::TAuxState;
+
+	int32 GetLastSerializedFrame() const { return SerializedFrame; }
+	bool IsReconcileFaultDetected() const { return bReconcileFaultDetected; }
+	const FNetworkSimTime& GetLastSerializedSimTime() const { return SerializedTime; }
+
+	TArray<INetworkSimulationModel*> DependentSimulations;
+	bool bDependentSimulationNeedsReconcile = false;
+
+	int32 GetProxyDirtyCount(TNetworkSimBufferContainer<TBufferTypes>& Buffers) const
+	{
+		return (Buffers.Sync.GetDirtyCount()) ^ (Buffers.Aux.GetDirtyCount() << 1) ^ (TBase::GetProxyDirtyCount(Buffers) << 2); 
+	}
+
+	// --------------------------------------------------------------------
+	//	NetSerialize
+	// --------------------------------------------------------------------
+	void NetSerialize(const FNetSerializeParams& P, TNetworkSimBufferContainer<TBufferTypes>& Buffers, TSimulationTicker<TTickSettings>& Ticker)
+	{
+		FArchive& Ar = P.Ar;
+		
+		SerializedFrame = FNetworkSimulationSerialization::SerializeFrame(Ar, Buffers.Sync.HeadFrame());
+		
+		SerializedTime = Ticker.GetTotalProcessedSimulationTime();
+		SerializedTime.NetSerialize(P.Ar);
+
+		if (Ar.IsSaving())
+		{
+			// Server serialize the latest state
+			Buffers.Sync.HeadElement()->NetSerialize(Ar);
+			Buffers.Aux.HeadElement()->NetSerialize(Ar);
+		}
+		else
+		{
+			SerializedSyncState.NetSerialize(Ar);
+			SerializedAuxState.NetSerialize(Ar);
+		}
+
+		if (Ar.IsLoading())
+		{
+			bReconcileFaultDetected = false;
+			bPendingReconciliation = false;
+
+			// The state the client predicted that corresponds to the state the server just serialized to us
+			if (TSyncState* ClientExistingState = Buffers.Sync[SerializedFrame])
+			{
+				// TODO: AuxState->ShouldReconcile
+				if (ClientExistingState->ShouldReconcile(SerializedSyncState) || (NetworkSimulationModelCVars::ForceReconcile() > 0) || (NetworkSimulationModelCVars::ForceReconcileSingle() > 0))
+				{
+					NetworkSimulationModelCVars::SetForceReconcileSingle(0);
+					UE_CLOG(!Buffers.Input.IsValidFrame(SerializedFrame-1), LogNetworkSim, Error, TEXT("::NetSerialize: Client InputBuffer does not contain data for frame %d. {%s} {%s}"), SerializedFrame, *Buffers.Input.GetBasicDebugStr(), *Buffers.Sync.GetBasicDebugStr());
+					bPendingReconciliation =  true;
+				}
+			}
+			else
+			{
+				if (SerializedFrame < Buffers.Sync.TailFrame())
+				{
+					// Case 1: the serialized state is older than what we've kept in our buffer. A bigger buffer would solve this! (at the price of more resimulated frames to recover when this happens)
+					// This is a reconcile fault and we just need to chill. We'll stop sending user commands until the cmds in flight flush through the system and we catch back up.
+					bReconcileFaultDetected = true;
+				}
+				else
+				{
+					// Case 2: We've received a newer frame than what we've processed locally. This could happen if we are buffering our inputs locally (but still sending to the server) or just not predicting
+					bPendingReconciliation =  true;
+				}
+			}
+		}
+	}
+
+	// --------------------------------------------------------------------
+	//	Reconcile
+	// --------------------------------------------------------------------
+	template<typename TSimulation, typename TDriver>
+	void Reconcile(TSimulation* Simulation, TDriver* Driver, TNetworkSimBufferContainer<TBufferTypes>& Buffers, TSimulationTicker<TTickSettings>& Ticker)
+	{
+		if (bPendingReconciliation == false && (bDependentSimulationNeedsReconcile == false || SerializedFrame == INDEX_NONE))
+		{
+			return;
+		}
+		bPendingReconciliation = false;
+		bDependentSimulationNeedsReconcile = false;
+
+		TSyncState* ClientSyncState = Buffers.Sync[SerializedFrame];
+		const bool bDoVisualLog = NetworkSimulationModelCVars::EnableLocalPrediction() > 0 && NETSIM_MODEL_DEBUG; // don't visual log if we have prediction disabled
+		
+		if (bDoVisualLog)
+		{
+			FVisualLoggingParameters VLogParameters(EVisualLoggingContext::LastConfirmed, SerializedFrame, EVisualLoggingLifetime::Persistent, TEXT("Serialized State"));
+			Driver->VisualLog(Buffers.Input[SerializedFrame], &SerializedSyncState, &SerializedAuxState, VLogParameters);
+		}
+
+		if (ClientSyncState)
+		{
+			// Existing ClientSyncState, log it before overwriting it
+			if (bDoVisualLog)
+			{
+				FVisualLoggingParameters VLogParameters(EVisualLoggingContext::FirstMispredicted, SerializedFrame, EVisualLoggingLifetime::Persistent, TEXT("Mispredicted State"));
+				Driver->VisualLog(Buffers.Input[SerializedFrame], ClientSyncState, Buffers.Aux[SerializedFrame], VLogParameters);
+			}
+		}
+		else
+		{
+			// No existing state, so create add it explicitly
+			ClientSyncState = Buffers.Sync.WriteFrame( SerializedFrame );
+		}
+
+		// Set client's sync state to the server version
+		check(ClientSyncState);
+		*ClientSyncState = SerializedSyncState;
+
+		// Set Client's aux state to the server version
+		*Buffers.Aux.WriteFrame(SerializedFrame) = SerializedAuxState;
+
+		const FNetworkSimTime RollbackDeltaTime = SerializedTime - Ticker.GetTotalProcessedSimulationTime();
+
+		// Set the canonical simulation time to what we received (we will advance it as we resimulate)
+		Ticker.SetTotalProcessedSimulationTime(SerializedTime, SerializedFrame);
+		Ticker.PendingFrame = SerializedFrame;
+		//Ticker.MaxAllowedFrame = FMath::Max(Ticker.MaxAllowedFrame, Ticker.PendingFrame); // Make sure this doesn't lag behind. This is the only place we should need to do this.
+
+		if (NetworkSimulationModelCVars::EnableLocalPrediction() == 0)
+		{
+			// If we aren't predicting at all, then we advanced the allowed sim time here, (since we aren't doing it in PreSimTick). This just keeps us constantly falling behind and not being able to toggle prediction on/off for debugging.
+			Ticker.SetTotalAllowedSimulationTime(SerializedTime);
+		}
+
+		// Tell dependent simulations to rollback
+		for (INetworkSimulationModel* DependentSim : DependentSimulations)
+		{
+			DependentSim->BeginRollback(RollbackDeltaTime, SerializedFrame);
+		}
+		
+		// Resimulate all user commands 
+		const int32 LastFrameToProcess = Ticker.MaxAllowedFrame;
+		for (int32 Frame = SerializedFrame; Frame <= LastFrameToProcess; ++Frame)
+		{
+			const int32 OutputFrame = Frame+1;
+
+			// Frame is the frame we are resimulating right now.
+			TInputCmd* ResimulateCmd  = Buffers.Input[Frame];
+			TAuxState* AuxState = Buffers.Aux[Frame];
+			TSyncState* PrevSyncState = Buffers.Sync[Frame];
+			TSyncState* NextSyncState = Buffers.Sync.WriteFrame(OutputFrame);
+			
+			check(ResimulateCmd);
+			check(PrevSyncState);
+			check(NextSyncState);
+			check(AuxState);
+
+			// Log out the Mispredicted state that we are about to overwrite.
+			if (NETSIM_MODEL_DEBUG)
+			{
+				FVisualLoggingParameters VLogParameters(Frame == LastFrameToProcess ? EVisualLoggingContext::LastMispredicted : EVisualLoggingContext::OtherMispredicted, Frame, EVisualLoggingLifetime::Persistent, TEXT("Resimulate Step: mispredicted"));
+				Driver->VisualLog(Buffers.Input[OutputFrame], Buffers.Sync[OutputFrame], Buffers.Aux[OutputFrame], VLogParameters);
+			}
+
+			// Do the actual update
+			{
+				TScopedSimulationTick UpdateScope(Ticker, OutputFrame, ResimulateCmd->GetFrameDeltaTime());
+				Simulation->SimulationTick( 
+					{ ResimulateCmd->GetFrameDeltaTime(), Ticker },
+					{ *ResimulateCmd, *PrevSyncState, *AuxState },
+					{ *NextSyncState, Buffers.Aux.LazyWriter(OutputFrame) } );
+			}
+
+			// Log out the newly predicted state that we got.
+			if (NETSIM_MODEL_DEBUG)
+			{			
+				FVisualLoggingParameters VLogParameters(Frame == LastFrameToProcess ? EVisualLoggingContext::LastPredicted : EVisualLoggingContext::OtherPredicted, Frame, EVisualLoggingLifetime::Persistent, TEXT("Resimulate Step: repredicted"));
+				Driver->VisualLog(Buffers.Input[OutputFrame], Buffers.Sync[OutputFrame], Buffers.Aux[OutputFrame], VLogParameters);
+			}
+
+			// Tell dependent simulations to advance
+			for (INetworkSimulationModel* DependentSim : DependentSimulations)
+			{
+				DependentSim->StepRollback(ResimulateCmd->GetFrameDeltaTime(), Frame, (Frame == LastFrameToProcess));
+			}
+		}
+	}
+
+	// --------------------------------------------------------------------
+	//	PreSimTick
+	// --------------------------------------------------------------------
+	template<typename TDriver>
+	void PreSimTick(TDriver* Driver, TNetworkSimBufferContainer<TBufferTypes>& Buffers, TSimulationTicker<TTickSettings>& Ticker, const FNetSimTickParameters& TickParameters)
+	{
+		// If we have a reconcile fault, we cannot continue on with the simulation until it clears itself out. This effectively drops the input time and does not sample new inputs
+		if (bReconcileFaultDetected)
+		{
+			return;
+		}
+
+		if (TickParameters.bGenerateLocalInputCmds)
+		{
+			if (NetworkSimulationModelCVars::EnableLocalPrediction() > 0)
+			{
+				// Prediction: add simulation time and generate new commands
+				Ticker.GiveSimulationTime(TickParameters.LocalDeltaTimeSeconds);
+				TryGenerateLocalInput(Driver, Buffers, Ticker);
+			}
+			else
+			{
+				// Local prediction disabled: we must use a separate time accumulator to figure out when we should add more input cmds.
+				// Since we aren't processing the simulation locally, our core simulation time will only advance from network updates.
+				// (still, we need *something* to tell us when to generate a new command and what the delta time should be)
+				FNetworkSimTime NonPredictedInputTime;
+				NonPredictedInputTimeAccumulator.Accumulate(NonPredictedInputTime, TickParameters.LocalDeltaTimeSeconds);
+				if (NonPredictedInputTime.IsPositive())
+				{
+					GenerateLocalInputCmdAtFrame(Driver, Buffers, NonPredictedInputTime, Buffers.Input.HeadFrame() + 1);
+				}
+			}
+		}
+	}
+
+private:
+	
+	TSyncState SerializedSyncState;
+	TAuxState SerializedAuxState;
+	FNetworkSimTime SerializedTime;
+	int32 SerializedFrame = -1;
+
+	bool bPendingReconciliation = false;	// Reconciliation is pending: we need to reconcile state from the server that differs from the locally predicted state
+	bool bReconcileFaultDetected = false;	// A fault was detected: we received state from the server that we are unable to reconcile with locally predicted state
+
+	TRealTimeAccumulator<TTickSettings> NonPredictedInputTimeAccumulator; // for tracking input time in the non predictive case
+};
+
+/** Special replicator for debug buffer, this preserves the local buffer and receives into a replicator-owned buffer (we want these buffers to be distinct/not merged) */
+template<typename TBufferTypes, typename TTickSettings, int32 MaxNumElements=5, typename TBase=TReplicatorEmpty<TBufferTypes, TTickSettings>>
+struct TReplicator_Debug : public TBase
+{
+	using TDebugState = typename TBufferTypes::TDebugState;
+	using TDebugBuffer = typename TNetworkSimBufferContainer<TBufferTypes>::TDebugBuffer;
+
+	int32 GetProxyDirtyCount(TNetworkSimBufferContainer<TBufferTypes>& Buffers) const 
+	{
+		return Buffers.Debug.GetDirtyCount() ^ (TBase::GetProxyDirtyCount(Buffers) << 2); 
+	}
+
+	void NetSerialize(const FNetSerializeParams& P, TNetworkSimBufferContainer<TBufferTypes>& Buffers, TSimulationTicker<TTickSettings>& Ticker)
+	{
+		TBase::NetSerialize(P, Buffers, Ticker);
+		FArchive& Ar = P.Ar;
+
+		TDebugBuffer& Buffer = Ar.IsSaving() ? Buffers.Debug : ReceivedBuffer;
+		
+		uint8 SerializedNumElements = FMath::Min<uint8>(MaxNumElements, Buffer.Num());
+		Ar << SerializedNumElements;
+
+		const int32 HeadFrame = FNetworkSimulationSerialization::SerializeFrame(Ar, Buffer.HeadFrame());
+		const int32 StartingFrame = FMath::Max(0, HeadFrame - SerializedNumElements + 1);
+
+		if (Ar.IsLoading())
+		{
+			const int32 PrevHead = Buffer.HeadFrame();
+			if (PrevHead < StartingFrame && PrevHead >= 0)
+			{
+				// There is a gap in the stream. In some cases, we want this to be a "fault" and bubble up. We may want to synthesize state or maybe we just skip ahead.
+				UE_LOG(LogNetworkSim, Warning, TEXT("Fault: gap in received Debug buffer. PrevHead: %d. Received: %d-%d. Reseting previous buffer contents"), PrevHead, StartingFrame, HeadFrame);
+			}
+		}
+
+		for (int32 Frame = StartingFrame; Frame <= HeadFrame; ++Frame)
+		{
+			// This, as is, is bad. The intention is that these functions serialize multiple items in some delta compressed fashion.
+			// As is, we are just serializing the elements individually.
+			auto* Cmd = Ar.IsLoading() ? Buffer.WriteFrame(Frame) : Buffer[Frame];
+			Cmd->NetSerialize(P);
+		}
+	}
+
+	TDebugBuffer ReceivedBuffer;
+};
\ No newline at end of file
diff --git a/Engine/Plugins/Runtime/NetworkPrediction/Source/NetworkPrediction/Public/NetworkedSimulationModelTypes.h b/Engine/Plugins/Runtime/NetworkPrediction/Source/NetworkPrediction/Public/NetworkedSimulationModelTypes.h
new file mode 100644
--- /dev/null
+++ b/Engine/Plugins/Runtime/NetworkPrediction/Source/NetworkPrediction/Public/NetworkedSimulationModelTypes.h
@@ -0,0 +1,739 @@
+// Copyright 1998-2019 Epic Games, Inc. All Rights Reserved.
+
+#pragma once
+
+#include "NetworkedSimulationModelBuffer.h"
+
+// ---------------------------------------------------------------------------------------------------------------------
+//	BufferTypes: template helpers for addressing the different buffer types of the system.
+// ---------------------------------------------------------------------------------------------------------------------
+
+// Enum to refer to buffer type. These are used as template arguments to write generic code that can act on any of the buffers.
+enum class ENetworkSimBufferTypeId : uint8
+{
+	Input,
+	Sync,
+	Aux,
+	Debug
+};
+
+inline FString LexToString(ENetworkSimBufferTypeId A)
+{
+	switch(A)
+	{
+		case ENetworkSimBufferTypeId::Input: return TEXT("Input");
+		case ENetworkSimBufferTypeId::Sync: return TEXT("Sync");
+		case ENetworkSimBufferTypeId::Aux: return TEXT("Aux");
+		case ENetworkSimBufferTypeId::Debug: return TEXT("Debug");
+	};
+	return TEXT("Unknown");
+}
+
+// Helper needed to specialize TNetworkSimBufferTypes::Get (must be done outside of templated struct)
+template<typename TBufferTypes, ENetworkSimBufferTypeId BufferId> 
+struct TSelectTypeHelper
+{
+	using type = void;
+};
+
+template<typename TBufferTypes>
+struct TSelectTypeHelper<TBufferTypes, ENetworkSimBufferTypeId::Input>
+{
+	using type = typename TBufferTypes::TInputCmd;
+};
+
+template<typename TBufferTypes>
+struct TSelectTypeHelper<TBufferTypes, ENetworkSimBufferTypeId::Sync>
+{
+	using type = typename TBufferTypes::TSyncState;
+};
+
+template<typename TBufferTypes>
+struct TSelectTypeHelper<TBufferTypes, ENetworkSimBufferTypeId::Aux>
+{
+	using type = typename TBufferTypes::TAuxState;
+};
+
+template<typename TBufferTypes>
+struct TSelectTypeHelper<TBufferTypes, ENetworkSimBufferTypeId::Debug>
+{
+	using type = typename TBufferTypes::TDebugState;
+};
+
+enum ENetworkSimBufferAllocationType
+{
+	Contiguous,
+	Sparse
+};
+
+template<ENetworkSimBufferAllocationType InAllocationType, int32 InSize>
+struct TNetworkSimBufferAllocation
+{
+	enum { Size = InSize };
+	static constexpr ENetworkSimBufferAllocationType AllocationType() { return InAllocationType; }
+};
+
+template<typename T>
+struct HasNetSerialize
+{
+	template<typename U, void (U::*)(const FNetSerializeParams& P)> struct SFINAE {};
+	template<typename U> static char Test(SFINAE<U, &U::NetSerialize>*);
+	template<typename U> static int Test(...);
+	static const bool Has = sizeof(Test<T>(0)) == sizeof(char);
+};
+
+template<typename T>
+struct HasLog
+{
+	template<typename U, void (U::*)(FStandardLoggingParameters& P) const> struct SFINAE {};
+	template<typename U> static char Test(SFINAE<U, &U::Log>*);
+	template<typename U> static int Test(...);
+	static const bool Has = sizeof(Test<T>(0)) == sizeof(char);
+};
+
+// A collection of the system's buffer types. This allows us to collapse the 4 types into a single type to use a template argument elsewhere.
+template<typename InInputCmd, typename InSyncState, typename InAuxState, typename InDebugState = FNetSimProcessedFrameDebugInfo>
+struct TNetworkSimBufferTypes
+{
+	// Quick access to types when you know what you want
+	using TInputCmd = InInputCmd;
+	using TSyncState = InSyncState;
+	using TAuxState = InAuxState;
+	using TDebugState = InDebugState;
+
+	// Template access via ENetworkSimBufferTypeId when "which buffer" is parameterized
+	template<ENetworkSimBufferTypeId Id>
+	struct select_type
+	{
+		using type = typename TSelectTypeHelper< TNetworkSimBufferTypes<TInputCmd, TSyncState, TAuxState, TDebugState>, Id >::type;
+	};
+
+	// Must implement NetSerialize and Log functions. This purposefully does not pass on inherited methods
+	static_assert(HasNetSerialize<InInputCmd>::Has == true, "InputCmd Must implement NetSerialize");
+	static_assert(HasNetSerialize<InSyncState>::Has == true, "SyncState Must implement NetSerialize");
+	static_assert(HasNetSerialize<InAuxState>::Has == true, "AuxState Must implement NetSerialize");
+	static_assert(HasNetSerialize<InDebugState>::Has == true, "DebugState Must implement NetSerialize");
+
+	static_assert(HasLog<InInputCmd>::Has == true, "InputCmd Must implement Log");
+	static_assert(HasLog<InSyncState>::Has == true, "SyncState Must implement Log");
+	static_assert(HasLog<InAuxState>::Has == true, "AuxState Must implement Log");
+	static_assert(HasLog<InDebugState>::Has == true, "DebugState Must implement Log");
+};
+
+// ---------------------------------------------------------------------------------------------------------------------
+//	TNetworkSimBufferContainer
+//	Container for the actual replicated buffers that the system uses.
+//	Has compile time accessors for retrieving the buffers based on templated enum value.
+// ---------------------------------------------------------------------------------------------------------------------
+
+// Helper struct for enum-based access. This has to be done in an outside struct because we cant specialize inside TNetworkSimBufferContainer on all compilers
+template<typename TContainer, ENetworkSimBufferTypeId BufferId>
+struct TBufferGetterHelper
+{
+	static typename TContainer::template select_buffer_type<BufferId>::type& Get(TContainer& Container)
+	{
+		static_assert(!Container, "Failed to find specialized Get for your BufferId");
+	}
+};
+
+template<typename TContainer>
+struct TBufferGetterHelper<TContainer, ENetworkSimBufferTypeId::Input>
+{
+	static typename TContainer::TInputBuffer& Get(TContainer& Container)
+	{
+		return Container.Input;
+	}
+};
+template<typename TContainer>
+struct TBufferGetterHelper<TContainer, ENetworkSimBufferTypeId::Sync>
+{
+	static typename TContainer::TSyncBuffer& Get(TContainer& Container)
+	{
+		return Container.Sync;
+	}
+};
+template<typename TContainer>
+struct TBufferGetterHelper<TContainer, ENetworkSimBufferTypeId::Aux>
+{
+	static typename TContainer::TAuxBuffer& Get(TContainer& Container)
+	{
+		return Container.Aux;
+	}
+};
+template<typename TContainer>
+struct TBufferGetterHelper<TContainer, ENetworkSimBufferTypeId::Debug>
+{
+	static typename TContainer::TDebugBuffer& Get(TContainer& Container)
+	{
+		return Container.Debug;
+	}
+};
+
+// -------------------------------------------------
+
+// Helper struct for accessing a netsim buffer given the underlying type
+template<typename TContainer, typename TState>
+struct TBufferGetterHelper_ByState
+{
+	static void GetBuffer(TContainer& Container)
+	{
+		static_assert(!Container, "Failed to find specialized Get for your BufferId");
+	}
+};
+
+template<typename TContainer>
+struct TBufferGetterHelper_ByState<TContainer, typename TContainer::TInputCmd>
+{
+	static typename TContainer::TInputBuffer& GetBuffer(TContainer& Container)
+	{
+		return Container.Input;
+	}
+};
+
+template<typename TContainer>
+struct TBufferGetterHelper_ByState<TContainer, typename TContainer::TSyncState>
+{
+	static typename TContainer::TSyncBuffer& GetBuffer(TContainer& Container)
+	{
+		return Container.Sync;
+	}
+};
+
+template<typename TContainer>
+struct TBufferGetterHelper_ByState<TContainer, typename TContainer::TAuxState>
+{
+	static typename TContainer::TAuxBuffer& GetBuffer(TContainer& Container)
+	{
+		return Container.Aux;
+	}
+};
+
+template<typename TContainer>
+struct TBufferGetterHelper_ByState<TContainer, typename TContainer::TDebugState>
+{
+	static typename TContainer::TDebugBuffer& GetBuffer(TContainer& Container)
+	{
+		return Container.Debug;
+	}
+};
+
+// -------------------------------------------------------------------
+
+// Struct that encapsulates writing a new element to a buffer. This is used to allow a new Aux state to be created in the ::SimulationTick loop.
+template<typename T>
+struct TLazyStateAccessor
+{
+	TLazyStateAccessor(TUniqueFunction<T*()> && InFunc)
+		: GetWriteNextFunc(MoveTemp(InFunc)) { }
+
+	T* GetWriteNext() const
+	{
+		if (CachedWriteNext == nullptr)
+		{
+			CachedWriteNext = GetWriteNextFunc();
+		}
+		return CachedWriteNext;
+	}
+
+private:
+
+	mutable T* CachedWriteNext = nullptr;
+	TUniqueFunction<T*()> GetWriteNextFunc;
+};
+
+// The main container for all of our buffers
+template<typename InBufferTypes>
+struct TNetworkSimBufferContainer
+{
+	// Collection of types we were assigned
+	using TBufferTypes = InBufferTypes;
+
+	using TInputCmd = typename TBufferTypes::TInputCmd;
+	using TSyncState = typename TBufferTypes::TSyncState;
+	using TAuxState = typename TBufferTypes::TAuxState;
+	using TDebugState = typename TBufferTypes::TDebugState;
+
+	// helper that returns the buffer type for a given BufferId (not the underlying type: the actual TReplicatedBuffer<TUnderlyingType>
+	template<ENetworkSimBufferTypeId TypeId>
+	struct select_buffer_type
+	{
+		// We are just wrapping a TNetworkSimContiguousBuffer around whatever type TBufferTypes::select_type returns (which returns the underlying type)
+		using type = TNetworkSimContiguousBuffer<typename TBufferTypes::template select_type<TypeId>::type>;
+	};
+
+	template<ENetworkSimBufferTypeId TypeId>
+	struct select_buffer_type_sparse
+	{
+		// We are just wrapping a TNetworkSimSparseBuffer around whatever type TBufferTypes::select_type returns (which returns the underlying type)
+		using type = TNetworkSimSparseBuffer<typename TBufferTypes::template select_type<TypeId>::type>;
+	};
+
+	// The buffer types. This may look intimidating but is not that bad!
+	// We are using select_buffer_type to decide what the type should be for a given ENetworkSimBufferTypeId.
+	using TInputBuffer = typename select_buffer_type<ENetworkSimBufferTypeId::Input>::type;
+	using TSyncBuffer = typename select_buffer_type<ENetworkSimBufferTypeId::Sync>::type;
+	using TAuxBuffer = typename select_buffer_type_sparse<ENetworkSimBufferTypeId::Aux>::type;
+	using TDebugBuffer = typename select_buffer_type<ENetworkSimBufferTypeId::Debug>::type;
+
+	// The buffers themselves. Types are already declared above.
+	// If you are reading just to find the damn underlying type here, its (TNetworkSimSparseBuffer||TNetworkSimContiguousBuffer) < whatever your struct type is >
+	TInputBuffer Input;
+	TSyncBuffer Sync;
+	TAuxBuffer Aux;
+	TDebugBuffer Debug;	
+
+	// Finally, template accessor for getting buffers based on enum. This is really what all this junk is about.
+	// This allows other templated classes in the system to access a specific buffer from another templated argument
+	template<ENetworkSimBufferTypeId BufferId>
+	typename select_buffer_type<BufferId>::type& Get()
+	{
+		return TBufferGetterHelper<TNetworkSimBufferContainer<TBufferTypes>, BufferId>::Get(*this);
+	}
+};
+
+// ----------------------------------------------------------------------------------------------------------------------------------------------
+//		Tick and time keeping related structures
+// ----------------------------------------------------------------------------------------------------------------------------------------------
+
+// The main Simulation time type. All sims use this to talk about time.
+struct FNetworkSimTime
+{
+	using FSimTime = int32; // Underlying type used to store simulation time
+	using FRealTime = float; // Underlying type used when dealing with realtime (usually coming in from the engine tick).
+
+	enum
+	{
+		RealToSimFactor = 1000 // Factor to go from RealTime (always seconds) to SimTime (MSec by default with factor of 1000)
+	};
+
+	static constexpr FRealTime GetRealToSimFactor() { return static_cast<FRealTime>(RealToSimFactor); }
+	static constexpr FRealTime GetSimToRealFactor() { return static_cast<FRealTime>(1.f / RealToSimFactor); }
+
+	// ---------------------------------------------------------------
+
+	FNetworkSimTime() { }
+
+	// Things get confusing with templated types and overloaded functions. To avoid that, use these funcs to construct from either msec or real time
+	static inline FNetworkSimTime FromMSec(const FSimTime& InTime) { return FNetworkSimTime(InTime); } 
+	static inline FNetworkSimTime FromRealTimeSeconds(const FRealTime& InRealTime) { return FNetworkSimTime( static_cast<FSimTime>(InRealTime * GetRealToSimFactor())); } 
+	FRealTime ToRealTimeSeconds() const { return (Time * GetSimToRealFactor()); }
+
+	// Direct casts to "real time MS" which should be rarely used in practice (TRealTimeAccumulator only current case). All other cases of "real time" imply seconds.
+	static inline FNetworkSimTime FromRealTimeMS(const FRealTime& InRealTime) { return FNetworkSimTime( static_cast<FSimTime>(InRealTime)); } 
+	FRealTime ToRealTimeMS() const { return static_cast<FRealTime>(Time); }
+
+	FString ToString() const { return LexToString(this->Time); }
+
+	bool IsPositive() const { return (Time > 0); }
+	bool IsNegative() const { return (Time < 0); }
+	void Reset() { Time = 0; }
+
+	// FIXME
+	void NetSerialize(FArchive& Ar) { Ar << Time; }
+
+	using T = FNetworkSimTime;
+	T& operator+= (const T &rhs) { this->Time += rhs.Time; return(*this); }
+	T& operator-= (const T &rhs) { this->Time -= rhs.Time; return(*this); }
+	
+	T operator+ (const T &rhs) const { return T(this->Time + rhs.Time); }
+	T operator- (const T &rhs) const { return T(this->Time - rhs.Time); }
+
+	bool operator<  (const T &rhs) const { return(this->Time < rhs.Time); }
+	bool operator<= (const T &rhs) const { return(this->Time <= rhs.Time); }
+	bool operator>  (const T &rhs) const { return(this->Time > rhs.Time); }
+	bool operator>= (const T &rhs) const { return(this->Time >= rhs.Time); }
+	bool operator== (const T &rhs) const { return(this->Time == rhs.Time); }
+	bool operator!= (const T &rhs) const { return(this->Time != rhs.Time); }
+
+private:
+
+	FSimTime Time = 0;
+
+	FNetworkSimTime(const FSimTime& InTime) { Time = InTime; }
+};
+
+// Holds per-simulation settings about how ticking is supposed to happen.
+template<int32 InFixedStepMS=0, int32 InMaxStepMS=0>
+struct TNetworkSimTickSettings
+{
+	static_assert(!(InFixedStepMS!=0 && InMaxStepMS != 0), "MaxStepMS is only applicable when using variable step (FixedStepMS == 0)");
+	enum 
+	{
+		MaxStepMS = InMaxStepMS,				// Max step. Only applicable to variable time step.
+		FixedStepMS = InFixedStepMS,			// Fixed step. If 0, then we are "variable time step"
+	};
+
+	// Typed accessors
+	static constexpr FNetworkSimTime::FSimTime GetMaxStepMS() { return static_cast<FNetworkSimTime::FSimTime>(InMaxStepMS); }
+	static constexpr FNetworkSimTime::FSimTime GetFixedStepMS() { return static_cast<FNetworkSimTime::FSimTime>(InFixedStepMS); }
+};
+
+// ----------------------------------------------------------------------------------------------------------------------------------------------
+//	Accumulator: Helper for accumulating real time into sim time based on TickSettings
+// ----------------------------------------------------------------------------------------------------------------------------------------------
+
+template<typename TickSettings, bool IsFixedTick=(TickSettings::FixedStepMS!=0)>
+struct TRealTimeAccumulator
+{
+	using TRealTime = FNetworkSimTime::FRealTime;
+	void Accumulate(FNetworkSimTime& NetworkSimTime, const TRealTime RealTimeSeconds)
+	{
+		// Even though we are variable tick, we still want to truncate down to an even msec. This keeps sim steps as whole integer values that serialize better, don't have denormals or other floating point weirdness.
+		// (If we ever wanted to just fully let floats pass through and be used as the msec sim time, this could be done through another specialization for float/float time).		
+		// Also note that MaxStepMS enforcement does NOT belong here. Dropping time due to MaxStepMS would just make the sim run slower. MaxStepMS is used at the input processing level.
+		
+		AccumulatedTimeMS += RealTimeSeconds * FNetworkSimTime::GetRealToSimFactor(); // convert input seconds -> msec
+		const FNetworkSimTime AccumulatedSimTimeMS = FNetworkSimTime::FromRealTimeMS(AccumulatedTimeMS);	// truncate (float) accumulated msec to (int32) sim time msec
+
+		NetworkSimTime += AccumulatedSimTimeMS;
+		AccumulatedTimeMS -= AccumulatedSimTimeMS.ToRealTimeMS(); // subtract out the "whole" msec, we are left with the remainder msec
+	}
+
+	void Reset()
+	{
+		AccumulatedTimeMS = 0.f;
+	}
+
+private:
+
+	TRealTime AccumulatedTimeMS = 0.f;
+};
+
+// Specialized version of FixedTicking. This accumulates real time that spills over into NetworkSimTime as it crosses the FixStep threshold
+template<typename TickSettings>
+struct TRealTimeAccumulator<TickSettings, true>
+{
+	using TRealTime = FNetworkSimTime::FRealTime;
+	const TRealTime RealTimeFixedStep = static_cast<TRealTime>(TickSettings::GetFixedStepMS() * FNetworkSimTime::GetSimToRealFactor());
+
+	void Accumulate(FNetworkSimTime& NetworkSimTime, const TRealTime RealTimeSeconds)
+	{
+		AccumulatedTime += RealTimeSeconds;
+		if (AccumulatedTime > RealTimeFixedStep)
+		{
+			const int32 NumFrames = AccumulatedTime / RealTimeFixedStep;
+			AccumulatedTime -= NumFrames * RealTimeFixedStep;
+				
+			if (FMath::Abs<TRealTime>(AccumulatedTime) < SMALL_NUMBER)
+			{
+				AccumulatedTime = TRealTime(0.f);
+			}
+
+			NetworkSimTime += FNetworkSimTime::FromMSec(NumFrames * TickSettings::FixedStepMS);
+		}
+	}
+
+	void Reset()
+	{
+		AccumulatedTime = 0.f;
+	}
+
+private:
+
+	TRealTime AccumulatedTime;
+};
+
+// ----------------------------------------------------------------------------------------------------------------------------------------------
+//	FSimulationTickState: Holds active state for simulation ticking. We track two things: frames and time.
+//
+//	PendingFrame is the next frame we will process: the input/sync/aux state @ PendingFrame will be run through ::SimulationTick and produce the 
+//	next frame's (PendingFrame+1) Sync and possibly Aux state (if it changes). "Out of band" modifications to the sync/aux state should happen
+//	to PendingFrame (e.g, before it is processed. Once a frame has been processed, we won't run it through ::SimulationTick again!).
+//
+//	MaxAllowedFrame is a frame based limiter on simulation updates. This must be incremented to allow the simulation to advance.
+//
+//	Time is also tracked. We keep running total for how much the sim has advanced and how much it is allowed to advance. There is also a historic buffer of
+//	simulation time in SimulationTimeBuffer.
+//
+//	Consider that Frames are essentially client dependent and gaps can happen due to packet loss, etc. Time will always be continuous though.
+//	
+// ----------------------------------------------------------------------------------------------------------------------------------------------
+
+struct FSimulationTickState
+{
+	int32 PendingFrame = 0;
+	int32 MaxAllowedFrame = -1;
+	bool bUpdateInProgress = false;
+	
+	FNetworkSimTime GetTotalProcessedSimulationTime() const 
+	{ 
+		return TotalProcessedSimulationTime; 
+	}
+
+	void SetTotalProcessedSimulationTime(const FNetworkSimTime& SimTime, int32 Frame)
+	{
+		TotalProcessedSimulationTime = SimTime;
+		*SimulationTimeBuffer.WriteFrame(Frame) = SimTime;
+	}
+
+	void IncrementTotalProcessedSimulationTime(const FNetworkSimTime& DeltaSimTime, int32 Frame)
+	{
+		TotalProcessedSimulationTime += DeltaSimTime;
+		*SimulationTimeBuffer.WriteFrame(Frame) = TotalProcessedSimulationTime;
+	}
+		
+	// Historic tracking of simulation time. This allows us to timestamp sync data as its produced
+	TNetworkSimContiguousBuffer<FNetworkSimTime>	SimulationTimeBuffer;
+
+	// How much granted simulation time is left to process
+	FNetworkSimTime GetRemainingAllowedSimulationTime() const
+	{
+		return TotalAllowedSimulationTime - TotalProcessedSimulationTime;
+	}
+
+	FNetworkSimTime GetTotalAllowedSimulationTime() const
+	{
+		return TotalAllowedSimulationTime;
+	}
+
+protected:
+
+	FNetworkSimTime TotalAllowedSimulationTime;	// Total time we have been "given" to process. We cannot process more simulation time than this: doing so would be speed hacking.
+	FNetworkSimTime TotalProcessedSimulationTime;	// How much time we've actually processed. The only way to increment this is to process user commands or receive authoritative state from the network.
+};
+
+// "Ticker" that actually allows us to give the simulation time. This struct will generally not be passed around outside of the core TNetworkedSimulationModel/Replicators
+template<typename TickSettings=TNetworkSimTickSettings<>>
+struct TSimulationTicker : public FSimulationTickState
+{
+	using TSettings = TickSettings;
+
+	void SetTotalAllowedSimulationTime(const FNetworkSimTime& SimTime)
+	{
+		TotalAllowedSimulationTime = SimTime;
+		RealtimeAccumulator.Reset();
+	}
+
+	// "Grants" allowed simulation time to this tick state. That is, we are now allowed to advance the simulation by this amount the next time the sim ticks.
+	// Note the input is RealTime in SECONDS. This is what the rest of the engine uses when dealing with float delta time.
+	void GiveSimulationTime(float RealTimeSeconds)
+	{
+		RealtimeAccumulator.Accumulate(TotalAllowedSimulationTime, RealTimeSeconds);
+	}	
+
+private:
+
+	TRealTimeAccumulator<TSettings>	RealtimeAccumulator;	
+};
+
+// Scoped helper to be used right before entering a call to the sim's ::SimulationTick function.
+// Important to note that this advances the PendingFrame to the output Frame. So that any writes that occur to the buffers during this scope will go to the output frame.
+struct TScopedSimulationTick
+{
+	TScopedSimulationTick(FSimulationTickState& InTicker, const int32& InOutputFrame, const FNetworkSimTime& InDeltaSimTime)
+		: Ticker(InTicker), OutputFrame(InOutputFrame), DeltaSimTime(InDeltaSimTime)
+	{
+		check(Ticker.bUpdateInProgress == false);
+		Ticker.PendingFrame = OutputFrame;
+		Ticker.bUpdateInProgress = true;
+	}
+	~TScopedSimulationTick()
+	{
+		Ticker.IncrementTotalProcessedSimulationTime(DeltaSimTime, OutputFrame);
+		Ticker.bUpdateInProgress = false;
+	}
+	FSimulationTickState& Ticker;
+	const int32& OutputFrame;
+	const FNetworkSimTime& DeltaSimTime;
+};
+
+// ----------------------------------------------------------------------------------------------------------------------------------------------
+//	Accessors - helper structs that provide safe/cleaner access to the underlying NetSim states/events
+// ----------------------------------------------------------------------------------------------------------------------------------------------
+
+// Accessor conditionally gives access to the current (pending) Sync/Aux state to outside code.
+// Reads are always allowed
+// Writes are conditional. Authority can always write to the pending frame. Non authority requires the netsim to be currently processing an ::SimulationTick.
+// If you aren't inside an ::SimulationTick call, it is really not safe to predict state changes. It is safest and simplest to just not predict these changes.
+//
+// Explanation: During the scope of an ::SimulationTick call, we know exactly 'when' we are relative to what the server is processing. If the predicting client wants
+// to predict a change to sync/aux state during an update, the server will do it at the exact same time (assuming not a mis prediction). When a state change
+// happens "out of band" (outside an ::SimulationTick call) - we really have no way to correlate when the server will do it. While its tempting to think "we will get
+// a correction anyways, might as well guess at it and maybe get a smaller correction" - but this opens us up to other problems. The server may actually not 
+// change the state at all and you may not get an update that corrects you. You could add a timeout and track the state change somewhere but that really complicates
+// things and could leave you open to "double" problems: if the state change is additive, you may stack the authority change on top of the local predicted change, or
+// you may roll back the predicted change to then later receive the authority change.
+//	
+// What still may make sense to do is allow the "In Update" bool to be temporarily disabled if we enter code that we know is not rollback friendly.
+template<typename TState>
+struct TNetworkSimStateAccessor
+{
+	template<typename TNetworkSimModel>
+	void Init(TNetworkSimModel* NetSimModel)
+	{
+		GetStateFunc = [NetSimModel](bool bWrite, TState*& OutState, bool& OutSafe)
+		{
+			// Gross: find the buffer our TState is in. This allows these accessors to be declared as class member variables without knowing the exact net sim it will
+			// be pulling from. (For example, you may want to templatize your network sim model).
+			auto& Buffer = TBufferGetterHelper_ByState<TNetworkSimBufferContainer<typename TNetworkSimModel::TBufferTypes>, TState>::GetBuffer(NetSimModel->Buffers);
+			OutState = bWrite ? Buffer.WriteFrameInitializedFromHead(NetSimModel->Ticker.PendingFrame) : Buffer[NetSimModel->Ticker.PendingFrame];
+			OutSafe = NetSimModel->Ticker.bUpdateInProgress;
+		};
+	}
+
+	void Clear()
+	{
+		GetStateFunc = nullptr;
+	}
+
+	/** Gets the current (PendingFrame) state for reading. This is not expected to fail outside of startup/shutdown edge cases */
+	const TState* GetStateRead() const
+	{
+		TState* State = nullptr;
+		bool bSafe = false;
+		if (GetStateFunc)
+		{
+			GetStateFunc(false, State, bSafe);
+		}
+		return State;
+	}
+
+	/** Gets the current (PendingFrame) state for writing. This is expected to fail outside of the core update loop when bHasAuthority=false. (E.g, it is not safe to predict writes) */
+	TState* GetStateWrite(bool bHasAuthority) const
+	{
+		TState* State = nullptr;
+		bool bSafe = false;
+		if (GetStateFunc)
+		{
+			GetStateFunc(true, State, bSafe);
+		}
+		return (bHasAuthority || bSafe) ? State : nullptr;
+	}
+
+private:
+
+	TFunction<void(bool bForWrite, TState*& OutState, bool& OutSafe)> GetStateFunc;
+};
+
+//TODO
+struct FNetworkSimEventAccessor
+{
+	int32 GetPendingFrame() { return -1; }
+
+	struct FFrameEvents
+	{
+		/** This frame has been received from the authority and will not be rolled back or resimulated again. */
+		DECLARE_MULTICAST_DELEGATE(FOnFrameConfirmed)
+		FOnFrameConfirmed Confirmed;
+
+		/** This frame was previously simulated and the simulation has now been rolled back */
+		DECLARE_MULTICAST_DELEGATE(FOnFrameRolledBack)
+		FOnFrameRolledBack RolledBack;
+
+		DECLARE_MULTICAST_DELEGATE(FOnFrameSimulated)
+		FOnFrameRolledBack Simulated;
+	};
+};
+
+// ----------------------------------------------------------------------------------------------------------------------------------------------
+//	FrameCmd - in variable tick simulations we store the timestep of each frame with the input. TFrameCmd wraps the user struct to do this.
+// ----------------------------------------------------------------------------------------------------------------------------------------------
+
+// Wraps an input command (BaseType) in a NetworkSimulation time. 
+template<typename BaseCmdType, typename TickSettings, bool IsFixedTick=(TickSettings::FixedStepMS!=0)>
+struct TFrameCmd : public BaseCmdType
+{
+	FNetworkSimTime GetFrameDeltaTime() const { return FrameDeltaTime; }
+	void SetFrameDeltaTime(const FNetworkSimTime& InTime) { FrameDeltaTime = InTime; }
+	void NetSerialize(const FNetSerializeParams& P)
+	{
+		FrameDeltaTime.NetSerialize(P.Ar);
+		BaseCmdType::NetSerialize(P); 
+	}
+
+	void Log(FStandardLoggingParameters& P) const { BaseCmdType::Log(P); }
+
+private:
+	FNetworkSimTime FrameDeltaTime;
+};
+
+// Fixed tick specialization
+template<typename BaseCmdType, typename TickSettings>
+struct TFrameCmd<BaseCmdType, TickSettings, true> : public BaseCmdType
+{
+	FNetworkSimTime GetFrameDeltaTime() const { return FNetworkSimTime::FromMSec(TickSettings::GetFixedStepMS()); }
+	void SetFrameDeltaTime(const FNetworkSimTime& InTime) { }
+	void NetSerialize(const FNetSerializeParams& P) { BaseCmdType::NetSerialize(P); }
+	void Log(FStandardLoggingParameters& P) const { BaseCmdType::Log(P); }
+};
+
+// Helper to turn user supplied buffer types into the "real" buffer types: the InputCmd struct is wrapped in TFrameCmd
+template<typename TUserBufferTypes, typename TTickSettings>
+struct TInternalBufferTypes : TNetworkSimBufferTypes< 
+	
+	// InputCmds are wrapped in TFrameCmd, which will store an explicit sim delta time if we are not a fixed tick sim
+	TFrameCmd< typename TUserBufferTypes::TInputCmd , TTickSettings>,
+
+	typename TUserBufferTypes::TSyncState,	// SyncState passes through
+	typename TUserBufferTypes::TAuxState,	// Auxstate passes through
+	typename TUserBufferTypes::TDebugState	// Debugstate passes through
+>
+{
+};
+
+/** This is the "system driver", it has functions that the TNetworkedSimulationModel needs to call internally, that are specific to the types but not specific to the simulation itself. */
+template<typename TBufferTypes>
+class TNetworkedSimulationModelDriver
+{
+public:
+	using TInputCmd = typename TBufferTypes::TInputCmd;
+	using TSyncState = typename TBufferTypes::TSyncState;
+	using TAuxState= typename TBufferTypes::TAuxState;
+
+	// Debug string that can be used in internal warning/error logs
+	virtual FString GetDebugName() const = 0;
+
+	// Owning object for Visual Logs so that the system can emit them internally
+	virtual const AActor* GetVLogOwner() const = 0;
+
+	// Call to visual log the given states. Note that not all 3 will always be present and you should check for nullptrs.
+	virtual void VisualLog(const TInputCmd* Input, const TSyncState* Sync, const TAuxState* Aux, const FVisualLoggingParameters& SystemParameters) const = 0;
+	
+	// Called whenever the sim is ready to process new local input.
+	virtual void ProduceInput(const FNetworkSimTime SimTime, TInputCmd&) = 0;
+	
+	// Called from the Network Sim at the end of the sim frame when there is new sync data.
+	virtual void FinalizeFrame(const typename TBufferTypes::TSyncState& SyncState, const typename TBufferTypes::TAuxState& AuxState) = 0;
+};
+
+// ----------------------------------------------------------------------------------------------------------------------------------------------
+//	SimulationTick parameters. These are the data structures passed into the simulation code each frame
+// ----------------------------------------------------------------------------------------------------------------------------------------------
+
+struct TNetSimTimeStep
+{
+	// The delta time step for this tick (in MS by default)
+	const FNetworkSimTime& StepMS;
+	// The tick state of the simulation prior to running this tick. E.g, does not "include" the above StepMS that we are simulating now.
+	// The first time ::SimulationTick runs, TickState.GetTotalProcessedSimulationTime() == 0.
+	const FSimulationTickState& TickState;
+};
+
+// Input state: const references to the InputCmd/SyncState/AuxStates
+template<typename TBufferTypes>
+struct TNetSimInput
+{
+	const typename TBufferTypes::TInputCmd& Cmd;
+	const typename TBufferTypes::TSyncState& Sync;
+	const typename TBufferTypes::TAuxState& Aux;
+
+	TNetSimInput(const typename TBufferTypes::TInputCmd& InInputCmd, const typename TBufferTypes::TSyncState& InSync, const typename TBufferTypes::TAuxState& InAux)
+		: Cmd(InInputCmd), Sync(InSync), Aux(InAux) { }	
+	
+	// Allows implicit downcasting to a parent simulation class's input types
+	template<typename T>
+	TNetSimInput(const TNetSimInput<T>& Other)
+		: Cmd(Other.Cmd), Sync(Other.Sync), Aux(Other.Aux) { }
+};
+
+// Output state: the output SyncState (always created) and TNetSimLazyWriter for the AuxState (created on demand since every tick does not generate a new aux frame)
+template<typename TBufferTypes>
+struct TNetSimOutput
+{
+	typename TBufferTypes::TSyncState& Sync;
+	const TNetSimLazyWriter<typename TBufferTypes::TAuxState>& Aux;
+
+	TNetSimOutput(typename TBufferTypes::TSyncState& InSync, const TNetSimLazyWriter<typename TBufferTypes::TAuxState>& InAux)
+		: Sync(InSync), Aux(InAux) { }
+};
+
diff --git a/Engine/Plugins/Runtime/NetworkPredictionExtras/Source/NetworkPredictionExtras/Public/MockAbilitySimulation.h b/Engine/Plugins/Runtime/NetworkPredictionExtras/Source/NetworkPredictionExtras/Public/MockAbilitySimulation.h
--- a/Engine/Plugins/Runtime/NetworkPredictionExtras/Source/NetworkPredictionExtras/Public/MockAbilitySimulation.h
+++ b/Engine/Plugins/Runtime/NetworkPredictionExtras/Source/NetworkPredictionExtras/Public/MockAbilitySimulation.h
@@ -14,7 +14,7 @@
 #include "Misc/OutputDevice.h"
 #include "Misc/CoreDelegates.h"
 
-#include "NetworkSimulationModel.h"
+#include "NetworkedSimulationModel.h"
 #include "NetworkPredictionComponent.h"
 #include "Movement/FlyingMovement.h"
 
diff --git a/Engine/Plugins/Runtime/NetworkPredictionExtras/Source/NetworkPredictionExtras/Public/MockNetworkSimulation.h b/Engine/Plugins/Runtime/NetworkPredictionExtras/Source/NetworkPredictionExtras/Public/MockNetworkSimulation.h
--- a/Engine/Plugins/Runtime/NetworkPredictionExtras/Source/NetworkPredictionExtras/Public/MockNetworkSimulation.h
+++ b/Engine/Plugins/Runtime/NetworkPredictionExtras/Source/NetworkPredictionExtras/Public/MockNetworkSimulation.h
@@ -14,7 +14,7 @@
 #include "Misc/OutputDevice.h"
 #include "Misc/CoreDelegates.h"
 
-#include "NetworkSimulationModel.h"
+#include "NetworkedSimulationModel.h"
 #include "NetworkPredictionComponent.h"
 
 #include "MockNetworkSimulation.generated.h"
