NetworkPrediction plugin:
-Continue iteration on mock abilty. Have "state transition" events piped into blueprint and effects hooked up.
-Simulated proxies will now get latest InputCmd serialized to them so they can use it in sim extrapolation.
-Moved NetworkPredictionExtras content around and added a bunch of prototype effects

#rnx

diff --git a/Engine/Plugins/Runtime/NetworkPrediction/Source/NetworkPrediction/Private/Movement/FlyingMovementComponent.cpp b/Engine/Plugins/Runtime/NetworkPrediction/Source/NetworkPrediction/Private/Movement/FlyingMovementComponent.cpp
--- a/Engine/Plugins/Runtime/NetworkPrediction/Source/NetworkPrediction/Private/Movement/FlyingMovementComponent.cpp
+++ b/Engine/Plugins/Runtime/NetworkPrediction/Source/NetworkPrediction/Private/Movement/FlyingMovementComponent.cpp
@@ -148,7 +148,7 @@ void UFlyingMovementComponent::FinalizeFrame(const FlyingMovement::FMoveState& S
 
 FString UFlyingMovementComponent::GetDebugName() const
 {
-	return FString::Printf(TEXT("FlyingMovement. %s. %s"), *UEnum::GetValueAsString(TEXT("Engine.ENetRole"), GetOwnerRole()), *GetName());
+	return FString::Printf(TEXT("FlyingMovement. %s. %s"), *UEnum::GetValueAsString(TEXT("Engine.ENetRole"), GetOwnerRole()), *GetPathName());
 }
 
 const AActor* UFlyingMovementComponent::GetVLogOwner() const
diff --git a/Engine/Plugins/Runtime/NetworkPrediction/Source/NetworkPrediction/Public/NetworkedSimulationModelBuffer.h b/Engine/Plugins/Runtime/NetworkPrediction/Source/NetworkPrediction/Public/NetworkedSimulationModelBuffer.h
--- a/Engine/Plugins/Runtime/NetworkPrediction/Source/NetworkPrediction/Public/NetworkedSimulationModelBuffer.h
+++ b/Engine/Plugins/Runtime/NetworkPrediction/Source/NetworkPrediction/Public/NetworkedSimulationModelBuffer.h
@@ -91,7 +91,10 @@ protected:
 		{
 			ElementType* HeadElementPtr = ((T*)this)->HeadElement();
 			ElementType* NewElement = ((T*)this)->WriteFrame(Frame);
-			*NewElement = *HeadElementPtr;
+			if (HeadElementPtr)
+			{
+				*NewElement = *HeadElementPtr;
+			}
 			return NewElement;
 		}
 		else if (Frame < ((T*)this)->TailFrame())
diff --git a/Engine/Plugins/Runtime/NetworkPrediction/Source/NetworkPrediction/Public/NetworkedSimulationModelReplicators.h b/Engine/Plugins/Runtime/NetworkPrediction/Source/NetworkPrediction/Public/NetworkedSimulationModelReplicators.h
--- a/Engine/Plugins/Runtime/NetworkPrediction/Source/NetworkPrediction/Public/NetworkedSimulationModelReplicators.h
+++ b/Engine/Plugins/Runtime/NetworkPrediction/Source/NetworkPrediction/Public/NetworkedSimulationModelReplicators.h
@@ -232,9 +232,9 @@ struct TReplicator_Single : public TBase
 	}
 };
 
-/** Helper that writes a new input cmd to the input buffer, at given Frame (usually the sim's PendingFrame). If frame doesn't exist, ProduceInput is called on the driver. */
+/** Helper that writes a new input cmd to the input buffer, at given Frame (usually the sim's PendingFrame). If frame doesn't exist, ProduceInput is called on the driver if bProduceInputViaDriver=true, otherwise the input cmd will be initialized from the previous input cmd. */
 template<typename TDriver, typename TBufferTypes>
-void GenerateLocalInputCmdAtFrame(TDriver* Driver, TNetworkSimBufferContainer<TBufferTypes>& Buffers, const FNetworkSimTime& DeltaSimTime, int32 Frame)
+void GenerateLocalInputCmdAtFrame(TDriver* Driver, TNetworkSimBufferContainer<TBufferTypes>& Buffers, const FNetworkSimTime& DeltaSimTime, int32 Frame, bool bProduceInputViaDriver=true)
 {
 	using TInputCmd = typename TBufferTypes::TInputCmd;
 	if (TInputCmd* InputCmd = Buffers.Input[Frame])
@@ -243,23 +243,25 @@ void GenerateLocalInputCmdAtFrame(TDriver* Driver, TNetworkSimBufferContainer<TB
 	}
 	else
 	{
-		InputCmd = Buffers.Input.WriteFrame(Frame);
-		*InputCmd = TInputCmd();
+		InputCmd = Buffers.Input.WriteFrameInitializedFromHead(Frame);
 		InputCmd->SetFrameDeltaTime(DeltaSimTime);
-		Driver->ProduceInput(DeltaSimTime, *InputCmd);
+		if (bProduceInputViaDriver)
+		{
+			Driver->ProduceInput(DeltaSimTime, *InputCmd);
+		}
 	}
 }
 
 /** Helper to generate a local input cmd if we have simulation time to spend and advance the simulation's MaxAllowedFrame so that it can be processed. */
 template<typename TDriver, typename TBufferTypes, typename TTickSettings>
-void TryGenerateLocalInput(TDriver* Driver, TNetworkSimBufferContainer<TBufferTypes>& Buffers, TSimulationTicker<TTickSettings>& Ticker)
+void TryGenerateLocalInput(TDriver* Driver, TNetworkSimBufferContainer<TBufferTypes>& Buffers, TSimulationTicker<TTickSettings>& Ticker, bool bProduceInputViaDriver=true)
 {
 	using TInputCmd = typename TBufferTypes::TInputCmd;
 
 	FNetworkSimTime DeltaSimTime = Ticker.GetRemainingAllowedSimulationTime();
 	if (DeltaSimTime.IsPositive())
 	{
-		GenerateLocalInputCmdAtFrame(Driver, Buffers, DeltaSimTime, Ticker.PendingFrame);
+		GenerateLocalInputCmdAtFrame(Driver, Buffers, DeltaSimTime, Ticker.PendingFrame, bProduceInputViaDriver);
 		Ticker.MaxAllowedFrame = Ticker.PendingFrame;
 	}
 }
@@ -301,7 +303,9 @@ struct TReplicator_Server : public TBase
 	}
 };
 
-/** Simulated: "non locally controlled" simulations. We support "Simulation Extrapolation" here (using the sim to fake inputs to advance the sim)  */
+// Simulated: "non locally controlled" simulations. We support "Simulation Extrapolation" here (using the sim to fake inputs to advance the sim)
+//	-TODO: this is replicating Input/Sync/Aux which is required to do accurate Simulation Extrapolation. For interpolated simulated proxies, we can skip the input and possibly the aux.
+//	More settings/config options would be nice here
 template<typename TBufferTypes, typename TTickSettings, typename TBase=TReplicatorEmpty<TBufferTypes, TTickSettings>>
 struct TReplicator_Simulated : public TBase
 {
@@ -352,14 +356,22 @@ struct TReplicator_Simulated : public TBase
 		LastSerializedSimulationTime.NetSerialize(P.Ar);
 
 		// Serialize latest element
+		TInputCmd* InputCmd = nullptr;
 		TSyncState* SyncState = nullptr;
 		TAuxState* AuxState = nullptr;
 		
 		if (Ar.IsSaving())
 		{
+			const int32 Frame = Buffers.Sync.HeadFrame();
+
+			InputCmd = Buffers.Input.HeadElement();
+			if (!InputCmd)
+			{
+				InputCmd = &LastSerializedInputCmd;
+			}
 			SyncState = Buffers.Sync.HeadElement();
-			AuxState = Buffers.Aux.HeadElement();
-			check(SyncState && AuxState	); // We should not be here if the buffer is empty. Want to avoid serializing an "empty" flag at the top here.
+			AuxState = Buffers.Aux[Frame];
+			check(InputCmd && SyncState && AuxState	); // We should not be here if the buffer is empty. Want to avoid serializing an "empty" flag at the top here.
 		}
 		else
 		{
@@ -406,6 +418,7 @@ struct TReplicator_Simulated : public TBase
 			check(DestinationFrame != INDEX_NONE);
 
 			// "Finalize" our buffers and time keeping such that we serialize the latest state from the server in the right spot
+			InputCmd = Buffers.Input.WriteFrame(DestinationFrame);
 			SyncState = Buffers.Sync.WriteFrame(DestinationFrame);
 			AuxState = Buffers.Aux.WriteFrame(DestinationFrame);
 			Buffers.Input.WriteFrame(DestinationFrame);
@@ -424,11 +437,13 @@ struct TReplicator_Simulated : public TBase
 		}
 
 		check(SyncState && AuxState);
+		InputCmd->NetSerialize(Ar);
 		SyncState->NetSerialize(Ar);
 		AuxState->NetSerialize(Ar);
 
 		if (Ar.IsLoading())
 		{
+			LastSerializedInputCmd = *InputCmd;
 			LastSerializedSyncState = *SyncState;
 			LastSerializedAuxState = *AuxState;
 		}
@@ -476,11 +491,9 @@ struct TReplicator_Simulated : public TBase
 			{
 				Ticker.GiveSimulationTime(TickParameters.LocalDeltaTimeSeconds);
 			}
-
-			if (TickParameters.bGenerateLocalInputCmds)
-			{
-				TryGenerateLocalInput(Driver, Buffers, Ticker);
-			}
+			
+			// Generate local input if we are ready to tick. Note that we pass false so we won't call into the Driver to produce the input, we will use the last serialized InputCmd's values
+			TryGenerateLocalInput(Driver, Buffers, Ticker, false);
 		}
 	}
 
@@ -579,8 +592,11 @@ private:
 	
 	FNetworkSimTime ReconcileSimulationTime;
 	FNetworkSimTime LastSerializedSimulationTime;
+	
+	TInputCmd LastSerializedInputCmd;
 	TSyncState LastSerializedSyncState;
 	TAuxState LastSerializedAuxState;	// Temp? This should be conditional or optional. We want to support not replicating the aux state to simulated proxies
+
 };
 
 
diff --git a/Engine/Plugins/Runtime/NetworkPredictionExtras/Source/NetworkPredictionExtras/Private/MockAbilitySimulation.cpp b/Engine/Plugins/Runtime/NetworkPredictionExtras/Source/NetworkPredictionExtras/Private/MockAbilitySimulation.cpp
--- a/Engine/Plugins/Runtime/NetworkPredictionExtras/Source/NetworkPredictionExtras/Private/MockAbilitySimulation.cpp
+++ b/Engine/Plugins/Runtime/NetworkPredictionExtras/Source/NetworkPredictionExtras/Private/MockAbilitySimulation.cpp
@@ -261,6 +261,26 @@ void UMockFlyingAbilityComponent::ProduceInput(const FNetworkSimTime SimTime, FM
 void UMockFlyingAbilityComponent::FinalizeFrame(const FMockAbilitySyncState& SyncState, const FMockAbilityAuxstate& AuxState)
 {
 	Super::FinalizeFrame(SyncState, AuxState);
+
+	if (AuxState.bIsSprinting != bIsSprinting)
+	{
+		bIsSprinting = AuxState.bIsSprinting;
+		OnSprintStateChange.Broadcast(AuxState.bIsSprinting);
+	}
+
+	const bool bLocalIsDashing = (AuxState.DashTimeLeft > 0);
+	if (bLocalIsDashing != bIsDashing)
+	{
+		bIsDashing = bLocalIsDashing;
+		OnDashStateChange.Broadcast(bIsDashing);
+	}
+
+	const bool bLocalIsBlinking = (AuxState.BlinkWarmupLeft > 0);
+	if (bLocalIsBlinking != bIsBlinking)
+	{
+		bIsBlinking = bLocalIsBlinking;
+		OnBlinkStateChange.Broadcast(bLocalIsBlinking);
+	}
 }
 
 FString UMockFlyingAbilityComponent::GetDebugName() const
@@ -280,24 +300,24 @@ void UMockFlyingAbilityComponent::VisualLog(const FMockAbilityInputCmd* Input, c
 
 // ---------------------------------------------------------------------------------
 
-void UMockFlyingAbilityComponent::NotifySprint(bool bIsSprinting)
+void UMockFlyingAbilityComponent::NotifySprint(bool bNewIsSprinting)
 {
 	if (GetNetMode() == NM_DedicatedServer)
 	{
 		return;
 	}
 
-	UE_LOG(LogTemp, Display, TEXT("Sprint: %d"), bIsSprinting);
+	//UE_LOG(LogTemp, Display, TEXT("Sprint: %d"), bNewIsSprinting);
 }
 
-void UMockFlyingAbilityComponent::NotifyDash(bool bIsDashing)
+void UMockFlyingAbilityComponent::NotifyDash(bool bNewIsDashing)
 {
 	if (GetNetMode() == NM_DedicatedServer)
 	{
 		return;
 	}
 
-	UE_LOG(LogTemp, Display, TEXT("Dash: %d"), bIsDashing);
+	//UE_LOG(LogTemp, Display, TEXT("Dash: %d"), bNewIsDashing);
 }
 
 void UMockFlyingAbilityComponent::NotifyBlinkStartup()
@@ -307,7 +327,7 @@ void UMockFlyingAbilityComponent::NotifyBlinkStartup()
 		return;
 	}
 
-	UE_LOG(LogTemp, Display, TEXT("Blink Startup"));
+	//UE_LOG(LogTemp, Display, TEXT("Blink Startup"));
 }
 
 void UMockFlyingAbilityComponent::NotifyBlinkFinished()
@@ -317,5 +337,5 @@ void UMockFlyingAbilityComponent::NotifyBlinkFinished()
 		return;
 	}
 
-	UE_LOG(LogTemp, Display, TEXT("Blink Finished"));
-}
\ No newline at end of file
+	//UE_LOG(LogTemp, Display, TEXT("Blink Finished"));
+}
diff --git a/Engine/Plugins/Runtime/NetworkPredictionExtras/Source/NetworkPredictionExtras/Public/MockAbilitySimulation.h b/Engine/Plugins/Runtime/NetworkPredictionExtras/Source/NetworkPredictionExtras/Public/MockAbilitySimulation.h
--- a/Engine/Plugins/Runtime/NetworkPredictionExtras/Source/NetworkPredictionExtras/Public/MockAbilitySimulation.h
+++ b/Engine/Plugins/Runtime/NetworkPredictionExtras/Source/NetworkPredictionExtras/Public/MockAbilitySimulation.h
@@ -186,6 +186,52 @@ public:
 	virtual void NotifyBlinkStartup() override;
 	virtual void NotifyBlinkFinished() override;
 
+	// -------------------------------------------------------------------------------------
+	//	Ability State and Notifications
+	//		-This allows user code/blueprints to respond to state changes.
+	//		-These values always reflect the latest simulation state
+	//		-StateChange events are just that: when the state changes. They are not emitted from the sim themselves.
+	//			-This means they "work" for interpolated simulations and are resilient to packet loss and crazy network conditions
+	//			-That said, its "latest" only. There is NO guarantee that you will receive every state transition
+	//
+	// -------------------------------------------------------------------------------------
+
+	DECLARE_DYNAMIC_MULTICAST_DELEGATE_OneParam(FMockAbilityNotifyStateChange, bool, bNewStateValue);
+
+	// Notifies when Sprint state changes
+	UPROPERTY(BlueprintAssignable, Category="Mock AbilitySystem")
+	FMockAbilityNotifyStateChange OnSprintStateChange;
+
+	// Notifies when Dash state changes
+	UPROPERTY(BlueprintAssignable, Category="Mock AbilitySystem")
+	FMockAbilityNotifyStateChange OnDashStateChange;
+
+	// Notifies when Blink Changes
+	UPROPERTY(BlueprintAssignable, Category="Mock AbilitySystem")
+	FMockAbilityNotifyStateChange OnBlinkStateChange;
+
+	// Are we currently in the sprinting state
+	UFUNCTION(BlueprintCallable, Category="Mock AbilitySystem")
+	bool IsSprinting() const { return bIsSprinting; }
+
+	// Are we currently in the dashing state
+	UFUNCTION(BlueprintCallable, Category="Mock AbilitySystem")
+	bool IsDashing() const { return bIsDashing; }
+
+	// Are we currently in the blinking (startup) state
+	UFUNCTION(BlueprintCallable, Category="Mock AbilitySystem")
+	bool IsBlinking() const { return bIsBlinking; }
+
+private:
+	
+	// Local cached values for detecting state changes from the sim in ::FinalizeFrame
+	// Its tempting to think ::FinalizeFrame could pass in the previous frames values but this could
+	// not be reliable if buffer sizes are small and network conditions etc - you may not always know
+	// what was the "last finalized frame" or even have it in the buffers anymore.
+	bool bIsSprinting = false;
+	bool bIsDashing = false;
+	bool bIsBlinking = false;
+
 protected:
 
 	// Network Prediction
