NetworkPrediction plugin:

-"Game code api" cleanup. Added GetPendingStateRead()/GetPendingStateWrite() functions on TNetworkedSimulationModel. These compliment TNetworkSimStateAccessor usage.
-Fixed issue with TNetworkSimStateAccessor not being able to work on child simulations (E.g, FlyingMovement accessors could not access MockAbility simulation state even though types are comptabile)
-Cleaned up ::InstantiateNetworkSimulation and existing classes to get rid of copy/paste init functions

Renames (in the name of consistency)
-TNetworkSimStateAccessor -> renamed TNetSimStateAccessor
-INetworkSimlationModel -> INetworkedSimulationModel
-InstantiateNetworkSimulation -> InstantiateNetworkedSimulation

#rnx

diff --git a/Engine/Plugins/Runtime/NetworkPrediction/Source/NetworkPrediction/Private/Movement/FlyingMovementComponent.cpp b/Engine/Plugins/Runtime/NetworkPrediction/Source/NetworkPrediction/Private/Movement/FlyingMovementComponent.cpp
--- a/Engine/Plugins/Runtime/NetworkPrediction/Source/NetworkPrediction/Private/Movement/FlyingMovementComponent.cpp
+++ b/Engine/Plugins/Runtime/NetworkPrediction/Source/NetworkPrediction/Private/Movement/FlyingMovementComponent.cpp
@@ -38,6 +38,8 @@ static FAutoConsoleVariableRef CVarRequestMispredict(TEXT("fp.RequestMispredict"
 	RequestMispredict, TEXT("Causes a misprediction by inserting random value into stream on authority side"), ECVF_Default);
 }
 
+float UFlyingMovementComponent::GetDefaultMaxSpeed() { return FlyingMovementCVars::MaxSpeed; }
+
 // ----------------------------------------------------------------------------------------------------------
 //	UFlyingMovementComponent setup/init
 // ----------------------------------------------------------------------------------------------------------
@@ -51,29 +53,18 @@ UFlyingMovementComponent::UFlyingMovementComponent()
 //	Core Network Prediction functions
 // ----------------------------------------------------------------------------------------------------------
 
-INetworkSimulationModel* UFlyingMovementComponent::InstantiateNetworkSimulation()
+INetworkedSimulationModel* UFlyingMovementComponent::InstantiateNetworkedSimulation()
 {
-	check(UpdatedComponent);
-
-	// Create the simulation
-	MovementSimulation.Reset(new FlyingMovement::FMovementSimulation());
-	MovementSimulation->UpdatedComponent = UpdatedComponent;
-	MovementSimulation->UpdatedPrimitive = UpdatedPrimitive;
-	
-	// Initial states
+	// The Simulation
 	FlyingMovement::FMoveState InitialSyncState;
-	InitialSyncState.Location = UpdatedComponent->GetComponentLocation();
-	InitialSyncState.Rotation = UpdatedComponent->GetComponentQuat().Rotator();	
-
 	FlyingMovement::FAuxState InitialAuxState;
-	InitialAuxState.MaxSpeed = FlyingMovementCVars::MaxSpeed;
 
-	// The NetworkModel
-	auto NewModel = new FlyingMovement::FMovementSystem<0>(MovementSimulation.Get(), this, InitialSyncState, InitialAuxState);
+	InitFlyingMovementSimulation(new FlyingMovement::FMovementSimulation(), InitialSyncState, InitialAuxState);
 
-	// Accessors to the latest sync/aux state
-	MovementSyncState.Init(NewModel);
-	MovementAuxState.Init(NewModel);
+	// The Model
+	auto NewModel = new FlyingMovement::FMovementSystem<0>(MovementSimulation, this, InitialSyncState, InitialAuxState);
+	InitFlyingMovementNetSimModel(NewModel);
+	
 	return NewModel;
 }
 
diff --git a/Engine/Plugins/Runtime/NetworkPrediction/Source/NetworkPrediction/Private/Movement/ParametricMovement.cpp b/Engine/Plugins/Runtime/NetworkPrediction/Source/NetworkPrediction/Private/Movement/ParametricMovement.cpp
--- a/Engine/Plugins/Runtime/NetworkPrediction/Source/NetworkPrediction/Private/Movement/ParametricMovement.cpp
+++ b/Engine/Plugins/Runtime/NetworkPrediction/Source/NetworkPrediction/Private/Movement/ParametricMovement.cpp
@@ -85,29 +85,29 @@ UParametricMovementComponent::UParametricMovementComponent()
 
 }
 
-INetworkSimulationModel* UParametricMovementComponent::InstantiateNetworkSimulation()
+INetworkedSimulationModel* UParametricMovementComponent::InstantiateNetworkedSimulation()
 {
 	if (bDisableParametricMovementSimulation)
 	{
 		return nullptr;
 	}
 
-	ParametricMovementSimulation.Reset(new FParametricMovementSimulation());
-	ParametricMovementSimulation->Motion = &ParametricMotion;
+	// Simulation
+	ParametricMovement::FMoveState InitialSyncState;
+	ParametricMovement::FAuxState InitialAuxState;
 
+	InitParametricMovementSimulation(new FParametricMovementSimulation(), InitialSyncState, InitialAuxState);
+
+	// Model
 	if (ParametricMoverCVars::FixStep > 0)
 	{
-		auto *NewModel = new ParametricMovement::FMovementSystem<32>(ParametricMovementSimulation.Get(), this);
-		NewModel->RepProxy_Simulated.bAllowSimulatedExtrapolation = !bEnableInterpolation;
-		MovementSyncState.Init(NewModel);
-		MovementAuxState.Init(NewModel);
+		auto *NewModel = new ParametricMovement::FMovementSystem<32>( ParametricMovementSimulation, this);
+		InitParametricMovementNetSimModel(NewModel);
 		return NewModel;
 	}
 	
-	auto *NewModel = new ParametricMovement::FMovementSystem<>(ParametricMovementSimulation.Get(), this);
-	NewModel->RepProxy_Simulated.bAllowSimulatedExtrapolation = !bEnableInterpolation;
-	MovementSyncState.Init(NewModel);
-	MovementAuxState.Init(NewModel);
+	auto *NewModel = new ParametricMovement::FMovementSystem<>(ParametricMovementSimulation, this);
+	InitParametricMovementNetSimModel(NewModel);
 	return NewModel;
 }
 
@@ -160,7 +160,7 @@ void UParametricMovementComponent::TickComponent(float DeltaTime, enum ELevelTic
 {
 	Super::TickComponent(DeltaTime, TickType, ThisTickFunction);
 
-	if (bEnableDependentSimulation && NetworkSim.IsValid() && NetworkSim->GetParentSimulation() == nullptr)
+	if (bEnableDependentSimulation && NetSimModel.IsValid() && NetSimModel->GetParentSimulation() == nullptr)
 	{
 		// This is a very simply and generic way of finding a locally controlled network sim.
 		// The long term vision for dependent simulations is more game specific and dynamic. 
@@ -174,9 +174,9 @@ void UParametricMovementComponent::TickComponent(float DeltaTime, enum ELevelTic
 			APawn* Pawn = Player->GetPlayerController(World)->GetPawn();
 			if (UNetworkPredictionComponent* NetComponent = Pawn->FindComponentByClass<UNetworkPredictionComponent>())
 			{
-				if (INetworkSimulationModel* ParentNetSim = NetComponent->GetNetworkSimulation())
+				if (INetworkedSimulationModel* ParentNetSim = NetComponent->GetNetworkSimulation())
 				{
-					NetworkSim->SetParentSimulation(ParentNetSim);
+					NetSimModel->SetParentSimulation(ParentNetSim);
 				}
 			}
 		}
diff --git a/Engine/Plugins/Runtime/NetworkPrediction/Source/NetworkPrediction/Private/NetworkPredictionComponent.cpp b/Engine/Plugins/Runtime/NetworkPrediction/Source/NetworkPrediction/Private/NetworkPredictionComponent.cpp
--- a/Engine/Plugins/Runtime/NetworkPrediction/Source/NetworkPrediction/Private/NetworkPredictionComponent.cpp
+++ b/Engine/Plugins/Runtime/NetworkPrediction/Source/NetworkPrediction/Private/NetworkPredictionComponent.cpp
@@ -22,8 +22,8 @@ void UNetworkPredictionComponent::InitializeComponent()
 	{
 
 		// Child class will instantiate
-		INetworkSimulationModel* NewNetworkSim = InstantiateNetworkSimulation();
-		NetworkSim.Reset(NewNetworkSim);
+		INetworkedSimulationModel* NewNetworkSim = InstantiateNetworkedSimulation();
+		NetSimModel.Reset(NewNetworkSim);
 
 		if (NewNetworkSim == nullptr)
 		{
@@ -50,18 +50,18 @@ void UNetworkPredictionComponent::UninitializeComponent()
 {
 	Super::UninitializeComponent();
 
-	if (NetworkSim.IsValid())
+	if (NetSimModel.IsValid())
 	{
 		UNetworkSimulationGlobalManager* NetworkSimGlobalManager = GetWorld()->GetSubsystem<UNetworkSimulationGlobalManager>();
 		if (ensure(NetworkSimGlobalManager))
 		{
-			NetworkSimGlobalManager->UnregisterModel(NetworkSim.Get(), GetOwner());
+			NetworkSimGlobalManager->UnregisterModel(NetSimModel.Get(), GetOwner());
 			if (ServerRPCHandle.IsValid())
 			{
 				NetworkSimGlobalManager->TickServerRPCDelegate.Remove(ServerRPCHandle);
 			}
 
-			NetworkSim.Reset(nullptr);
+			NetSimModel.Reset(nullptr);
 		}
 	}
 }
@@ -70,7 +70,7 @@ void UNetworkPredictionComponent::PreReplication(IRepChangedPropertyTracker & Ch
 {
 	Super::PreReplication(ChangedPropertyTracker);
 
-	if (NetworkSim.IsValid())
+	if (NetSimModel.IsValid())
 	{
 		// We have to update our replication proxies so they can be accurately compared against client shadowstate during property replication. ServerRPC proxy does not need to do this.
 		ReplicationProxy_Autonomous.OnPreReplication();
@@ -103,9 +103,9 @@ void UNetworkPredictionComponent::GetLifetimeReplicatedProps(TArray< FLifetimePr
 
 void UNetworkPredictionComponent::InitializeForNetworkRole(ENetRole Role)
 {
-	if (NetworkSim.IsValid())
+	if (NetSimModel.IsValid())
 	{
-		NetworkSim->InitializeForNetworkRole(Role, GetSimulationInitParameters(Role));
+		NetSimModel->InitializeForNetworkRole(Role, GetSimulationInitParameters(Role));
 	}
 }
 
@@ -179,8 +179,8 @@ void UNetworkPredictionComponent::ServerRecieveClientInput_Implementation(const
 
 void UNetworkPredictionComponent::TickServerRPC(float DeltaSeconds)
 {
-	check(NetworkSim.IsValid());
-	if (NetworkSim->ShouldSendServerRPC(DeltaSeconds))
+	check(NetSimModel.IsValid());
+	if (NetSimModel->ShouldSendServerRPC(DeltaSeconds))
 	{
 		// Temp hack to make sure the ServerRPC doesn't get suppressed from bandwidth limiting
 		// (system hasn't been optimized and not mature enough yet to handle gaps in input stream)
diff --git a/Engine/Plugins/Runtime/NetworkPrediction/Source/NetworkPrediction/Private/NetworkedSimulationGlobalManager.cpp b/Engine/Plugins/Runtime/NetworkPrediction/Source/NetworkPrediction/Private/NetworkedSimulationGlobalManager.cpp
--- a/Engine/Plugins/Runtime/NetworkPrediction/Source/NetworkPrediction/Private/NetworkedSimulationGlobalManager.cpp
+++ b/Engine/Plugins/Runtime/NetworkPrediction/Source/NetworkPrediction/Private/NetworkedSimulationGlobalManager.cpp
@@ -61,12 +61,12 @@ void UNetworkSimulationGlobalManager::BeginNewSimulationFrame(UWorld* InWorld, E
 	TickServerRPCDelegate.Broadcast(InDeltaSeconds);
 }
 
-void UNetworkSimulationGlobalManager::RegisterModel(INetworkSimulationModel* Model, AActor* OwningActor)
+void UNetworkSimulationGlobalManager::RegisterModel(INetworkedSimulationModel* Model, AActor* OwningActor)
 {
 	SimulationGroupMap.FindOrAdd(Model->GetSimulationGroupName()).Simulations.Emplace(Model, OwningActor);
 }
 
-void UNetworkSimulationGlobalManager::UnregisterModel(INetworkSimulationModel* Model, AActor* OwningActor)
+void UNetworkSimulationGlobalManager::UnregisterModel(INetworkedSimulationModel* Model, AActor* OwningActor)
 {
 	auto& SimulationList = SimulationGroupMap.FindOrAdd(Model->GetSimulationGroupName()).Simulations;
 	int32 idx = SimulationList.IndexOfByKey(Model);
diff --git a/Engine/Plugins/Runtime/NetworkPrediction/Source/NetworkPrediction/Public/Movement/BaseMovementComponent.h b/Engine/Plugins/Runtime/NetworkPrediction/Source/NetworkPrediction/Public/Movement/BaseMovementComponent.h
--- a/Engine/Plugins/Runtime/NetworkPrediction/Source/NetworkPrediction/Public/Movement/BaseMovementComponent.h
+++ b/Engine/Plugins/Runtime/NetworkPrediction/Source/NetworkPrediction/Public/Movement/BaseMovementComponent.h
@@ -34,7 +34,7 @@ public:
 
 // -------------------------------------------------------------------------------------------------------------------------------
 //	Base component for movement. This essentially has the generic glue for selecting an UpdatedComponent and moving it along the world
-//	It is abstract in that you still need to define which simulation the component runs (via ::InstantiateNetworkSimulation)
+//	It is abstract in that you still need to define which simulation the component runs (via ::InstantiateNetworkedSimulation)
 // -------------------------------------------------------------------------------------------------------------------------------
 UCLASS(Abstract)
 class NETWORKPREDICTION_API UBaseMovementComponent : public UNetworkPredictionComponent
diff --git a/Engine/Plugins/Runtime/NetworkPrediction/Source/NetworkPrediction/Public/Movement/FlyingMovement.h b/Engine/Plugins/Runtime/NetworkPrediction/Source/NetworkPrediction/Public/Movement/FlyingMovement.h
--- a/Engine/Plugins/Runtime/NetworkPrediction/Source/NetworkPrediction/Public/Movement/FlyingMovement.h
+++ b/Engine/Plugins/Runtime/NetworkPrediction/Source/NetworkPrediction/Public/Movement/FlyingMovement.h
@@ -198,8 +198,8 @@ public:
 	DECLARE_DELEGATE_TwoParams(FProduceFlyingInput, const FNetworkSimTime /*SimTime*/, FlyingMovement::FInputCmd& /*Cmd*/)
 	FProduceFlyingInput ProduceInputDelegate;
 
-	TNetworkSimStateAccessor<FlyingMovement::FMoveState> MovementSyncState;
-	TNetworkSimStateAccessor<FlyingMovement::FAuxState> MovementAuxState;
+	TNetSimStateAccessor<FlyingMovement::FMoveState> MovementSyncState;
+	TNetSimStateAccessor<FlyingMovement::FAuxState> MovementAuxState;
 
 	// IFlyingMovementSystemDriver
 	FString GetDebugName() const override;
@@ -212,6 +212,33 @@ public:
 protected:
 
 	// Network Prediction
-	virtual INetworkSimulationModel* InstantiateNetworkSimulation() override;
-	TUniquePtr<FlyingMovement::FMovementSimulation> MovementSimulation;
+	virtual INetworkedSimulationModel* InstantiateNetworkedSimulation() override;
+	FlyingMovement::FMovementSimulation* MovementSimulation = nullptr;
+
+	// Init function. This is broken up from ::InstantiateNetworkedSimulation and templated so that subclasses can share the init code
+	template<typename TSimulation>
+	void InitFlyingMovementSimulation(TSimulation* Simulation, FlyingMovement::FMoveState& InitialSyncState, FlyingMovement::FAuxState& InitialAuxState)
+	{
+		check(UpdatedComponent);
+		check(MovementSimulation == nullptr); // Reinstantiation not supported
+		MovementSimulation = Simulation;
+		
+
+		Simulation->UpdatedComponent = UpdatedComponent;
+		Simulation->UpdatedPrimitive = UpdatedPrimitive;
+
+		InitialSyncState.Location = UpdatedComponent->GetComponentLocation();
+		InitialSyncState.Rotation = UpdatedComponent->GetComponentQuat().Rotator();	
+
+		InitialAuxState.MaxSpeed = GetDefaultMaxSpeed();
+	}
+
+	template<typename TNetSimModel>
+	void InitFlyingMovementNetSimModel(TNetSimModel* Model)
+	{
+		MovementSyncState.Bind(Model);
+		MovementAuxState.Bind(Model);
+	}
+
+	static float GetDefaultMaxSpeed();
 };
\ No newline at end of file
diff --git a/Engine/Plugins/Runtime/NetworkPrediction/Source/NetworkPrediction/Public/Movement/ParametricMovement.h b/Engine/Plugins/Runtime/NetworkPrediction/Source/NetworkPrediction/Public/Movement/ParametricMovement.h
--- a/Engine/Plugins/Runtime/NetworkPrediction/Source/NetworkPrediction/Public/Movement/ParametricMovement.h
+++ b/Engine/Plugins/Runtime/NetworkPrediction/Source/NetworkPrediction/Public/Movement/ParametricMovement.h
@@ -183,13 +183,29 @@ class NETWORKPREDICTION_API UParametricMovementComponent : public UBaseMovementC
 
 protected:
 
-	TNetworkSimStateAccessor<ParametricMovement::FMoveState> MovementSyncState;
-	TNetworkSimStateAccessor<ParametricMovement::FAuxState> MovementAuxState;
+	TNetSimStateAccessor<ParametricMovement::FMoveState> MovementSyncState;
+	TNetSimStateAccessor<ParametricMovement::FAuxState> MovementAuxState;
 
-	virtual INetworkSimulationModel* InstantiateNetworkSimulation() override;
+	virtual INetworkedSimulationModel* InstantiateNetworkedSimulation() override;
 	FNetworkSimulationModelInitParameters GetSimulationInitParameters(ENetRole Role) override;
 
-	TUniquePtr<FParametricMovementSimulation> ParametricMovementSimulation;
+	FParametricMovementSimulation* ParametricMovementSimulation = nullptr;
+
+	template<typename TSimulation>
+	void InitParametricMovementSimulation(TSimulation* Simulation, ParametricMovement::FMoveState& InitialSyncState, ParametricMovement::FAuxState& InitialAuxState)
+	{
+		check(ParametricMovementSimulation == nullptr);
+		ParametricMovementSimulation = Simulation;
+		ParametricMovementSimulation->Motion = &ParametricMotion;
+	}
+
+	template<typename TNetSimModel>
+	void InitParametricMovementNetSimModel(TNetSimModel* Model)
+	{
+		Model->RepProxy_Simulated.bAllowSimulatedExtrapolation = !bEnableInterpolation;
+		MovementSyncState.Bind(Model);
+		MovementAuxState.Bind(Model);
+	}
 
 	// ------------------------------------------------------------------------
 	// Temp Parametric movement example
diff --git a/Engine/Plugins/Runtime/NetworkPrediction/Source/NetworkPrediction/Public/NetworkPredictionComponent.h b/Engine/Plugins/Runtime/NetworkPrediction/Source/NetworkPrediction/Public/NetworkPredictionComponent.h
--- a/Engine/Plugins/Runtime/NetworkPrediction/Source/NetworkPrediction/Public/NetworkPredictionComponent.h
+++ b/Engine/Plugins/Runtime/NetworkPrediction/Source/NetworkPrediction/Public/NetworkPredictionComponent.h
@@ -11,7 +11,7 @@
 //	This is the base component for running a TNetworkedSimulationModel through an actor component. This contains the boiler plate hooks into getting the system
 //	initialized and plugged into the UE4 replication system.
 //
-//	This is an abstract component and cannot function on its own. It must be subclassed and InstantiateNetworkSimulation must be implemented.
+//	This is an abstract component and cannot function on its own. It must be subclassed and InstantiateNetworkedSimulation must be implemented.
 //	Ticking and RPC sending will be handled automatically, but the subclass will need to supply input to the simulation (via TNetworkedSimulationModelDriver::ProduceInput).
 //
 //	Its also worth pointing out that nothing about being a UActorComponent is essential here. All that this component does could be done within an AActor itself.
@@ -34,12 +34,12 @@ public:
 	virtual void Reconcile() { }
 	virtual void TickSimulation(float DeltaTimeSeconds) { }
 
-	INetworkSimulationModel* GetNetworkSimulation() const { return NetworkSim.Get(); }
+	INetworkedSimulationModel* GetNetworkSimulation() const { return NetSimModel.Get(); }
 
 protected:
 	
 	// Classes must instantiate their own NetworkSim here. The UNetworkPredictionComponent will manage its lifetime
-	virtual INetworkSimulationModel* InstantiateNetworkSimulation() { check(false); return nullptr; }
+	virtual INetworkedSimulationModel* InstantiateNetworkedSimulation() { check(false); return nullptr; }
 
 	// Finalizes initialization when NetworkRole changes. Does not need to be overridden.
 	virtual void InitializeForNetworkRole(ENetRole Role);
@@ -87,5 +87,5 @@ protected:
 	FDelegateHandle ServerRPCHandle;
 
 	// The Network sim that this component is managing. This is what is doing all the work.
-	TUniquePtr<INetworkSimulationModel> NetworkSim;
+	TUniquePtr<INetworkedSimulationModel> NetSimModel;
 };
\ No newline at end of file
diff --git a/Engine/Plugins/Runtime/NetworkPrediction/Source/NetworkPrediction/Public/NetworkPredictionTypes.h b/Engine/Plugins/Runtime/NetworkPrediction/Source/NetworkPrediction/Public/NetworkPredictionTypes.h
--- a/Engine/Plugins/Runtime/NetworkPrediction/Source/NetworkPrediction/Public/NetworkPredictionTypes.h
+++ b/Engine/Plugins/Runtime/NetworkPrediction/Source/NetworkPrediction/Public/NetworkPredictionTypes.h
@@ -221,7 +221,7 @@ public:
 	virtual int32 GetProxyDirtyCount(EReplicationProxyTarget Target) = 0;
 };
 
-class INetworkSimulationModel : public IReplicationProxy
+class INetworkedSimulationModel : public IReplicationProxy
 {
 public:
 
@@ -239,11 +239,11 @@ public:
 	// ----------------------------------------------------------------------
 	
 	// Main function to call on simulated proxy sim
-	virtual void SetParentSimulation(INetworkSimulationModel* Simulation) = 0;
-	virtual INetworkSimulationModel* GetParentSimulation() const = 0;
+	virtual void SetParentSimulation(INetworkedSimulationModel* Simulation) = 0;
+	virtual INetworkedSimulationModel* GetParentSimulation() const = 0;
 	
-	virtual void AddDependentSimulation(INetworkSimulationModel* Simulation) = 0;
-	virtual void RemoveDependentSimulation(INetworkSimulationModel* Simulation) = 0;
+	virtual void AddDependentSimulation(INetworkedSimulationModel* Simulation) = 0;
+	virtual void RemoveDependentSimulation(INetworkedSimulationModel* Simulation) = 0;
 	
 	// Tell parent sim that a dependent sim needs to reconcile (parent sim drives this)
 	virtual void NotifyDependentSimNeedsReconcile() = 0;
diff --git a/Engine/Plugins/Runtime/NetworkPrediction/Source/NetworkPrediction/Public/NetworkedSimulationGlobalManager.h b/Engine/Plugins/Runtime/NetworkPrediction/Source/NetworkPrediction/Public/NetworkedSimulationGlobalManager.h
--- a/Engine/Plugins/Runtime/NetworkPrediction/Source/NetworkPrediction/Public/NetworkedSimulationGlobalManager.h
+++ b/Engine/Plugins/Runtime/NetworkPrediction/Source/NetworkPrediction/Public/NetworkedSimulationGlobalManager.h
@@ -52,8 +52,8 @@ public:
 	virtual void Deinitialize() override;
 
 	// Netsim Register/Unregister
-	void RegisterModel(INetworkSimulationModel* Model, AActor* OwningActor);
-	void UnregisterModel(INetworkSimulationModel* Model, AActor* OwningActor);
+	void RegisterModel(INetworkedSimulationModel* Model, AActor* OwningActor);
+	void UnregisterModel(INetworkedSimulationModel* Model, AActor* OwningActor);
 
 	// Delegate for auto proxies to register with to send their server RPC
 	DECLARE_MULTICAST_DELEGATE_OneParam(FTickServerRPC, float /*RealTime DeltaSeconds*/)
@@ -70,10 +70,10 @@ private:
 	{
 		struct FItem
 		{
-			FItem(INetworkSimulationModel* InSim, AActor* InActor) : Simulation(InSim), OwningActor(InActor) { }
-			bool operator == (const INetworkSimulationModel* InSim) const { return Simulation == InSim; }
+			FItem(INetworkedSimulationModel* InSim, AActor* InActor) : Simulation(InSim), OwningActor(InActor) { }
+			bool operator == (const INetworkedSimulationModel* InSim) const { return Simulation == InSim; }
 
-			INetworkSimulationModel* Simulation;
+			INetworkedSimulationModel* Simulation;
 			AActor* OwningActor;
 		};
 
diff --git a/Engine/Plugins/Runtime/NetworkPrediction/Source/NetworkPrediction/Public/NetworkedSimulationModel.h b/Engine/Plugins/Runtime/NetworkPrediction/Source/NetworkPrediction/Public/NetworkedSimulationModel.h
--- a/Engine/Plugins/Runtime/NetworkPrediction/Source/NetworkPrediction/Public/NetworkedSimulationModel.h
+++ b/Engine/Plugins/Runtime/NetworkPrediction/Source/NetworkPrediction/Public/NetworkedSimulationModel.h
@@ -24,7 +24,14 @@
 //		-UNetworkSimulationGlobalManager: responsible for ticking simulation (after recv net traffic, prior to UE4 actor ticking)
 //		-External game code can interact with the system:
 //			-The TNetworkedSimulationModel is mostly public and exposed. It is not recommend to publicly expose to your "user" code (high level scripting, designers, etc).
-//			-TNetworkSimStateAccessor is a helper for safely reading/writing to state within the system.
+//			-TNetkSimStateAccessor is a helper for safely reading/writing to state within the system.
+//
+//	* Notes on ownership and lifetime:
+//		-User code will instantiate the underlying simulation class (TSimulation)
+//		-Ownership of the simulation is taken over by TNetworkedSimulationModel on instantiation
+//		-(user code may still want to cache a pointer to the TSimulation instance for its own, but does not need to destroy the simulation)
+//		-Swapping TSimulation instance or releasing it early would be possible but use case is not clear so not implemented.
+//		-TNetworkedSimulationModel lifetime is the responsibility of the code that created it (UNetworkPredictionComponent in most cases).
 //
 // ----------------------------------------------------------------------------------------------------------------------------------------------------------------------
 
@@ -42,7 +49,7 @@ template <
 	typename TRepProxyReplay =		TReplicator_Sequence	<TInternalBufferTypes<TUserBufferTypes, InTTickSettings>,	InTTickSettings, ENetworkSimBufferTypeId::Sync,  3>,
 	typename TRepProxyDebug =		TReplicator_Debug		<TInternalBufferTypes<TUserBufferTypes, InTTickSettings>,	InTTickSettings>
 >
-class TNetworkedSimulationModel : public INetworkSimulationModel
+class TNetworkedSimulationModel : public INetworkedSimulationModel
 {
 public:
 
@@ -60,10 +67,11 @@ public:
 	using TSimTime = FNetworkSimTime;
 	using TRealTime = FNetworkSimTime::FRealTime;
 	
+	// Note: Ownership of TSimulation* is taken by this constructor
 	TNetworkedSimulationModel(TSimulation* InSimulation, TDriver* InDriver, const TSyncState& InitialSyncState = TSyncState(), const TAuxState& InitialAuxState = TAuxState())
+		: Simulation(InSimulation)
 	{
-		check(InSimulation && InDriver);
-		Simulation = InSimulation;
+		check(Simulation.IsValid() && InDriver);
 		Driver = InDriver;
 		*Buffers.Sync.WriteFrame(0) = InitialSyncState;
 		*Buffers.Aux.WriteFrame(0) = InitialAuxState;
@@ -238,15 +246,15 @@ public:
 		switch (Role)
 		{
 			case ROLE_Authority:
-				RepProxy_ServerRPC.template Reconcile<TSimulation, TDriver>(Simulation, Driver, Buffers, Ticker);
+				RepProxy_ServerRPC.template Reconcile<TSimulation, TDriver>(Simulation.Get(), Driver, Buffers, Ticker);
 			break;
 
 			case ROLE_AutonomousProxy:
-				RepProxy_Autonomous.template Reconcile<TSimulation, TDriver>(Simulation, Driver, Buffers, Ticker);
+				RepProxy_Autonomous.template Reconcile<TSimulation, TDriver>(Simulation.Get(), Driver, Buffers, Ticker);
 			break;
 
 			case ROLE_SimulatedProxy:
-				RepProxy_Simulated.template Reconcile<TSimulation, TDriver>(Simulation, Driver, Buffers, Ticker);
+				RepProxy_Simulated.template Reconcile<TSimulation, TDriver>(Simulation.Get(), Driver, Buffers, Ticker);
 			break;
 		}
 	}
@@ -327,9 +335,51 @@ public:
 		return RepProxy_Simulated.GetSimulatedUpdateMode();
 	}
 
+	// ------------------------------------------------------------------------------------------------------
+	// State Accessors. This allows game code to read or conditionally write to the pending frame values.
+	// This mirrors what TNetSimStateAccessor do, but is callable directly on the NetworkedSimulationModel.
+	// Accessor conditionally gives access to the current (pending) Sync/Aux state to outside code.
+	// Reads are always allowed, Writes are conditional and null will sometimes be returned!
+	//
+	// Authority can always write to the pending frame. Non authority requires the netsim to be currently processing an ::SimulationTick.
+	// If you aren't inside an ::SimulationTick call, it is really not safe to predict state changes. It is safest and simplest to just not predict these changes.
+	//
+	// Explanation: During the scope of an ::SimulationTick call, we know exactly 'when' we are relative to what the server is processing. If the predicting client wants
+	// to predict a change to sync/aux state during an update, the server will do it at the exact same time (assuming not a mis prediction). When a state change
+	// happens "out of band" (outside an ::SimulationTick call) - we really have no way to correlate when the server will do it. While its tempting to think "we will get
+	// a correction anyways, might as well guess at it and maybe get a smaller correction" - but this opens us up to other problems. The server may actually not 
+	// change the state at all and you may not get an update that corrects you. You could add a timeout and track the state change somewhere but that really complicates
+	// things and could leave you open to "double" problems: if the state change is additive, you may stack the authority change on top of the local predicted change, or
+	// you may roll back the predicted change to then later receive the authority change.
+	//	
+	// What still may make sense to do is allow the "In Update" bool to be temporarily disabled if we enter code that we know is not rollback friendly.
+	// ------------------------------------------------------------------------------------------------------
+
+	template<typename TState>
+	const TState* GetPendingStateRead() const
+	{
+		// Reading is always allowed
+		auto& Buffer = NetSimBufferSelect::Get<TNetworkSimBufferContainer<TBufferTypes>, TState>(Buffers);
+		return Buffer[Ticker.PendingFrame];
+	}
+
+	template<typename TState>
+	TState* GetPendingStateWrite(bool bHasAuthority)
+	{
+		// Writes are always allowed on the authority or during a simulation update for predicting clients
+		if (bHasAuthority || Ticker.bUpdateInProgress)
+		{
+			auto& Buffer = NetSimBufferSelect::Get<TNetworkSimBufferContainer<TBufferTypes>, TState>(Buffers);
+			return Buffer.WriteFrameInitializedFromHead(Ticker.PendingFrame);
+		}
+		return nullptr;
+	}
+
+	// ------------------------------------------------------------------------------------------------------
+	//	Dependent Simulations
 	// ------------------------------------------------------------------------------------------------------
 
-	void SetParentSimulation(INetworkSimulationModel* ParentSimulation) final override
+	void SetParentSimulation(INetworkedSimulationModel* ParentSimulation) final override
 	{
 		if (RepProxy_Simulated.ParentSimulation)
 		{
@@ -343,19 +393,19 @@ public:
 		}
 	}
 
-	INetworkSimulationModel* GetParentSimulation() const final override
+	INetworkedSimulationModel* GetParentSimulation() const final override
 	{
 		return RepProxy_Simulated.ParentSimulation;
 	}
 
-	void AddDependentSimulation(INetworkSimulationModel* DependentSimulation) final override
+	void AddDependentSimulation(INetworkedSimulationModel* DependentSimulation) final override
 	{
 		check(RepProxy_Autonomous.DependentSimulations.Contains(DependentSimulation) == false);
 		RepProxy_Autonomous.DependentSimulations.Add(DependentSimulation);
 		NotifyDependentSimNeedsReconcile(); // force reconcile on purpose
 	}
 
-	void RemoveDependentSimulation(INetworkSimulationModel* DependentSimulation) final override
+	void RemoveDependentSimulation(INetworkedSimulationModel* DependentSimulation) final override
 	{
 		RepProxy_Autonomous.DependentSimulations.Remove(DependentSimulation);
 	}
@@ -367,18 +417,18 @@ public:
 
 	void BeginRollback(const FNetworkSimTime& RollbackDeltaTime, const int32 ParentFrame) final override
 	{
-		RepProxy_Simulated.template DependentRollbackBegin<TSimulation, TDriver>(Simulation, Driver, Buffers, Ticker, RollbackDeltaTime, ParentFrame);
+		RepProxy_Simulated.template DependentRollbackBegin<TSimulation, TDriver>(Simulation.Get(), Driver, Buffers, Ticker, RollbackDeltaTime, ParentFrame);
 	}
 
 	void StepRollback(const FNetworkSimTime& Step, const int32 ParentFrame, const bool bFinalStep) final override
 	{
-		RepProxy_Simulated.template DependentRollbackStep<TSimulation, TDriver>(Simulation, Driver, Buffers, Ticker, Step, ParentFrame, bFinalStep);
+		RepProxy_Simulated.template DependentRollbackStep<TSimulation, TDriver>(Simulation.Get(), Driver, Buffers, Ticker, Step, ParentFrame, bFinalStep);
 	}
 
 	void ClearAllDependentSimulations()
 	{
-		TArray<INetworkSimulationModel*> LocalList = MoveTemp(RepProxy_Autonomous.DependentSimulations);
-		for (INetworkSimulationModel* DependentSim : LocalList)
+		TArray<INetworkedSimulationModel*> LocalList = MoveTemp(RepProxy_Autonomous.DependentSimulations);
+		for (INetworkedSimulationModel* DependentSim : LocalList)
 		{
 			DependentSim->SetParentSimulation(nullptr);
 		}
@@ -386,7 +436,7 @@ public:
 
 	// -------------------------------------------------------------------------------------------------------
 
-	TSimulation* Simulation = nullptr;
+	TUniquePtr<TSimulation> Simulation;
 	TDriver* Driver = nullptr;
 
 	TSimulationTicker<TTickSettings> Ticker;
diff --git a/Engine/Plugins/Runtime/NetworkPrediction/Source/NetworkPrediction/Public/NetworkedSimulationModelReplicators.h b/Engine/Plugins/Runtime/NetworkPrediction/Source/NetworkPrediction/Public/NetworkedSimulationModelReplicators.h
--- a/Engine/Plugins/Runtime/NetworkPrediction/Source/NetworkPrediction/Public/NetworkedSimulationModelReplicators.h
+++ b/Engine/Plugins/Runtime/NetworkPrediction/Source/NetworkPrediction/Public/NetworkedSimulationModelReplicators.h
@@ -314,7 +314,7 @@ struct TReplicator_Simulated : public TBase
 	using TAuxState = typename TBufferTypes::TAuxState;
 
 	// Parent Simulation. If this is set, this simulation will forward predict in sync with this parent sim. The parent sim should be an autonomous proxy driven simulation
-	INetworkSimulationModel* ParentSimulation = nullptr;
+	INetworkedSimulationModel* ParentSimulation = nullptr;
 
 	// Instance flag for enabling simulated extrapolation
 	bool bAllowSimulatedExtrapolation = true;
@@ -612,7 +612,7 @@ struct TReplicator_Autonomous : public TBase
 	bool IsReconcileFaultDetected() const { return bReconcileFaultDetected; }
 	const FNetworkSimTime& GetLastSerializedSimTime() const { return SerializedTime; }
 
-	TArray<INetworkSimulationModel*> DependentSimulations;
+	TArray<INetworkedSimulationModel*> DependentSimulations;
 	bool bDependentSimulationNeedsReconcile = false;
 
 	int32 GetProxyDirtyCount(TNetworkSimBufferContainer<TBufferTypes>& Buffers) const
@@ -735,7 +735,7 @@ struct TReplicator_Autonomous : public TBase
 		}
 
 		// Tell dependent simulations to rollback
-		for (INetworkSimulationModel* DependentSim : DependentSimulations)
+		for (INetworkedSimulationModel* DependentSim : DependentSimulations)
 		{
 			DependentSim->BeginRollback(RollbackDeltaTime, SerializedFrame);
 		}
@@ -781,7 +781,7 @@ struct TReplicator_Autonomous : public TBase
 			}
 
 			// Tell dependent simulations to advance
-			for (INetworkSimulationModel* DependentSim : DependentSimulations)
+			for (INetworkedSimulationModel* DependentSim : DependentSimulations)
 			{
 				DependentSim->StepRollback(ResimulateCmd->GetFrameDeltaTime(), Frame, (Frame == LastFrameToProcess));
 			}
diff --git a/Engine/Plugins/Runtime/NetworkPrediction/Source/NetworkPrediction/Public/NetworkedSimulationModelTypes.h b/Engine/Plugins/Runtime/NetworkPrediction/Source/NetworkPrediction/Public/NetworkedSimulationModelTypes.h
--- a/Engine/Plugins/Runtime/NetworkPrediction/Source/NetworkPrediction/Public/NetworkedSimulationModelTypes.h
+++ b/Engine/Plugins/Runtime/NetworkPrediction/Source/NetworkPrediction/Public/NetworkedSimulationModelTypes.h
@@ -60,19 +60,6 @@ struct TSelectTypeHelper<TBufferTypes, ENetworkSimBufferTypeId::Debug>
 	using type = typename TBufferTypes::TDebugState;
 };
 
-enum ENetworkSimBufferAllocationType
-{
-	Contiguous,
-	Sparse
-};
-
-template<ENetworkSimBufferAllocationType InAllocationType, int32 InSize>
-struct TNetworkSimBufferAllocation
-{
-	enum { Size = InSize };
-	static constexpr ENetworkSimBufferAllocationType AllocationType() { return InAllocationType; }
-};
-
 template<typename T>
 struct HasNetSerialize
 {
@@ -169,55 +156,37 @@ struct TBufferGetterHelper<TContainer, ENetworkSimBufferTypeId::Debug>
 	}
 };
 
-// -------------------------------------------------
-
-// Helper struct for accessing a netsim buffer given the underlying type
-template<typename TContainer, typename TState>
-struct TBufferGetterHelper_ByState
-{
-	static void GetBuffer(TContainer& Container)
-	{
-		static_assert(!Container, "Failed to find specialized Get for your BufferId");
-	}
-};
-
-template<typename TContainer>
-struct TBufferGetterHelper_ByState<TContainer, typename TContainer::TInputCmd>
+// Helper for accessing a netsim buffer given the underlying type
+namespace NetSimBufferSelect
 {
-	static typename TContainer::TInputBuffer& GetBuffer(TContainer& Container)
+	template<typename TContainer, typename TState>
+	typename TEnableIf< TIsDerivedFrom<typename TContainer::TInputCmd, TState>::IsDerived, typename TContainer::TInputBuffer&>::Type
+	Get(TContainer& Container)
 	{
 		return Container.Input;
 	}
-};
 
-template<typename TContainer>
-struct TBufferGetterHelper_ByState<TContainer, typename TContainer::TSyncState>
-{
-	static typename TContainer::TSyncBuffer& GetBuffer(TContainer& Container)
+	template<typename TContainer, typename TState>
+	typename TEnableIf< TIsDerivedFrom<typename TContainer::TSyncState, TState>::IsDerived, typename TContainer::TSyncBuffer&>::Type
+	Get(TContainer& Container)
 	{
 		return Container.Sync;
 	}
-};
 
-template<typename TContainer>
-struct TBufferGetterHelper_ByState<TContainer, typename TContainer::TAuxState>
-{
-	static typename TContainer::TAuxBuffer& GetBuffer(TContainer& Container)
+	template<typename TContainer, typename TState>
+	typename TEnableIf< TIsDerivedFrom<typename TContainer::TAuxState, TState>::IsDerived, typename TContainer::TAuxBuffer&>::Type
+	Get(TContainer& Container)
 	{
 		return Container.Aux;
 	}
-};
 
-template<typename TContainer>
-struct TBufferGetterHelper_ByState<TContainer, typename TContainer::TDebugState>
-{
-	static typename TContainer::TDebugBuffer& GetBuffer(TContainer& Container)
+	template<typename TContainer, typename TState>
+	typename TEnableIf< TIsDerivedFrom<typename TContainer::TDebugState, TState>::IsDerived, typename TContainer::TDebugState&>::Type
+	Get(TContainer& Container)
 	{
 		return Container.Debug;
 	}
-};
-
-// -------------------------------------------------------------------
+}
 
 // Struct that encapsulates writing a new element to a buffer. This is used to allow a new Aux state to be created in the ::SimulationTick loop.
 template<typename T>
@@ -541,33 +510,30 @@ struct TScopedSimulationTick
 
 // ----------------------------------------------------------------------------------------------------------------------------------------------
 //	Accessors - helper structs that provide safe/cleaner access to the underlying NetSim states/events
+//	This is the equivlent of directly calling TNetworkedSimulationModel::GetPendingStateRead() / GetPendingStateWrite()
+//
+//	This is useful because it is not always practical to have a typed pointer to your TNetworkedSimulationModel instance. For example,
+//	the TNetworkedSimulationModel may be templated and conditionally instantiated (for example fix tick vs non fixed tick models, and eventually
+//	we may offer more memory allocation related templated parameters).
+//		
+//	Essentially, this allows you to plop a TNetSimStateAccessor<MyStateType> on a class and have it bind to any TNetworkedSimulationModel instantiation  
+//	that has uses that type.
+//
+//	See additional comments in TNetworkedSimulationModel::GetPendingStateRead() / GetPendingStateWrite()
 // ----------------------------------------------------------------------------------------------------------------------------------------------
 
-// Accessor conditionally gives access to the current (pending) Sync/Aux state to outside code.
-// Reads are always allowed
-// Writes are conditional. Authority can always write to the pending frame. Non authority requires the netsim to be currently processing an ::SimulationTick.
-// If you aren't inside an ::SimulationTick call, it is really not safe to predict state changes. It is safest and simplest to just not predict these changes.
-//
-// Explanation: During the scope of an ::SimulationTick call, we know exactly 'when' we are relative to what the server is processing. If the predicting client wants
-// to predict a change to sync/aux state during an update, the server will do it at the exact same time (assuming not a mis prediction). When a state change
-// happens "out of band" (outside an ::SimulationTick call) - we really have no way to correlate when the server will do it. While its tempting to think "we will get
-// a correction anyways, might as well guess at it and maybe get a smaller correction" - but this opens us up to other problems. The server may actually not 
-// change the state at all and you may not get an update that corrects you. You could add a timeout and track the state change somewhere but that really complicates
-// things and could leave you open to "double" problems: if the state change is additive, you may stack the authority change on top of the local predicted change, or
-// you may roll back the predicted change to then later receive the authority change.
-//	
-// What still may make sense to do is allow the "In Update" bool to be temporarily disabled if we enter code that we know is not rollback friendly.
+
 template<typename TState>
-struct TNetworkSimStateAccessor
+struct TNetSimStateAccessor
 {
+	// Bind to the NetsimModel that we are accessing. This is a templated method so that a single declared TNetSimStateAccessor can bind to any instantiated netsim model
+	// that has the underlying type of the accessor. This allows, for example, templated netsim models to be instantiated and bind to a single accessor. (E.g, variable or fixed tick sim)
 	template<typename TNetworkSimModel>
-	void Init(TNetworkSimModel* NetSimModel)
+	void Bind(TNetworkSimModel* NetSimModel)
 	{
 		GetStateFunc = [NetSimModel](bool bWrite, TState*& OutState, bool& OutSafe)
 		{
-			// Gross: find the buffer our TState is in. This allows these accessors to be declared as class member variables without knowing the exact net sim it will
-			// be pulling from. (For example, you may want to templatize your network sim model).
-			auto& Buffer = TBufferGetterHelper_ByState<TNetworkSimBufferContainer<typename TNetworkSimModel::TBufferTypes>, TState>::GetBuffer(NetSimModel->Buffers);
+			auto& Buffer = NetSimBufferSelect::Get<TNetworkSimBufferContainer<typename TNetworkSimModel::TBufferTypes>, TState>(NetSimModel->Buffers);
 			OutState = bWrite ? Buffer.WriteFrameInitializedFromHead(NetSimModel->Ticker.PendingFrame) : Buffer[NetSimModel->Ticker.PendingFrame];
 			OutSafe = NetSimModel->Ticker.bUpdateInProgress;
 		};
diff --git a/Engine/Plugins/Runtime/NetworkPredictionExtras/Source/NetworkPredictionExtras/Private/MockAbilitySimulation.cpp b/Engine/Plugins/Runtime/NetworkPredictionExtras/Source/NetworkPredictionExtras/Private/MockAbilitySimulation.cpp
--- a/Engine/Plugins/Runtime/NetworkPredictionExtras/Source/NetworkPredictionExtras/Private/MockAbilitySimulation.cpp
+++ b/Engine/Plugins/Runtime/NetworkPredictionExtras/Source/NetworkPredictionExtras/Private/MockAbilitySimulation.cpp
@@ -228,27 +228,19 @@ UMockFlyingAbilityComponent::UMockFlyingAbilityComponent()
 
 }
 
-INetworkSimulationModel* UMockFlyingAbilityComponent::InstantiateNetworkSimulation()
+INetworkedSimulationModel* UMockFlyingAbilityComponent::InstantiateNetworkedSimulation()
 {
 	check(UpdatedComponent);
-	
-	// FIXME: init sim helper on super?
-	MockAbilitySimulation.Reset(new FMockAbilitySimulation());
-	MockAbilitySimulation->UpdatedComponent = UpdatedComponent;
-	MockAbilitySimulation->UpdatedPrimitive = UpdatedPrimitive;
-	MockAbilitySimulation->EventHandler = this;
-	
-	FMockAbilitySyncState InitialSyncState;
-	InitialSyncState.Location = UpdatedComponent->GetComponentLocation();
-	InitialSyncState.Rotation = UpdatedComponent->GetComponentQuat().Rotator();	
 
+	// Simulation
+	FMockAbilitySyncState InitialSyncState;
 	FMockAbilityAuxstate InitialAuxState;
-	InitialAuxState.MaxSpeed = MockAbilityCVars::DefaultMaxSpeed();
+	InitMockAbilitySimulation(new FMockAbilitySimulation(), InitialSyncState, InitialAuxState);
 
-	auto NewModel = new FMockAbilitySystem<0>(MockAbilitySimulation.Get(), this, InitialSyncState, InitialAuxState);
+	// Model
+	auto NewModel = new FMockAbilitySystem<0>( MockAbilitySimulation, this, InitialSyncState, InitialAuxState);
+	InitMockAbilityNetSimModel(NewModel);
 
-	MovementSyncState.Init(NewModel);
-	MovementAuxState.Init(NewModel);
 	return NewModel;
 }
 
@@ -339,3 +331,4 @@ void UMockFlyingAbilityComponent::NotifyBlinkFinished()
 
 	//UE_LOG(LogTemp, Display, TEXT("Blink Finished"));
 }
+
diff --git a/Engine/Plugins/Runtime/NetworkPredictionExtras/Source/NetworkPredictionExtras/Private/MockNetworkSimulation.cpp b/Engine/Plugins/Runtime/NetworkPredictionExtras/Source/NetworkPredictionExtras/Private/MockNetworkSimulation.cpp
--- a/Engine/Plugins/Runtime/NetworkPredictionExtras/Source/NetworkPredictionExtras/Private/MockNetworkSimulation.cpp
+++ b/Engine/Plugins/Runtime/NetworkPredictionExtras/Source/NetworkPredictionExtras/Private/MockNetworkSimulation.cpp
@@ -106,14 +106,15 @@ UMockNetworkSimulationComponent::UMockNetworkSimulationComponent()
 	}
 }
 
-INetworkSimulationModel* UMockNetworkSimulationComponent::InstantiateNetworkSimulation()
+INetworkedSimulationModel* UMockNetworkSimulationComponent::InstantiateNetworkedSimulation()
 {
-	MockNetworkSimulation.Reset(new FMockNetworkSimulation());
+	check(MockNetworkSimulation == nullptr);
+	MockNetworkSimulation = new FMockNetworkSimulation();
 
 	FMockSyncState InitialSyncState;
 	InitialSyncState.Total = MockValue;
 
-	auto* NewModel = new FMockNetworkModel(MockNetworkSimulation.Get(), this, InitialSyncState);
+	auto* NewModel = new FMockNetworkModel(MockNetworkSimulation, this, InitialSyncState);
 	return NewModel;
 }
 
diff --git a/Engine/Plugins/Runtime/NetworkPredictionExtras/Source/NetworkPredictionExtras/Private/NetworkPredictionExtrasFlyingPawn.cpp b/Engine/Plugins/Runtime/NetworkPredictionExtras/Source/NetworkPredictionExtras/Private/NetworkPredictionExtrasFlyingPawn.cpp
--- a/Engine/Plugins/Runtime/NetworkPredictionExtras/Source/NetworkPredictionExtras/Private/NetworkPredictionExtrasFlyingPawn.cpp
+++ b/Engine/Plugins/Runtime/NetworkPredictionExtras/Source/NetworkPredictionExtras/Private/NetworkPredictionExtrasFlyingPawn.cpp
@@ -474,9 +474,30 @@ const UMockFlyingAbilityComponent* ANetworkPredictionExtrasFlyingPawn_MockAbilit
 void ANetworkPredictionExtrasFlyingPawn_MockAbility::ProduceInput(const FNetworkSimTime SimTime, FMockAbilityInputCmd& Cmd)
 {
 	Super::ProduceInput(SimTime, Cmd);
-	Cmd.bSprintPressed = bSprintPressed;
-	Cmd.bDashPressed = bDashPressed;
-	Cmd.bBlinkPressed = bBlinkPressed;
+
+	switch(AbilityInputPreset)
+	{
+	case ENetworkPredictionExtrasMockAbilityInputPreset::None:
+		Cmd.bSprintPressed = bSprintPressed;
+		Cmd.bDashPressed = bDashPressed;
+		Cmd.bBlinkPressed = bBlinkPressed;
+		break;
+	case ENetworkPredictionExtrasMockAbilityInputPreset::Sprint:
+		Cmd.bSprintPressed = true;
+		Cmd.bDashPressed = false;
+		Cmd.bBlinkPressed = false;
+		break;
+	case ENetworkPredictionExtrasMockAbilityInputPreset::Dash:
+		Cmd.bSprintPressed = false;
+		Cmd.bDashPressed = true;
+		Cmd.bBlinkPressed = false;
+		break;
+	case ENetworkPredictionExtrasMockAbilityInputPreset::Blink:
+		Cmd.bSprintPressed = false;
+		Cmd.bDashPressed = false;
+		Cmd.bBlinkPressed = true;
+		break;
+	};
 }
 
 void ANetworkPredictionExtrasFlyingPawn_MockAbility::BeginPlay()
@@ -496,7 +517,7 @@ float ANetworkPredictionExtrasFlyingPawn_MockAbility::GetStamina() const
 {
 	if (const UMockFlyingAbilityComponent* FlyingAbilityComponent = GetMockFlyingAbilityComponent())
 	{
-		if (const FMockAbilitySyncState* SyncState = FlyingAbilityComponent->MovementSyncState.GetStateRead())
+		if (const FMockAbilitySyncState* SyncState = FlyingAbilityComponent->AbilitySyncState.GetStateRead())
 		{
 			return SyncState->Stamina;
 		}
@@ -508,7 +529,7 @@ float ANetworkPredictionExtrasFlyingPawn_MockAbility::GetMaxStamina() const
 {
 	if (const UMockFlyingAbilityComponent* FlyingAbilityComponent = GetMockFlyingAbilityComponent())
 	{
-		if (const FMockAbilityAuxstate* AuxState = FlyingAbilityComponent->MovementAuxState.GetStateRead())
+		if (const FMockAbilityAuxstate* AuxState = FlyingAbilityComponent->AbilityAuxState.GetStateRead())
 		{
 			return AuxState->MaxStamina;
 		}
diff --git a/Engine/Plugins/Runtime/NetworkPredictionExtras/Source/NetworkPredictionExtras/Public/MockAbilitySimulation.h b/Engine/Plugins/Runtime/NetworkPredictionExtras/Source/NetworkPredictionExtras/Public/MockAbilitySimulation.h
--- a/Engine/Plugins/Runtime/NetworkPredictionExtras/Source/NetworkPredictionExtras/Public/MockAbilitySimulation.h
+++ b/Engine/Plugins/Runtime/NetworkPredictionExtras/Source/NetworkPredictionExtras/Public/MockAbilitySimulation.h
@@ -164,8 +164,8 @@ public:
 	DECLARE_DELEGATE_TwoParams(FProduceMockAbilityInput, const FNetworkSimTime /*SimTime*/, FMockAbilityInputCmd& /*Cmd*/)
 	FProduceMockAbilityInput ProduceInputDelegate;
 
-	TNetworkSimStateAccessor<FMockAbilitySyncState> MovementSyncState;
-	TNetworkSimStateAccessor<FMockAbilityAuxstate> MovementAuxState;
+	TNetSimStateAccessor<FMockAbilitySyncState> AbilitySyncState;
+	TNetSimStateAccessor<FMockAbilityAuxstate> AbilityAuxState;
 
 	// IMockFlyingAbilitySystemDriver
 	FString GetDebugName() const override;
@@ -235,6 +235,25 @@ private:
 protected:
 
 	// Network Prediction
-	virtual INetworkSimulationModel* InstantiateNetworkSimulation() override;
-	TUniquePtr<FMockAbilitySimulation> MockAbilitySimulation;
+	virtual INetworkedSimulationModel* InstantiateNetworkedSimulation() override;
+	FMockAbilitySimulation* MockAbilitySimulation = nullptr;
+
+	template<typename TSimulation>
+	void InitMockAbilitySimulation(TSimulation* Simulation, FMockAbilitySyncState& InitialSyncState, FMockAbilityAuxstate& InitialAuxState)
+	{
+		check(MockAbilitySimulation == nullptr);
+		MockAbilitySimulation = Simulation;
+		MockAbilitySimulation->EventHandler = this;
+
+		InitFlyingMovementSimulation(Simulation, InitialSyncState, InitialAuxState);
+	}
+
+	template<typename TNetSimModel>
+	void InitMockAbilityNetSimModel(TNetSimModel* Model)
+	{
+		AbilitySyncState.Bind(Model);
+		MovementAuxState.Bind(Model);
+
+		InitFlyingMovementNetSimModel(Model);
+	}
 };
\ No newline at end of file
diff --git a/Engine/Plugins/Runtime/NetworkPredictionExtras/Source/NetworkPredictionExtras/Public/MockNetworkSimulation.h b/Engine/Plugins/Runtime/NetworkPredictionExtras/Source/NetworkPredictionExtras/Public/MockNetworkSimulation.h
--- a/Engine/Plugins/Runtime/NetworkPredictionExtras/Source/NetworkPredictionExtras/Public/MockNetworkSimulation.h
+++ b/Engine/Plugins/Runtime/NetworkPredictionExtras/Source/NetworkPredictionExtras/Public/MockNetworkSimulation.h
@@ -166,7 +166,7 @@ public:
 
 	virtual void TickComponent(float DeltaTime, enum ELevelTick TickType, FActorComponentTickFunction *ThisTickFunction) override;
 	
-	virtual INetworkSimulationModel* InstantiateNetworkSimulation() override;
+	virtual INetworkedSimulationModel* InstantiateNetworkedSimulation() override;
 	
 public:
 
@@ -180,5 +180,5 @@ public:
 	// Mock representation of "syncing' to the sync state in the network sim.
 	float MockValue = 1000.f;
 
-	TUniquePtr<FMockNetworkSimulation> MockNetworkSimulation;
+	FMockNetworkSimulation* MockNetworkSimulation = nullptr;
 };
\ No newline at end of file
diff --git a/Engine/Plugins/Runtime/NetworkPredictionExtras/Source/NetworkPredictionExtras/Public/NetworkPredictionExtrasFlyingPawn.h b/Engine/Plugins/Runtime/NetworkPredictionExtras/Source/NetworkPredictionExtras/Public/NetworkPredictionExtrasFlyingPawn.h
--- a/Engine/Plugins/Runtime/NetworkPredictionExtras/Source/NetworkPredictionExtras/Public/NetworkPredictionExtrasFlyingPawn.h
+++ b/Engine/Plugins/Runtime/NetworkPredictionExtras/Source/NetworkPredictionExtras/Public/NetworkPredictionExtrasFlyingPawn.h
@@ -108,6 +108,16 @@ private:
 	void Action_RightShoulder_Released() { }
 };
 
+UENUM()
+enum class ENetworkPredictionExtrasMockAbilityInputPreset: uint8
+{
+	/** No input */
+	None,
+	Sprint,
+	Dash,
+	Blink
+};
+
 // Example subclass of ANetworkPredictionExtrasFlyingPawn that uses the MockAbility simulation
 UCLASS()
 class NETWORKPREDICTIONEXTRAS_API ANetworkPredictionExtrasFlyingPawn_MockAbility : public ANetworkPredictionExtrasFlyingPawn
@@ -124,6 +134,9 @@ public:
 	UMockFlyingAbilityComponent* GetMockFlyingAbilityComponent();
 	const UMockFlyingAbilityComponent* GetMockFlyingAbilityComponent() const;
 
+	UPROPERTY(EditAnywhere, BlueprintReadWrite, Category="Automation")
+	ENetworkPredictionExtrasMockAbilityInputPreset AbilityInputPreset = ENetworkPredictionExtrasMockAbilityInputPreset::None;
+
 	UFUNCTION(BlueprintCallable, Category="Ability")
 	float GetStamina() const;
 
