NetworkPrediction Plugin:
-Exposed temp interpolation setting on flying movemen component (will probably unify sim update settings soon)
-Fix some issues with interpolation check not honoring cvars
-initial test maps with mock ability system

#rnx

diff --git a/Engine/Plugins/Runtime/NetworkPrediction/Source/NetworkPrediction/Public/Movement/FlyingMovement.h b/Engine/Plugins/Runtime/NetworkPrediction/Source/NetworkPrediction/Public/Movement/FlyingMovement.h
--- a/Engine/Plugins/Runtime/NetworkPrediction/Source/NetworkPrediction/Public/Movement/FlyingMovement.h
+++ b/Engine/Plugins/Runtime/NetworkPrediction/Source/NetworkPrediction/Public/Movement/FlyingMovement.h
@@ -209,6 +209,9 @@ public:
 	void ProduceInput(const FNetworkSimTime SimTime, FlyingMovement::FInputCmd& Cmd) override;
 	void FinalizeFrame(const FlyingMovement::FMoveState& SyncState, const FlyingMovement::FAuxState& AuxState) override;
 
+	UPROPERTY(EditAnywhere, BlueprintReadOnly, Category=FlyingMovementNetworking)
+	bool bEnableInterpolation = true;
+
 protected:
 
 	// Network Prediction
@@ -236,6 +239,7 @@ protected:
 	template<typename TNetSimModel>
 	void InitFlyingMovementNetSimModel(TNetSimModel* Model)
 	{
+		Model->RepProxy_Simulated.bAllowSimulatedExtrapolation = !bEnableInterpolation;
 		MovementSyncState.Bind(Model);
 		MovementAuxState.Bind(Model);
 	}
diff --git a/Engine/Plugins/Runtime/NetworkPrediction/Source/NetworkPrediction/Public/NetworkedSimulationModelReplicators.h b/Engine/Plugins/Runtime/NetworkPrediction/Source/NetworkPrediction/Public/NetworkedSimulationModelReplicators.h
--- a/Engine/Plugins/Runtime/NetworkPrediction/Source/NetworkPrediction/Public/NetworkedSimulationModelReplicators.h
+++ b/Engine/Plugins/Runtime/NetworkPrediction/Source/NetworkPrediction/Public/NetworkedSimulationModelReplicators.h
@@ -459,7 +459,7 @@ struct TReplicator_Simulated : public TBase
 
 		check(Ticker.GetTotalProcessedSimulationTime() <= Ticker.GetTotalAllowedSimulationTime());
 
-		if (bAllowSimulatedExtrapolation && ParentSimulation == nullptr && NetworkSimulationModelCVars::EnableSimulatedExtrapolation() && NetworkSimulationModelCVars::EnableSimulatedReconcile())
+		if (GetSimulatedUpdateMode() == ESimulatedUpdateMode::Extrapolate && NetworkSimulationModelCVars::EnableSimulatedReconcile())
 		{
 			// Simulated Reconcile requires the input buffer to be kept up to date with the Sync buffer
 			// Generate a new, fake, command since we just added a new sync state to head
@@ -484,7 +484,7 @@ struct TReplicator_Simulated : public TBase
 	void PreSimTick(TDriver* Driver, TNetworkSimBufferContainer<TBufferTypes>& Buffers, TSimulationTicker<TTickSettings>& Ticker, const FNetSimTickParameters& TickParameters)
 	{
 		// Tick if we are dependent simulation or extrapolation is enabled
-		if (ParentSimulation || (bAllowSimulatedExtrapolation && NetworkSimulationModelCVars::EnableSimulatedExtrapolation()))
+		if (ParentSimulation || (GetSimulatedUpdateMode() == ESimulatedUpdateMode::Extrapolate))
 		{
 			// Don't start this simulation until you've gotten at least one update from the server
 			if (Ticker.GetTotalProcessedSimulationTime().IsPositive())
@@ -500,7 +500,11 @@ struct TReplicator_Simulated : public TBase
 	template<typename TDriver>
 	void PostSimTick(TDriver* Driver, const TNetworkSimBufferContainer<TBufferTypes>& Buffers, const TSimulationTicker<TTickSettings>& Ticker, const FNetSimTickParameters& TickParameters)
 	{
-		if (bAllowSimulatedExtrapolation || ParentSimulation)
+		if (GetSimulatedUpdateMode() == ESimulatedUpdateMode::Interpolate)
+		{
+			Interpolator.template PostSimTick<TDriver>(Driver, Buffers, Ticker, TickParameters);
+		}
+		else
 		{
 			// Sync to latest frame if there is any
 			if (Buffers.Sync.Num() > 0)
@@ -509,10 +513,6 @@ struct TReplicator_Simulated : public TBase
 			}
 			
 		}
-		else
-		{
-			Interpolator.template PostSimTick<TDriver>(Driver, Buffers, Ticker, TickParameters);
-		}
 	}
 
 	template<typename TSimulation, typename TDriver>
diff --git a/Engine/Plugins/Runtime/NetworkPredictionExtras/Source/NetworkPredictionExtras/Public/NetworkPredictionExtrasFlyingPawn.h b/Engine/Plugins/Runtime/NetworkPredictionExtras/Source/NetworkPredictionExtras/Public/NetworkPredictionExtrasFlyingPawn.h
--- a/Engine/Plugins/Runtime/NetworkPredictionExtras/Source/NetworkPredictionExtras/Public/NetworkPredictionExtrasFlyingPawn.h
+++ b/Engine/Plugins/Runtime/NetworkPredictionExtras/Source/NetworkPredictionExtras/Public/NetworkPredictionExtrasFlyingPawn.h
@@ -87,7 +87,7 @@ protected:
 
 	void ProduceInput(const FNetworkSimTime SimTime, FlyingMovement::FInputCmd& Cmd);
 
-	UPROPERTY()
+	UPROPERTY(Category=Movement, VisibleAnywhere)
 	UFlyingMovementComponent* FlyingMovementComponent;
 
 private:
