NetworkPrediction Plugin:
-NetSimCues initial check in
-These are the async/deferred events that are outputted during the network simulation tick
-Different classification in types wrt prediction and replication
-See FMockAbilityBlinkCue in code, "Misc Notes" at bottom of readme file
-"Rewindable NetSimCue" api still todo

#rnx

diff --git a/Engine/Plugins/Runtime/NetworkPrediction/Source/NetworkPrediction/Private/Movement/FlyingMovementSimulation.cpp b/Engine/Plugins/Runtime/NetworkPrediction/Source/NetworkPrediction/Private/Movement/FlyingMovementSimulation.cpp
--- a/Engine/Plugins/Runtime/NetworkPrediction/Source/NetworkPrediction/Private/Movement/FlyingMovementSimulation.cpp
+++ b/Engine/Plugins/Runtime/NetworkPrediction/Source/NetworkPrediction/Private/Movement/FlyingMovementSimulation.cpp
@@ -150,7 +150,7 @@ namespace FlyingMovement
 		return 0.f;
 	}
 
-	void FMovementSimulation::SimulationTick(const TNetSimTimeStep& TimeStep, const TNetSimInput<TMovementBufferTypes>& Input, const TNetSimOutput<TMovementBufferTypes>& Output)
+	void FMovementSimulation::SimulationTick(const FNetSimTimeStep& TimeStep, const TNetSimInput<TMovementBufferTypes>& Input, const TNetSimOutput<TMovementBufferTypes>& Output)
 	{
 		Output.Sync = Input.Sync;
 
diff --git a/Engine/Plugins/Runtime/NetworkPrediction/Source/NetworkPrediction/Private/Movement/ParametricMovement.cpp b/Engine/Plugins/Runtime/NetworkPrediction/Source/NetworkPrediction/Private/Movement/ParametricMovement.cpp
--- a/Engine/Plugins/Runtime/NetworkPrediction/Source/NetworkPrediction/Private/Movement/ParametricMovement.cpp
+++ b/Engine/Plugins/Runtime/NetworkPrediction/Source/NetworkPrediction/Private/Movement/ParametricMovement.cpp
@@ -41,7 +41,7 @@ namespace ParametricMovement
 {
 	const FName FMovementSimulation::GroupName(TEXT("Parametric"));
 
-	void FMovementSimulation::SimulationTick(const TNetSimTimeStep& TimeStep, const TNetSimInput<TMovementBufferTypes>& Input, const TNetSimOutput<TMovementBufferTypes>& Output)
+	void FMovementSimulation::SimulationTick(const FNetSimTimeStep& TimeStep, const TNetSimInput<TMovementBufferTypes>& Input, const TNetSimOutput<TMovementBufferTypes>& Output)
 	{
 		check(Motion != nullptr); // Must set motion mapping prior to running the simulation
 		const float DeltaSeconds = TimeStep.StepMS.ToRealTimeSeconds();
diff --git a/Engine/Plugins/Runtime/NetworkPrediction/Source/NetworkPrediction/Private/NetworkPredictionComponent.cpp b/Engine/Plugins/Runtime/NetworkPrediction/Source/NetworkPrediction/Private/NetworkPredictionComponent.cpp
--- a/Engine/Plugins/Runtime/NetworkPrediction/Source/NetworkPrediction/Private/NetworkPredictionComponent.cpp
+++ b/Engine/Plugins/Runtime/NetworkPrediction/Source/NetworkPrediction/Private/NetworkPredictionComponent.cpp
@@ -10,6 +10,11 @@
 UNetworkPredictionComponent::UNetworkPredictionComponent()
 {
 	SetIsReplicatedByDefault(true);
+	
+	// Tick in order to dispatch NetSimCues. FIXME: Might be useful to have this as an option. Some may want to dispatch cues from actor tick (avoid component ticking at all)
+	PrimaryComponentTick.bCanEverTick = true;
+	PrimaryComponentTick.bStartWithTickEnabled = true;
+	PrimaryComponentTick.TickGroup = TG_PrePhysics;
 }
 
 void UNetworkPredictionComponent::InitializeComponent()
@@ -207,3 +212,11 @@ void UNetworkPredictionComponent::UnregisterServerRPCDelegate()
 		ServerRPCHandle.Reset();
 	}
 }
+
+void UNetworkPredictionComponent::TickComponent(float DeltaTime, enum ELevelTick TickType, FActorComponentTickFunction *ThisTickFunction)
+{
+	Super::TickComponent(DeltaTime, TickType, ThisTickFunction);
+
+	// Handle all pending cues (dispatch them to their handler)
+	NetSimModel->ProcessPendingNetSimCues();
+}
\ No newline at end of file
diff --git a/Engine/Plugins/Runtime/NetworkPrediction/Source/NetworkPrediction/Private/NetworkedSimulationModelCues.cpp b/Engine/Plugins/Runtime/NetworkPrediction/Source/NetworkPrediction/Private/NetworkedSimulationModelCues.cpp
new file mode 100644
--- /dev/null
+++ b/Engine/Plugins/Runtime/NetworkPrediction/Source/NetworkPrediction/Private/NetworkedSimulationModelCues.cpp
@@ -0,0 +1,188 @@
+// Copyright 1998-2019 Epic Games, Inc. All Rights Reserved.
+
+#include "NetworkedSimulationModelCues.h"
+
+DEFINE_LOG_CATEGORY(LogNetSimCues);
+
+FGlobalCueTypeTable FGlobalCueTypeTable::Singleton;
+
+// ------------------------------------------------------------------------------------------------------------------------------------------
+//	Mock Cue example
+// ------------------------------------------------------------------------------------------------------------------------------------------
+
+// Minimal example of NetSimCue pipeline
+#define ENABLE_MOCK_CUE 0
+#if ENABLE_MOCK_CUE
+
+// ---------------------------------------------------------------------
+// Mock Cue Types
+// ---------------------------------------------------------------------
+
+struct FMockImpactCue
+{
+	NETSIMCUE_BODY();
+
+	FVector ImpactLocation;
+
+	void NetSerialize(FArchive& Ar)
+	{
+		Ar << ImpactLocation;
+	}
+
+	static bool Unique(const FMockImpactCue& A, const FMockImpactCue& B)
+	{
+		const float ErrorTolerance = 1.f;
+		return !A.ImpactLocation.Equals(B.ImpactLocation, ErrorTolerance);
+	}
+
+};
+NETSIMCUE_REGISTER(FMockImpactCue, TEXT("Impact"));
+
+// ----------------
+
+struct FMockDamageCue
+{
+	NETSIMCUE_BODY();
+
+	int32 SourceID;
+	int32 DamageType;
+	FVector HitLocation;
+	
+	void NetSerialize(FArchive& Ar)
+	{
+		Ar << SourceID;
+		Ar << DamageType;
+		Ar << HitLocation;
+	}
+
+	static bool Unique(const FMockDamageCue& A, const FMockDamageCue& B)
+	{
+		const float ErrorTolerance = 1.f;
+		return A.SourceID != B.SourceID || A.DamageType != B.DamageType || !A.HitLocation.Equals(B.HitLocation, ErrorTolerance);
+	}
+};
+NETSIMCUE_REGISTER(FMockDamageCue, TEXT("Damage"));
+
+struct FMockHealingCue
+{
+	NETSIMCUE_BODY();
+
+	int32 SourceID;
+	int32 HealingAmount;
+	
+	void NetSerialize(FArchive& Ar)
+	{
+		Ar << SourceID;
+		Ar << HealingAmount;
+	}
+
+	static bool Unique(const FMockHealingCue& A, const FMockHealingCue& B)
+	{
+		const float ErrorTolerance = 1.f;
+		return A.SourceID != B.SourceID || A.HealingAmount != B.HealingAmount;
+	}
+};
+NETSIMCUE_REGISTER(FMockHealingCue, TEXT("Healing"));
+
+// ---------------------------------------------------------------------
+// Mock Handlers (equevilent to actor component hierarchy)
+// ---------------------------------------------------------------------
+
+class TestHandlerBase
+{
+public:
+
+	template<typename TDispatchTable>
+	static void RegisterNetSimCueTypes(TDispatchTable& DispatchTable)
+	{
+		DispatchTable.template RegisterType<FMockImpactCue>();
+	}
+
+	void HandleCue(FMockImpactCue& ImpactCue, const FNetworkSimTime& ElapsedTime)
+	{
+		UE_LOG(LogTemp, Warning, TEXT("Impact!"));
+	}
+};
+
+class TestHandlerChild : public TestHandlerBase
+{
+public:
+
+	template<typename TDispatchTable>
+	static void RegisterNetSimCueTypes(TDispatchTable& DispatchTable)
+	{
+		TestHandlerBase::RegisterNetSimCueTypes(DispatchTable);
+		DispatchTable.template RegisterType<FMockDamageCue>();
+		DispatchTable.template RegisterType<FMockHealingCue>();
+	}
+
+	using TestHandlerBase::HandleCue; // Required for us to "use" the HandleCue methods in our parent class
+
+	void HandleCue(FMockDamageCue& DamageCue, const FNetworkSimTime& ElapsedTime)
+	{
+		UE_LOG(LogTemp, Warning, TEXT("HandleCue: Damage. SourceID: %d DamageType: %d HitLocation: %s"), DamageCue.SourceID, DamageCue.DamageType, *DamageCue.HitLocation.ToString());
+	}
+
+	void HandleCue(FMockHealingCue& HealingCue, const FNetworkSimTime& ElapsedTime)
+	{
+		UE_LOG(LogTemp, Warning, TEXT("HandleCue: Healing. SourceID: %d Healing Amount: %d"), HealingCue.SourceID, HealingCue.HealingAmount);
+	}
+};
+
+NETSIMCUEHANDLER_REGISTER(TestHandlerChild);
+
+void TestCues()
+{
+	// -------------------------------------
+	// Simuation/Invoke
+	// -------------------------------------
+
+	struct FBaseSimulation
+	{
+		static void TickSimulation(FCueDispatcher& Dispatcher)
+		{
+			Dispatcher.Invoke<FMockImpactCue>( { FVector(1.f, 2.f, 3.f) } );
+		}
+	};
+
+	struct FChildSimulation
+	{
+		static void TickSimulation(FCueDispatcher& Dispatcher)
+		{
+			FBaseSimulation::TickSimulation(Dispatcher);
+
+			Dispatcher.Invoke<FMockDamageCue>( { 1, 2, FVector(1.f, 2.f, 3.f) } );
+			Dispatcher.Invoke<FMockHealingCue>( { 10, 32 } );
+		}
+	};
+
+	FCueDispatcher ServerDispatcher;
+	FChildSimulation::TickSimulation(ServerDispatcher);
+
+	// -------------------------------------
+	// Send
+	// -------------------------------------
+
+	FNetBitWriter TempWriter(1024 << 3);
+	ServerDispatcher.NetSerializeSavedCues(TempWriter, false);
+
+	// -------------------------------------
+	// Receive
+	// -------------------------------------
+
+	FNetBitReader TempReader(nullptr, TempWriter.GetData(), TempWriter.GetNumBits());
+	FCueDispatcher ClientDispatcher;
+	ClientDispatcher.NetSerializeSavedCues(TempReader, false);
+
+	TestHandlerChild MyObject;
+
+	ClientDispatcher.DispatchCueRecord(MyObject, FNetworkSimTime());
+}
+
+FAutoConsoleCommandWithWorldAndArgs NetworkSimulationModelCueCmd(TEXT("nms.CueTest"), TEXT(""),
+	FConsoleCommandWithWorldAndArgsDelegate::CreateLambda([](const TArray< FString >& Args, UWorld* World) 
+{
+	TestCues();
+}));
+
+#endif // ENABLE_MOCK_CUE
\ No newline at end of file
diff --git a/Engine/Plugins/Runtime/NetworkPrediction/Source/NetworkPrediction/Public/Movement/FlyingMovement.h b/Engine/Plugins/Runtime/NetworkPrediction/Source/NetworkPrediction/Public/Movement/FlyingMovement.h
--- a/Engine/Plugins/Runtime/NetworkPrediction/Source/NetworkPrediction/Public/Movement/FlyingMovement.h
+++ b/Engine/Plugins/Runtime/NetworkPrediction/Source/NetworkPrediction/Public/Movement/FlyingMovement.h
@@ -155,7 +155,7 @@ namespace FlyingMovement
 		static bool ForceMispredict;
 
 		/** Main update function */
-		NETWORKPREDICTION_API void SimulationTick(const TNetSimTimeStep& TimeStep, const TNetSimInput<TMovementBufferTypes>& Input, const TNetSimOutput<TMovementBufferTypes>& Output);
+		NETWORKPREDICTION_API void SimulationTick(const FNetSimTimeStep& TimeStep, const TNetSimInput<TMovementBufferTypes>& Input, const TNetSimOutput<TMovementBufferTypes>& Output);
 
 		// Called prior to running the sim to make sure to make sure the collision component is in the right place. 
 		// This is unfortunate and not good, but is needed to ensure our collision and world position have not been moved out from under us.
diff --git a/Engine/Plugins/Runtime/NetworkPrediction/Source/NetworkPrediction/Public/Movement/ParametricMovement.h b/Engine/Plugins/Runtime/NetworkPrediction/Source/NetworkPrediction/Public/Movement/ParametricMovement.h
--- a/Engine/Plugins/Runtime/NetworkPrediction/Source/NetworkPrediction/Public/Movement/ParametricMovement.h
+++ b/Engine/Plugins/Runtime/NetworkPrediction/Source/NetworkPrediction/Public/Movement/ParametricMovement.h
@@ -112,7 +112,7 @@ namespace ParametricMovement
 	public:
 		static const FName GroupName;
 
-		void SimulationTick(const TNetSimTimeStep& TimeStep, const TNetSimInput<TMovementBufferTypes>& Input, const TNetSimOutput<TMovementBufferTypes>& Output);
+		void SimulationTick(const FNetSimTimeStep& TimeStep, const TNetSimInput<TMovementBufferTypes>& Input, const TNetSimOutput<TMovementBufferTypes>& Output);
 
 		// Pointer to our static mapping of time->position
 		const FSimpleParametricMotion* Motion = nullptr;
diff --git a/Engine/Plugins/Runtime/NetworkPrediction/Source/NetworkPrediction/Public/NetworkPredictionComponent.h b/Engine/Plugins/Runtime/NetworkPrediction/Source/NetworkPrediction/Public/NetworkPredictionComponent.h
--- a/Engine/Plugins/Runtime/NetworkPrediction/Source/NetworkPrediction/Public/NetworkPredictionComponent.h
+++ b/Engine/Plugins/Runtime/NetworkPrediction/Source/NetworkPrediction/Public/NetworkPredictionComponent.h
@@ -30,6 +30,7 @@ public:
 	virtual void UninitializeComponent() override;
 	virtual void PreReplication(IRepChangedPropertyTracker & ChangedPropertyTracker) override;
 	virtual void PreNetReceive() override;
+	virtual void TickComponent(float DeltaTime, enum ELevelTick TickType, FActorComponentTickFunction *ThisTickFunction) override;
 	
 	virtual void Reconcile() { }
 	virtual void TickSimulation(float DeltaTimeSeconds) { }
diff --git a/Engine/Plugins/Runtime/NetworkPrediction/Source/NetworkPrediction/Public/NetworkPredictionTypes.h b/Engine/Plugins/Runtime/NetworkPrediction/Source/NetworkPrediction/Public/NetworkPredictionTypes.h
--- a/Engine/Plugins/Runtime/NetworkPrediction/Source/NetworkPrediction/Public/NetworkPredictionTypes.h
+++ b/Engine/Plugins/Runtime/NetworkPrediction/Source/NetworkPrediction/Public/NetworkPredictionTypes.h
@@ -251,6 +251,11 @@ public:
 	// Called by parent sim on the dependent sim as it reconciles
 	virtual void BeginRollback(const struct FNetworkSimTime& RollbackDeltaTime, const int32 ParentFrame) = 0;
 	virtual void StepRollback(const struct FNetworkSimTime& Step, const int32 ParentFrame, const bool FinalStep) = 0;
+
+	void ProcessPendingNetSimCues() { ProcessPendingNetSimCuesFunc(); }
+
+protected:
+	TFunction<void()> ProcessPendingNetSimCuesFunc;
 };
 
 // -------------------------------------------------------------------------------------------------------------------------------
diff --git a/Engine/Plugins/Runtime/NetworkPrediction/Source/NetworkPrediction/Public/NetworkedSimulationModel.h b/Engine/Plugins/Runtime/NetworkPrediction/Source/NetworkPrediction/Public/NetworkedSimulationModel.h
--- a/Engine/Plugins/Runtime/NetworkPrediction/Source/NetworkPrediction/Public/NetworkedSimulationModel.h
+++ b/Engine/Plugins/Runtime/NetworkPrediction/Source/NetworkPrediction/Public/NetworkedSimulationModel.h
@@ -68,7 +68,8 @@ public:
 	using TRealTime = FNetworkSimTime::FRealTime;
 	
 	// Note: Ownership of TSimulation* is taken by this constructor
-	TNetworkedSimulationModel(TSimulation* InSimulation, TDriver* InDriver, const TSyncState& InitialSyncState = TSyncState(), const TAuxState& InitialAuxState = TAuxState())
+	template<typename TInDriver>
+	TNetworkedSimulationModel(TSimulation* InSimulation, TInDriver* InDriver, const TSyncState& InitialSyncState = TSyncState(), const TAuxState& InitialAuxState = TAuxState())
 		: Simulation(InSimulation)
 	{
 		check(Simulation.IsValid() && InDriver);
@@ -77,6 +78,11 @@ public:
 		*Buffers.Aux.WriteFrame(0) = InitialAuxState;
 		Ticker.SetTotalProcessedSimulationTime(FNetworkSimTime(), 0);
 		DO_NETSIM_MODEL_DEBUG(FNetworkSimulationModelDebuggerManager::Get().RegisterNetworkSimulationModel(this, Driver->GetVLogOwner()));
+
+		ProcessPendingNetSimCuesFunc = [this, InDriver]()
+		{ 
+			Buffers.CueDispatcher.template DispatchCueRecord<TInDriver>(*InDriver, Ticker.GetTotalProcessedSimulationTime()); 
+		};
 	}
 
 	virtual ~TNetworkedSimulationModel()
@@ -176,12 +182,13 @@ public:
 				check(OutSyncState);
 
 				{
-					TScopedSimulationTick UpdateScope(Ticker, OutputFrame, InputCmd->GetFrameDeltaTime());
+					ESimulationTickContext Context = Parameters.Role == ROLE_Authority ? ESimulationTickContext::Authority : ESimulationTickContext::Predict;
+					TScopedSimulationTick UpdateScope(Ticker, Buffers.CueDispatcher, Context, OutputFrame, InputCmd->GetFrameDeltaTime());
 					
 					Simulation->SimulationTick( 
 						{ InputCmd->GetFrameDeltaTime(), Ticker },
 						{ *InputCmd, *InSyncState, *InAuxState },
-						{ *OutSyncState, Buffers.Aux.LazyWriter(OutputFrame) } );
+						{ *OutSyncState, Buffers.Aux.LazyWriter(OutputFrame), Buffers.CueDispatcher } );
 													
 				}
 
diff --git a/Engine/Plugins/Runtime/NetworkPrediction/Source/NetworkPrediction/Public/NetworkedSimulationModelCues.h b/Engine/Plugins/Runtime/NetworkPrediction/Source/NetworkPrediction/Public/NetworkedSimulationModelCues.h
new file mode 100644
--- /dev/null
+++ b/Engine/Plugins/Runtime/NetworkPrediction/Source/NetworkPrediction/Public/NetworkedSimulationModelCues.h
@@ -0,0 +1,535 @@
+// Copyright 1998-2019 Epic Games, Inc. All Rights Reserved.
+#pragma  once
+
+#include "Engine/EngineTypes.h"
+#include "NetworkedSimulationModelTime.h"
+#include "Containers/SortedMap.h"
+
+NETWORKPREDICTION_API DECLARE_LOG_CATEGORY_EXTERN(LogNetSimCues, Log, All);
+
+#ifndef NETSIMCUE_TYPEID_TYPE
+#define NETSIMCUE_TYPEID_TYPE uint8;
+#endif
+
+using FNetSimCueTypeId = NETSIMCUE_TYPEID_TYPE;
+
+enum class ESimulationTickContext : uint8
+{
+	None			= 0,
+	Authority		= 1 << 0,
+	Predict			= 1 << 1,
+	Resimulate		= 1 << 2,
+
+	All				= (Authority | Predict | Resimulate),
+};
+
+enum class ENetSimCueReplicationTarget : uint8
+{
+	None,			// Do not replicate cue to anyone
+	Interpolators,	// Only replicate to /accept on clients that are interpolating (not running the simulation themselves)
+	All,			// Replicate to everyone
+};
+
+// ----------------------------------------------------------------------------------------------------------------------------------------------
+//	NetSimCue Traits
+// ----------------------------------------------------------------------------------------------------------------------------------------------
+
+struct TNetSimCueTraitsBase
+{
+	// Who can Invoke this Cue in their simulation
+	static constexpr uint8 InvokeMask { (uint8)ESimulationTickContext::All };
+
+	// Does the cue replicate? (from authority)
+	static constexpr ENetSimCueReplicationTarget ReplicationTarget { ENetSimCueReplicationTarget::All };
+};
+
+// Preset: non replicated cue that only plays during "latest" simulate. Will not be played during rewind/resimulate.
+struct TNetSimCueTraits_Weak
+{
+	static constexpr uint8 InvokeMask { (uint8)ESimulationTickContext::Authority | (uint8)ESimulationTickContext::Predict };
+	static constexpr ENetSimCueReplicationTarget ReplicationTarget { ENetSimCueReplicationTarget::None };
+};
+
+// Preset: Replicated, non predicted. Only invoked on authority and will replicate to everyone else.
+struct TNetSimCueTraits_ReplicatedNonPredicted
+{
+	static constexpr uint8 InvokeMask { (uint8)ESimulationTickContext::Authority };
+	static constexpr ENetSimCueReplicationTarget ReplicationTarget { ENetSimCueReplicationTarget::All };
+};
+
+// Preset:  Replicated to simulated proxy, predicted by autonomous proxy
+struct TNetSimCueTraits_ReplicatedXOrPredicted
+{
+	static constexpr uint8 InvokeMask { (uint8)ESimulationTickContext::Authority | (uint8)ESimulationTickContext::Predict };
+	static constexpr ENetSimCueReplicationTarget ReplicationTarget { ENetSimCueReplicationTarget::Interpolators };
+};
+
+// Preset: Invoked and replicated to all. Uniqueness testing to avoid double playing etc
+struct TNetSimCueTraits_Strong
+{
+	static constexpr uint8 InvokeMask { (uint8)ESimulationTickContext::All };
+	static constexpr ENetSimCueReplicationTarget ReplicationTarget { ENetSimCueReplicationTarget::All };
+};
+
+// Actual trait struct that we use to look up traits. User cues must specialize this
+template<typename TCue>
+struct TCueHandlerTraits : public TNetSimCueTraitsBase
+{
+};
+
+// ----------------------------------------------------------------------------------------------------------------------------------------------
+//	Wrapper
+// ----------------------------------------------------------------------------------------------------------------------------------------------
+
+struct FNetSimCueWrapperBase
+{
+	virtual ~FNetSimCueWrapperBase() { }
+	virtual void NetSerialize(FArchive& Ar) = 0;
+	virtual bool Unique(const FNetSimCueWrapperBase* Other) const = 0;
+	virtual void* CueData() const = 0;
+	virtual ENetSimCueReplicationTarget GetReplicationTarget() const = 0;
+};
+
+template<typename TCue>
+struct TNetSimCueWrapper : FNetSimCueWrapperBase
+{
+	TNetSimCueWrapper() = default;
+	TNetSimCueWrapper(const TCue& Source)
+	{
+		Instance = Source; // Fixme: can we use move semantics to avoid this copy?
+	}
+
+	void NetSerialize(FArchive& Ar) override final
+	{
+		Instance.NetSerialize(Ar);
+	}
+
+	bool Unique(const FNetSimCueWrapperBase* Other) const override final
+	{
+		return TCue::Unique(Instance, *((TCue*)Other->CueData()));
+	}
+
+	void* CueData() const override final
+	{
+		return (void*)&Instance;
+	}
+
+	ENetSimCueReplicationTarget GetReplicationTarget() const override final
+	{
+		return TCueHandlerTraits<TCue>::ReplicationTarget;
+	}
+
+	TCue Instance;
+};
+
+// ----------------------------------------------------------------------------------------------------------------------------------------------
+// GlobalCueTypeTable: Cue types register with this to get a Type ID assigned. That ID is used in net serialization to talk about types.
+// ----------------------------------------------------------------------------------------------------------------------------------------------
+class FGlobalCueTypeTable
+{
+public:
+
+	NETWORKPREDICTION_API static FGlobalCueTypeTable& Get()
+	{
+		return Singleton;
+	}
+
+	template<typename TCue>
+	void RegisterType(const FString& TypeName)
+	{
+		ensure(!bFinalized);
+
+		FTypeInfo TypeInfo = 
+		{
+			&TCue::ID, 
+			[](){ return new TNetSimCueWrapper<TCue>();},
+			TypeName
+		};
+		PendingCueTypes.Emplace(TypeInfo);
+	}
+
+	void FinalizeTypes()
+	{
+		ensure(!bFinalized);
+		bFinalized = true;
+
+		PendingCueTypes.Sort([](const FTypeInfo& LHS, const FTypeInfo& RHS) { return LHS.TypeName < RHS.TypeName; });
+		int32 ID=0;
+		for (auto& TypeInfo: PendingCueTypes)
+		{
+			check(TypeInfo.IDPtr != nullptr);
+			*TypeInfo.IDPtr = ++ID;
+			TypeInfoMap.Add(*TypeInfo.IDPtr) = TypeInfo;
+		}
+		for (auto& Func : PostFinalizeCallbacks)
+		{
+			Func();
+		}
+	}
+	
+	bool HasBeenFinalized() const { return bFinalized; }
+	void AddPostFinalizedCallback(TFunction<void()> Callback) { PostFinalizeCallbacks.Emplace(Callback); }
+
+	// ---------------------------
+
+	FNetSimCueWrapperBase* Allocate(FNetSimCueTypeId ID)
+	{
+		return TypeInfoMap.FindChecked(ID).AllocateFunc();
+	}
+
+	FString GetTypeName(FNetSimCueTypeId ID) const
+	{
+		return TypeInfoMap.FindChecked(ID).TypeName;
+	}
+
+private:
+
+	struct FTypeInfo
+	{
+		FNetSimCueTypeId* IDPtr = nullptr;
+		TFunction<FNetSimCueWrapperBase*()> AllocateFunc;
+		FString TypeName;
+	};
+
+	static FGlobalCueTypeTable Singleton;
+
+	TArray<TFunction<void()>> PostFinalizeCallbacks;
+	TArray<FTypeInfo> PendingCueTypes;
+	bool bFinalized = false;
+
+	TMap<FNetSimCueTypeId, FTypeInfo> TypeInfoMap;
+};
+
+// ----------------------------------------------------------------------------------------------------------------------------------------------
+// SavedCue: a recorded Invocation of a NetSimCue
+// ----------------------------------------------------------------------------------------------------------------------------------------------
+
+struct FSavedCue
+{
+	FSavedCue() = default;
+
+	FSavedCue(const FSavedCue&) = delete;
+	FSavedCue& operator=(const FSavedCue&) = delete;
+
+	FSavedCue(FSavedCue&&) = default;
+	FSavedCue& operator=(FSavedCue&&) = default;
+
+	template<typename TCue>
+	FSavedCue(FNetworkSimTime InTime, const TCue& SourceCue) : Time(InTime)
+	{
+		ID = TCue::ID;
+		CueInstance.Reset(new TNetSimCueWrapper<TCue>(SourceCue)); // Fixme: can the copy be avoided with move semantics?
+	}
+
+	void NetSerialize(FArchive& Ar)
+	{
+		if (Ar.IsSaving())
+		{
+			check(CueInstance.IsValid());
+			Ar << ID;
+			CueInstance->NetSerialize(Ar);
+
+			Time.NetSerialize(Ar);
+		}
+		else
+		{
+			Ar << ID;
+			CueInstance.Reset(FGlobalCueTypeTable::Get().Allocate(ID));
+			CueInstance->NetSerialize(Ar);
+
+			Time.NetSerialize(Ar);
+		}
+	}
+
+	bool Unique(FSavedCue& OtherCue) const
+	{
+		if (ID != OtherCue.ID)
+		{
+			return false;
+		}
+
+		return CueInstance->Unique(OtherCue.CueInstance.Get());
+	}
+
+	FString GetTypeName() const
+	{
+		return FGlobalCueTypeTable::Get().GetTypeName(ID);
+	}
+
+	FNetSimCueTypeId ID = 0;
+	FNetworkSimTime Time;
+	TUniquePtr<FNetSimCueWrapperBase> CueInstance;
+	bool bDispatched = false;
+};
+
+
+// ----------------------------------------------------------------------------------------------------------------------------------------------
+// Per-Receiver dispatch table. This is how we go from a serialized ID to a function call
+// ----------------------------------------------------------------------------------------------------------------------------------------------
+template<typename TCueHandler>
+class TCueDispatchTable
+{
+public:
+	static TCueDispatchTable<TCueHandler>& Get()
+	{
+		return Singleton;
+	}
+
+	// All types that the receiver can handle must be registered here. This is where we create the TFunction to call ::HandleCue
+	template<typename TCue>
+	void RegisterType()
+	{
+		// Registeration has to be deferred until the FGlobalCueTypeTable is ready
+		auto Register = [this]() {
+			check(TCue::ID != 0);
+
+			FCueTypeInfo& CueTypeInfo = CueTypeInfoMap.Add(TCue::ID);
+
+			// The actual Dispatch func that gets invoked
+			CueTypeInfo.Dispatch = [](FNetSimCueWrapperBase* Cue, TCueHandler& Handler, const FNetworkSimTime& ElapsedTime)
+			{
+				Handler.HandleCue( *static_cast<TCue*>(Cue->CueData()), ElapsedTime );
+			};
+		};
+		
+		if (!FGlobalCueTypeTable::Get().HasBeenFinalized())
+		{
+			FGlobalCueTypeTable::Get().AddPostFinalizedCallback(Register);
+		}
+		else
+		{
+			Register();
+		}
+	}
+
+	void Dispatch(FSavedCue& SavedCue, TCueHandler& Handler, const FNetworkSimTime& ElapsedTime)
+	{
+		if (FCueTypeInfo* TypeInfo = CueTypeInfoMap.Find(SavedCue.ID))
+		{
+			check(TypeInfo->Dispatch);
+			TypeInfo->Dispatch(SavedCue.CueInstance.Get(), Handler, ElapsedTime);
+		}
+	}
+
+private:
+
+	static TCueDispatchTable<TCueHandler> Singleton;
+
+	struct FCueTypeInfo
+	{
+		TFunction<void(FNetSimCueWrapperBase* Cue, TCueHandler& Handler, const FNetworkSimTime& ElapsedTime)> Dispatch;
+	};
+
+	TMap<FNetSimCueTypeId, FCueTypeInfo> CueTypeInfoMap;
+};
+
+template<typename TCueHandler>
+TCueDispatchTable<TCueHandler> TCueDispatchTable<TCueHandler>::Singleton;
+
+// ------------------------------------------------------------------------------
+//	CueDispatcher
+//	-Entry point for invoking cues during a SimulationTick
+//	-Holds recorded cue state for replication
+// ------------------------------------------------------------------------------
+
+struct FCueDispatcher
+{
+	template<typename T>
+	void Invoke(T&& Cue)
+	{
+		if (EnsureValidContext())
+		{
+			if ((TCueHandlerTraits<T>::InvokeMask & (FNetSimCueTypeId)Context.TickContext) > 0)
+			{
+				// There is an implicit contract that we will invoke events going forward OR will receive explicit rollback notification (TODO)
+				ensure(SavedCues.Num() == 0 || SavedCues.Last().Time <= Context.CurrentSimTime);
+
+
+				bool bTransient = false;
+				if (Context.TickContext == ESimulationTickContext::Authority && TCueHandlerTraits<T>::ReplicationTarget == ENetSimCueReplicationTarget::None)
+				{
+					bTransient = true;
+				}
+
+				UE_LOG(LogNetSimCues, Log, TEXT("Invoking Cue %s. Transient: %d. Mask: %d. ReplicationTarget: %d"), *FGlobalCueTypeTable::Get().GetTypeName(T::ID), bTransient, TCueHandlerTraits<T>::InvokeMask, TCueHandlerTraits<T>::ReplicationTarget);
+
+				if (bTransient)
+				{
+					TransientCues.Emplace(Context.CurrentSimTime, Cue);
+				}
+				else
+				{
+					SavedCues.Emplace(Context.CurrentSimTime, Cue);
+				}
+			}
+			else
+			{
+				UE_LOG(LogNetSimCues, Log, TEXT("Suppressing Cue Invocation %s. Mask: %d. TickContext: %d"), *FGlobalCueTypeTable::Get().GetTypeName(T::ID), TCueHandlerTraits<T>::InvokeMask, (int32)Context.TickContext);
+			}
+		}
+	}
+
+	// Serializes the recorded cue record
+	// Note: this is not NetSerializing the individual cues in the record, it is serializing the record which contains the "net bits" of each individual cue
+	void NetSerializeSavedCues(FArchive& Ar, bool bIsInterpolatingSim)
+	{
+		FNetSimCueTypeId NumCues = SavedCues.Num();
+		Ar << NumCues;
+		
+		if (Ar.IsSaving())
+		{
+			for (FSavedCue& SavedCue : SavedCues)
+			{
+				SavedCue.NetSerialize(Ar);
+			}
+		}
+		else
+		{
+			int32 StartingNum = SavedCues.Num();
+
+			for (int32 CueIdx=0; CueIdx < NumCues; ++CueIdx)
+			{
+				FSavedCue SerializedCue;
+				SerializedCue.NetSerialize(Ar);
+
+				// Decide if we should accept the cue:
+
+				// ReplicationTarget: Cues can be set to only replicate to interpolators
+				if (SerializedCue.CueInstance->GetReplicationTarget() == ENetSimCueReplicationTarget::Interpolators && !bIsInterpolatingSim)
+				{
+					UE_LOG(LogNetSimCues, Log, TEXT("Discarding replicated NetSimCue %s intended for interpolators."), *SerializedCue.GetTypeName());
+					continue;
+				}
+
+				// Uniqueness: have we already received/predicted it?
+				bool bUniqueCue = true;
+				for (int32 ExistingIdx=0; ExistingIdx < StartingNum; ++ExistingIdx)
+				{
+					if (SerializedCue.Unique(SavedCues[ExistingIdx]) == false)
+					{
+						// These cues are not unique ("close enough") so we are skipping receiving this one
+						UE_LOG(LogNetSimCues, Log, TEXT("Discarding replicated NetSimCue %s because we've already processed it."), *SerializedCue.GetTypeName());
+						bUniqueCue = false;
+						break;
+					}
+				}
+
+				if (bUniqueCue)
+				{
+					SavedCues.Emplace(MoveTemp(SerializedCue));
+				}
+			}
+		}
+	}
+
+	template<typename T>
+	void DispatchCueRecord(T& Handler, FNetworkSimTime CurrentSimTime)
+	{
+		// Prune
+		FNetworkSimTime CutOffTime = CurrentSimTime - FNetworkSimTime::FromMSec(200); // fixme: hardcoded bad. But also, should take into consideration confirmed server frame in prediction case
+		int32 FirstValidElement = 0;
+		for (; FirstValidElement < SavedCues.Num(); ++FirstValidElement)
+		{
+			if ( SavedCues[FirstValidElement].Time > CutOffTime)
+			{
+				break;
+			}
+		}
+
+		if (FirstValidElement > 0)
+		{
+			SavedCues.RemoveAt(0, FirstValidElement, false);
+		}
+
+		// Dispatch
+		for (FSavedCue& SavedCue : SavedCues)
+		{
+			if (SavedCue.bDispatched == false)
+			{
+				UE_LOG(LogNetSimCues, Log, TEXT("Dispatching NetSimCue %s."), *SavedCue.GetTypeName());
+				SavedCue.bDispatched = true;
+				TCueDispatchTable<T>::Get().Dispatch(SavedCue, Handler, CurrentSimTime - SavedCue.Time);
+			}
+		}
+
+		for (FSavedCue& TransientCue : TransientCues)
+		{
+			UE_LOG(LogNetSimCues, Log, TEXT("Dispatching transient NetSimCue %s."), *TransientCue.GetTypeName());
+			TCueDispatchTable<T>::Get().Dispatch(TransientCue, Handler, CurrentSimTime - TransientCue.Time);
+		}
+		TransientCues.Reset();
+	}
+
+	// Sim Context: the Sim has to tell the dispatcher what its doing so that it can decide if it should supress Invocations or not
+	struct FContext
+	{
+		FNetworkSimTime CurrentSimTime;
+		ESimulationTickContext TickContext;
+	};
+
+	void PushContext(const FContext& InContext) { Context = InContext; }
+	void PopContext() { Context = FContext(); }
+
+private:
+
+	TArray<FSavedCue> SavedCues;		// Cues that must be saved for some period of time, either for replication or for uniqueness testing
+	TArray<FSavedCue> TransientCues;	// Cues that are dispatched on this frame and then forgotten about
+
+	FContext Context;
+
+	bool EnsureValidContext()
+	{
+		return ensure(Context.CurrentSimTime.IsPositive() && Context.TickContext != ESimulationTickContext::None);
+	}
+};
+
+// ------------------------------------------------------------------------------------------------------------------------------------------------------------------------------------------------
+//	Registration Helpers and Macros
+// ------------------------------------------------------------------------------------------------------------------------------------------------------------------------------------------------
+
+// Body of NetSimCue header. Gives type a static ID to be identified with
+#define NETSIMCUE_BODY() static FNetSimCueTypeId ID
+
+// ------------------------------------------------------------------------------------------------
+// Register a cue type with the global table.
+// ------------------------------------------------------------------------------------------------
+template<typename T>
+struct TNetSimCueTypeAutoRegisterHelper
+{
+	TNetSimCueTypeAutoRegisterHelper(const FString& Name)
+	{
+		FGlobalCueTypeTable::Get().RegisterType<T>(Name);
+	}
+};
+
+// Note this also defines the internal static ID for the cue type
+#define NETSIMCUE_REGISTER(X, STR) FNetSimCueTypeId X::ID=0; TNetSimCueTypeAutoRegisterHelper<X> NetSimCueAr_##X(STR);
+
+// ------------------------------------------------------------------------------------------------
+// Register a handler's cue types via a static "RegisterNetSimCueTypes" on the handler itself
+// ------------------------------------------------------------------------------------------------
+template<typename T>
+struct TNetSimCueHandlerAutoRegisterHelper
+{
+	TNetSimCueHandlerAutoRegisterHelper()
+	{
+		T::RegisterNetSimCueTypes( TCueDispatchTable<T>::Get() );
+	}
+};
+
+#define NETSIMCUEHANDLER_REGISTER(X) TNetSimCueHandlerAutoRegisterHelper<X> NetSimCueHandlerAr_##X;
+
+// ------------------------------------------------------------------------------------------------
+// Register a handler's cue types via an intermediate "set" class with a static "RegisterNetSimCueTypes" function
+// ------------------------------------------------------------------------------------------------
+template<typename THandler, typename TCueSet>
+struct TNetSimCueSetHandlerAutoRegisterHelper
+{
+	TNetSimCueSetHandlerAutoRegisterHelper()
+	{
+		TCueSet::RegisterNetSimCueTypes( TCueDispatchTable<THandler>::Get() );
+	}
+};
+
+#define NETSIMCUESET_REGISTER(THandler, TCueSet) TNetSimCueSetHandlerAutoRegisterHelper<THandler,TCueSet> NetSimCueSetHandlerAr_##THandler_##TCueSet;
+
diff --git a/Engine/Plugins/Runtime/NetworkPrediction/Source/NetworkPrediction/Public/NetworkedSimulationModelInterpolator.h b/Engine/Plugins/Runtime/NetworkPrediction/Source/NetworkPrediction/Public/NetworkedSimulationModelInterpolator.h
--- a/Engine/Plugins/Runtime/NetworkPrediction/Source/NetworkPrediction/Public/NetworkedSimulationModelInterpolator.h
+++ b/Engine/Plugins/Runtime/NetworkPrediction/Source/NetworkPrediction/Public/NetworkedSimulationModelInterpolator.h
@@ -8,6 +8,7 @@
 #include "NetworkPredictionTypes.h"
 #include "Engine/World.h"
 #include "NetworkedSimulationModelTypes.h"
+#include "NetworkedSimulationModelTick.h"
 
 DEFINE_LOG_CATEGORY_STATIC(LogNetInterpolation, Log, All);
 
diff --git a/Engine/Plugins/Runtime/NetworkPrediction/Source/NetworkPrediction/Public/NetworkedSimulationModelReplicators.h b/Engine/Plugins/Runtime/NetworkPrediction/Source/NetworkPrediction/Public/NetworkedSimulationModelReplicators.h
--- a/Engine/Plugins/Runtime/NetworkPrediction/Source/NetworkPrediction/Public/NetworkedSimulationModelReplicators.h
+++ b/Engine/Plugins/Runtime/NetworkPrediction/Source/NetworkPrediction/Public/NetworkedSimulationModelReplicators.h
@@ -441,6 +441,8 @@ struct TReplicator_Simulated : public TBase
 		SyncState->NetSerialize(Ar);
 		AuxState->NetSerialize(Ar);
 
+		Buffers.CueDispatcher.NetSerializeSavedCues(Ar, GetSimulatedUpdateMode() == ESimulatedUpdateMode::Interpolate);
+
 		if (Ar.IsLoading())
 		{
 			LastSerializedInputCmd = *InputCmd;
@@ -580,11 +582,11 @@ private:
 
 		// Do the actual update
 		{
-			TScopedSimulationTick UpdateScope(Ticker, OutputFrame, NewCmd->GetFrameDeltaTime());
+			TScopedSimulationTick UpdateScope(Ticker, Buffers.CueDispatcher, ESimulationTickContext::Resimulate, OutputFrame, NewCmd->GetFrameDeltaTime());
 			Simulation->SimulationTick( 
 				{ NewCmd->GetFrameDeltaTime(), Ticker },
 				{ *NewCmd, *PrevSyncState, *AuxState },
-				{ *NextSyncState, Buffers.Aux.LazyWriter(OutputFrame) } );
+				{ *NextSyncState, Buffers.Aux.LazyWriter(OutputFrame), Buffers.CueDispatcher } );
 		}
 
 		Ticker.MaxAllowedFrame = OutputFrame;
@@ -644,6 +646,8 @@ struct TReplicator_Autonomous : public TBase
 			SerializedAuxState.NetSerialize(Ar);
 		}
 
+		Buffers.CueDispatcher.NetSerializeSavedCues(Ar, false);
+
 		if (Ar.IsLoading())
 		{
 			bReconcileFaultDetected = false;
@@ -766,11 +770,11 @@ struct TReplicator_Autonomous : public TBase
 
 			// Do the actual update
 			{
-				TScopedSimulationTick UpdateScope(Ticker, OutputFrame, ResimulateCmd->GetFrameDeltaTime());
+				TScopedSimulationTick UpdateScope(Ticker, Buffers.CueDispatcher, ESimulationTickContext::Resimulate, OutputFrame, ResimulateCmd->GetFrameDeltaTime());
 				Simulation->SimulationTick( 
 					{ ResimulateCmd->GetFrameDeltaTime(), Ticker },
 					{ *ResimulateCmd, *PrevSyncState, *AuxState },
-					{ *NextSyncState, Buffers.Aux.LazyWriter(OutputFrame) } );
+					{ *NextSyncState, Buffers.Aux.LazyWriter(OutputFrame), Buffers.CueDispatcher } );
 			}
 
 			// Log out the newly predicted state that we got.
diff --git a/Engine/Plugins/Runtime/NetworkPrediction/Source/NetworkPrediction/Public/NetworkedSimulationModelTick.h b/Engine/Plugins/Runtime/NetworkPrediction/Source/NetworkPrediction/Public/NetworkedSimulationModelTick.h
new file mode 100644
--- /dev/null
+++ b/Engine/Plugins/Runtime/NetworkPrediction/Source/NetworkPrediction/Public/NetworkedSimulationModelTick.h
@@ -0,0 +1,210 @@
+// Copyright 1998-2019 Epic Games, Inc. All Rights Reserved.
+
+#pragma once
+
+#include "NetworkedSimulationModelTime.h"
+#include "NetworkedSimulationModelCues.h"
+
+
+// Holds per-simulation settings about how ticking is supposed to happen.
+template<int32 InFixedStepMS=0, int32 InMaxStepMS=0>
+struct TNetworkSimTickSettings
+{
+	static_assert(!(InFixedStepMS!=0 && InMaxStepMS != 0), "MaxStepMS is only applicable when using variable step (FixedStepMS == 0)");
+	enum 
+	{
+		MaxStepMS = InMaxStepMS,				// Max step. Only applicable to variable time step.
+		FixedStepMS = InFixedStepMS,			// Fixed step. If 0, then we are "variable time step"
+	};
+
+	// Typed accessors
+	static constexpr FNetworkSimTime::FSimTime GetMaxStepMS() { return static_cast<FNetworkSimTime::FSimTime>(InMaxStepMS); }
+	static constexpr FNetworkSimTime::FSimTime GetFixedStepMS() { return static_cast<FNetworkSimTime::FSimTime>(InFixedStepMS); }
+};
+
+// ----------------------------------------------------------------------------------------------------------------------------------------------
+//	Accumulator: Helper for accumulating real time into sim time based on TickSettings
+// ----------------------------------------------------------------------------------------------------------------------------------------------
+
+template<typename TickSettings, bool IsFixedTick=(TickSettings::FixedStepMS!=0)>
+struct TRealTimeAccumulator
+{
+	using TRealTime = FNetworkSimTime::FRealTime;
+	void Accumulate(FNetworkSimTime& NetworkSimTime, const TRealTime RealTimeSeconds)
+	{
+		// Even though we are variable tick, we still want to truncate down to an even msec. This keeps sim steps as whole integer values that serialize better, don't have denormals or other floating point weirdness.
+		// (If we ever wanted to just fully let floats pass through and be used as the msec sim time, this could be done through another specialization for float/float time).		
+		// Also note that MaxStepMS enforcement does NOT belong here. Dropping time due to MaxStepMS would just make the sim run slower. MaxStepMS is used at the input processing level.
+		
+		AccumulatedTimeMS += RealTimeSeconds * FNetworkSimTime::GetRealToSimFactor(); // convert input seconds -> msec
+		const FNetworkSimTime AccumulatedSimTimeMS = FNetworkSimTime::FromRealTimeMS(AccumulatedTimeMS);	// truncate (float) accumulated msec to (int32) sim time msec
+
+		NetworkSimTime += AccumulatedSimTimeMS;
+		AccumulatedTimeMS -= AccumulatedSimTimeMS.ToRealTimeMS(); // subtract out the "whole" msec, we are left with the remainder msec
+	}
+
+	void Reset()
+	{
+		AccumulatedTimeMS = 0.f;
+	}
+
+private:
+
+	TRealTime AccumulatedTimeMS = 0.f;
+};
+
+// Specialized version of FixedTicking. This accumulates real time that spills over into NetworkSimTime as it crosses the FixStep threshold
+template<typename TickSettings>
+struct TRealTimeAccumulator<TickSettings, true>
+{
+	using TRealTime = FNetworkSimTime::FRealTime;
+	const TRealTime RealTimeFixedStep = static_cast<TRealTime>(TickSettings::GetFixedStepMS() * FNetworkSimTime::GetSimToRealFactor());
+
+	void Accumulate(FNetworkSimTime& NetworkSimTime, const TRealTime RealTimeSeconds)
+	{
+		AccumulatedTime += RealTimeSeconds;
+		if (AccumulatedTime > RealTimeFixedStep)
+		{
+			const int32 NumFrames = AccumulatedTime / RealTimeFixedStep;
+			AccumulatedTime -= NumFrames * RealTimeFixedStep;
+				
+			if (FMath::Abs<TRealTime>(AccumulatedTime) < SMALL_NUMBER)
+			{
+				AccumulatedTime = TRealTime(0.f);
+			}
+
+			NetworkSimTime += FNetworkSimTime::FromMSec(NumFrames * TickSettings::FixedStepMS);
+		}
+	}
+
+	void Reset()
+	{
+		AccumulatedTime = 0.f;
+	}
+
+private:
+
+	TRealTime AccumulatedTime;
+};
+
+// ----------------------------------------------------------------------------------------------------------------------------------------------
+//	FSimulationTickState: Holds active state for simulation ticking. We track two things: frames and time.
+//
+//	PendingFrame is the next frame we will process: the input/sync/aux state @ PendingFrame will be run through ::SimulationTick and produce the 
+//	next frame's (PendingFrame+1) Sync and possibly Aux state (if it changes). "Out of band" modifications to the sync/aux state should happen
+//	to PendingFrame (e.g, before it is processed. Once a frame has been processed, we won't run it through ::SimulationTick again!).
+//
+//	MaxAllowedFrame is a frame based limiter on simulation updates. This must be incremented to allow the simulation to advance.
+//
+//	Time is also tracked. We keep running total for how much the sim has advanced and how much it is allowed to advance. There is also a historic buffer of
+//	simulation time in SimulationTimeBuffer.
+//
+//	Consider that Frames are essentially client dependent and gaps can happen due to packet loss, etc. Time will always be continuous though.
+//	
+// ----------------------------------------------------------------------------------------------------------------------------------------------
+
+struct FSimulationTickState
+{
+	int32 PendingFrame = 0;
+	int32 MaxAllowedFrame = -1;
+	bool bUpdateInProgress = false;
+	
+	FNetworkSimTime GetTotalProcessedSimulationTime() const 
+	{ 
+		return TotalProcessedSimulationTime; 
+	}
+
+	void SetTotalProcessedSimulationTime(const FNetworkSimTime& SimTime, int32 Frame)
+	{
+		TotalProcessedSimulationTime = SimTime;
+		*SimulationTimeBuffer.WriteFrame(Frame) = SimTime;
+	}
+
+	void IncrementTotalProcessedSimulationTime(const FNetworkSimTime& DeltaSimTime, int32 Frame)
+	{
+		TotalProcessedSimulationTime += DeltaSimTime;
+		*SimulationTimeBuffer.WriteFrame(Frame) = TotalProcessedSimulationTime;
+	}
+		
+	// Historic tracking of simulation time. This allows us to timestamp sync data as its produced
+	TNetworkSimContiguousBuffer<FNetworkSimTime>	SimulationTimeBuffer;
+
+	// How much granted simulation time is left to process
+	FNetworkSimTime GetRemainingAllowedSimulationTime() const
+	{
+		return TotalAllowedSimulationTime - TotalProcessedSimulationTime;
+	}
+
+	FNetworkSimTime GetTotalAllowedSimulationTime() const
+	{
+		return TotalAllowedSimulationTime;
+	}
+
+protected:
+
+	FNetworkSimTime TotalAllowedSimulationTime;	// Total time we have been "given" to process. We cannot process more simulation time than this: doing so would be speed hacking.
+	FNetworkSimTime TotalProcessedSimulationTime;	// How much time we've actually processed. The only way to increment this is to process user commands or receive authoritative state from the network.
+};
+
+// "Ticker" that actually allows us to give the simulation time. This struct will generally not be passed around outside of the core TNetworkedSimulationModel/Replicators
+template<typename TickSettings=TNetworkSimTickSettings<>>
+struct TSimulationTicker : public FSimulationTickState
+{
+	using TSettings = TickSettings;
+
+	void SetTotalAllowedSimulationTime(const FNetworkSimTime& SimTime)
+	{
+		TotalAllowedSimulationTime = SimTime;
+		RealtimeAccumulator.Reset();
+	}
+
+	// "Grants" allowed simulation time to this tick state. That is, we are now allowed to advance the simulation by this amount the next time the sim ticks.
+	// Note the input is RealTime in SECONDS. This is what the rest of the engine uses when dealing with float delta time.
+	void GiveSimulationTime(float RealTimeSeconds)
+	{
+		RealtimeAccumulator.Accumulate(TotalAllowedSimulationTime, RealTimeSeconds);
+	}	
+
+private:
+
+	TRealTimeAccumulator<TSettings>	RealtimeAccumulator;	
+};
+
+// Scoped helper to be used right before entering a call to the sim's ::SimulationTick function.
+// Important to note that this advances the PendingFrame to the output Frame. So that any writes that occur to the buffers during this scope will go to the output frame.
+struct TScopedSimulationTick
+{
+	TScopedSimulationTick(FSimulationTickState& InTicker, FCueDispatcher& InDispatcher, ESimulationTickContext TickContext, const int32& InOutputFrame, const FNetworkSimTime& InDeltaSimTime)
+		: Ticker(InTicker), Dispatcher(InDispatcher), OutputFrame(InOutputFrame), DeltaSimTime(InDeltaSimTime)
+	{
+		check(Ticker.bUpdateInProgress == false);
+		Ticker.PendingFrame = OutputFrame;
+		Ticker.bUpdateInProgress = true;
+
+		Dispatcher.PushContext({Ticker.GetTotalProcessedSimulationTime() + DeltaSimTime, TickContext}); // Cues "take place" at the end of the frame
+	}
+	~TScopedSimulationTick()
+	{
+		Ticker.IncrementTotalProcessedSimulationTime(DeltaSimTime, OutputFrame);
+		Ticker.bUpdateInProgress = false;
+
+		Dispatcher.PopContext();
+	}
+	FSimulationTickState& Ticker;
+	FCueDispatcher& Dispatcher;
+	const int32& OutputFrame;
+	const FNetworkSimTime& DeltaSimTime;
+};
+
+// ----------------------------------------------------------------------------------------------------------------------------------------------
+//	FNetSimTimeStep Time data structure that is passed into SimulationTick
+// ----------------------------------------------------------------------------------------------------------------------------------------------
+
+struct FNetSimTimeStep
+{
+	// The delta time step for this tick (in MS by default)
+	const FNetworkSimTime& StepMS;
+	// The tick state of the simulation prior to running this tick. E.g, does not "include" the above StepMS that we are simulating now.
+	// The first time ::SimulationTick runs, TickState.GetTotalProcessedSimulationTime() == 0.
+	const FSimulationTickState& TickState;
+};
\ No newline at end of file
diff --git a/Engine/Plugins/Runtime/NetworkPrediction/Source/NetworkPrediction/Public/NetworkedSimulationModelTime.h b/Engine/Plugins/Runtime/NetworkPrediction/Source/NetworkPrediction/Public/NetworkedSimulationModelTime.h
new file mode 100644
--- /dev/null
+++ b/Engine/Plugins/Runtime/NetworkPrediction/Source/NetworkPrediction/Public/NetworkedSimulationModelTime.h
@@ -0,0 +1,68 @@
+// Copyright 1998-2019 Epic Games, Inc. All Rights Reserved.
+
+#pragma once
+
+#include "NetworkedSimulationModelBuffer.h"
+
+// ----------------------------------------------------------------------------------------------------------------------------------------------
+//		Tick and time keeping related structures
+// ----------------------------------------------------------------------------------------------------------------------------------------------
+
+// The main Simulation time type. All sims use this to talk about time.
+struct FNetworkSimTime
+{
+	using FSimTime = int32; // Underlying type used to store simulation time
+	using FRealTime = float; // Underlying type used when dealing with realtime (usually coming in from the engine tick).
+
+	enum
+	{
+		RealToSimFactor = 1000 // Factor to go from RealTime (always seconds) to SimTime (MSec by default with factor of 1000)
+	};
+
+	static constexpr FRealTime GetRealToSimFactor() { return static_cast<FRealTime>(RealToSimFactor); }
+	static constexpr FRealTime GetSimToRealFactor() { return static_cast<FRealTime>(1.f / RealToSimFactor); }
+
+	// ---------------------------------------------------------------
+
+	FNetworkSimTime() { }
+
+	// Things get confusing with templated types and overloaded functions. To avoid that, use these funcs to construct from either msec or real time
+	static inline FNetworkSimTime FromMSec(const FSimTime& InTime) { return FNetworkSimTime(InTime); } 
+	static inline FNetworkSimTime FromRealTimeSeconds(const FRealTime& InRealTime) { return FNetworkSimTime( static_cast<FSimTime>(InRealTime * GetRealToSimFactor())); } 
+	FRealTime ToRealTimeSeconds() const { return (Time * GetSimToRealFactor()); }
+
+	// Direct casts to "real time MS" which should be rarely used in practice (TRealTimeAccumulator only current case). All other cases of "real time" imply seconds.
+	static inline FNetworkSimTime FromRealTimeMS(const FRealTime& InRealTime) { return FNetworkSimTime( static_cast<FSimTime>(InRealTime)); } 
+	FRealTime ToRealTimeMS() const { return static_cast<FRealTime>(Time); }
+
+	FString ToString() const { return LexToString(this->Time); }
+
+	bool IsPositive() const { return (Time > 0); }
+	bool IsNegative() const { return (Time < 0); }
+	void Reset() { Time = 0; }
+
+	// FIXME
+	void NetSerialize(FArchive& Ar) { Ar << Time; }
+
+	operator FSimTime() const { return Time; }
+
+	using T = FNetworkSimTime;
+	T& operator+= (const T &rhs) { this->Time += rhs.Time; return(*this); }
+	T& operator-= (const T &rhs) { this->Time -= rhs.Time; return(*this); }
+	
+	T operator+ (const T &rhs) const { return T(this->Time + rhs.Time); }
+	T operator- (const T &rhs) const { return T(this->Time - rhs.Time); }
+
+	bool operator<  (const T &rhs) const { return(this->Time < rhs.Time); }
+	bool operator<= (const T &rhs) const { return(this->Time <= rhs.Time); }
+	bool operator>  (const T &rhs) const { return(this->Time > rhs.Time); }
+	bool operator>= (const T &rhs) const { return(this->Time >= rhs.Time); }
+	bool operator== (const T &rhs) const { return(this->Time == rhs.Time); }
+	bool operator!= (const T &rhs) const { return(this->Time != rhs.Time); }
+
+private:
+
+	FSimTime Time = 0;
+
+	FNetworkSimTime(const FSimTime& InTime) { Time = InTime; }
+};
\ No newline at end of file
diff --git a/Engine/Plugins/Runtime/NetworkPrediction/Source/NetworkPrediction/Public/NetworkedSimulationModelTypes.h b/Engine/Plugins/Runtime/NetworkPrediction/Source/NetworkPrediction/Public/NetworkedSimulationModelTypes.h
--- a/Engine/Plugins/Runtime/NetworkPrediction/Source/NetworkPrediction/Public/NetworkedSimulationModelTypes.h
+++ b/Engine/Plugins/Runtime/NetworkPrediction/Source/NetworkPrediction/Public/NetworkedSimulationModelTypes.h
@@ -3,6 +3,8 @@
 #pragma once
 
 #include "NetworkedSimulationModelBuffer.h"
+#include "NetworkedSimulationModelCues.h"
+#include "NetworkedSimulationModelTime.h"
 
 // ---------------------------------------------------------------------------------------------------------------------
 //	BufferTypes: template helpers for addressing the different buffer types of the system.
@@ -249,7 +251,9 @@ struct TNetworkSimBufferContainer
 	TInputBuffer Input;
 	TSyncBuffer Sync;
 	TAuxBuffer Aux;
-	TDebugBuffer Debug;	
+	TDebugBuffer Debug;
+
+	FCueDispatcher CueDispatcher;
 
 	// Finally, template accessor for getting buffers based on enum. This is really what all this junk is about.
 	// This allows other templated classes in the system to access a specific buffer from another templated argument
@@ -260,254 +264,6 @@ struct TNetworkSimBufferContainer
 	}
 };
 
-// ----------------------------------------------------------------------------------------------------------------------------------------------
-//		Tick and time keeping related structures
-// ----------------------------------------------------------------------------------------------------------------------------------------------
-
-// The main Simulation time type. All sims use this to talk about time.
-struct FNetworkSimTime
-{
-	using FSimTime = int32; // Underlying type used to store simulation time
-	using FRealTime = float; // Underlying type used when dealing with realtime (usually coming in from the engine tick).
-
-	enum
-	{
-		RealToSimFactor = 1000 // Factor to go from RealTime (always seconds) to SimTime (MSec by default with factor of 1000)
-	};
-
-	static constexpr FRealTime GetRealToSimFactor() { return static_cast<FRealTime>(RealToSimFactor); }
-	static constexpr FRealTime GetSimToRealFactor() { return static_cast<FRealTime>(1.f / RealToSimFactor); }
-
-	// ---------------------------------------------------------------
-
-	FNetworkSimTime() { }
-
-	// Things get confusing with templated types and overloaded functions. To avoid that, use these funcs to construct from either msec or real time
-	static inline FNetworkSimTime FromMSec(const FSimTime& InTime) { return FNetworkSimTime(InTime); } 
-	static inline FNetworkSimTime FromRealTimeSeconds(const FRealTime& InRealTime) { return FNetworkSimTime( static_cast<FSimTime>(InRealTime * GetRealToSimFactor())); } 
-	FRealTime ToRealTimeSeconds() const { return (Time * GetSimToRealFactor()); }
-
-	// Direct casts to "real time MS" which should be rarely used in practice (TRealTimeAccumulator only current case). All other cases of "real time" imply seconds.
-	static inline FNetworkSimTime FromRealTimeMS(const FRealTime& InRealTime) { return FNetworkSimTime( static_cast<FSimTime>(InRealTime)); } 
-	FRealTime ToRealTimeMS() const { return static_cast<FRealTime>(Time); }
-
-	FString ToString() const { return LexToString(this->Time); }
-
-	bool IsPositive() const { return (Time > 0); }
-	bool IsNegative() const { return (Time < 0); }
-	void Reset() { Time = 0; }
-
-	// FIXME
-	void NetSerialize(FArchive& Ar) { Ar << Time; }
-
-	operator FSimTime() const { return Time; }
-
-	using T = FNetworkSimTime;
-	T& operator+= (const T &rhs) { this->Time += rhs.Time; return(*this); }
-	T& operator-= (const T &rhs) { this->Time -= rhs.Time; return(*this); }
-	
-	T operator+ (const T &rhs) const { return T(this->Time + rhs.Time); }
-	T operator- (const T &rhs) const { return T(this->Time - rhs.Time); }
-
-	bool operator<  (const T &rhs) const { return(this->Time < rhs.Time); }
-	bool operator<= (const T &rhs) const { return(this->Time <= rhs.Time); }
-	bool operator>  (const T &rhs) const { return(this->Time > rhs.Time); }
-	bool operator>= (const T &rhs) const { return(this->Time >= rhs.Time); }
-	bool operator== (const T &rhs) const { return(this->Time == rhs.Time); }
-	bool operator!= (const T &rhs) const { return(this->Time != rhs.Time); }
-
-private:
-
-	FSimTime Time = 0;
-
-	FNetworkSimTime(const FSimTime& InTime) { Time = InTime; }
-};
-
-// Holds per-simulation settings about how ticking is supposed to happen.
-template<int32 InFixedStepMS=0, int32 InMaxStepMS=0>
-struct TNetworkSimTickSettings
-{
-	static_assert(!(InFixedStepMS!=0 && InMaxStepMS != 0), "MaxStepMS is only applicable when using variable step (FixedStepMS == 0)");
-	enum 
-	{
-		MaxStepMS = InMaxStepMS,				// Max step. Only applicable to variable time step.
-		FixedStepMS = InFixedStepMS,			// Fixed step. If 0, then we are "variable time step"
-	};
-
-	// Typed accessors
-	static constexpr FNetworkSimTime::FSimTime GetMaxStepMS() { return static_cast<FNetworkSimTime::FSimTime>(InMaxStepMS); }
-	static constexpr FNetworkSimTime::FSimTime GetFixedStepMS() { return static_cast<FNetworkSimTime::FSimTime>(InFixedStepMS); }
-};
-
-// ----------------------------------------------------------------------------------------------------------------------------------------------
-//	Accumulator: Helper for accumulating real time into sim time based on TickSettings
-// ----------------------------------------------------------------------------------------------------------------------------------------------
-
-template<typename TickSettings, bool IsFixedTick=(TickSettings::FixedStepMS!=0)>
-struct TRealTimeAccumulator
-{
-	using TRealTime = FNetworkSimTime::FRealTime;
-	void Accumulate(FNetworkSimTime& NetworkSimTime, const TRealTime RealTimeSeconds)
-	{
-		// Even though we are variable tick, we still want to truncate down to an even msec. This keeps sim steps as whole integer values that serialize better, don't have denormals or other floating point weirdness.
-		// (If we ever wanted to just fully let floats pass through and be used as the msec sim time, this could be done through another specialization for float/float time).		
-		// Also note that MaxStepMS enforcement does NOT belong here. Dropping time due to MaxStepMS would just make the sim run slower. MaxStepMS is used at the input processing level.
-		
-		AccumulatedTimeMS += RealTimeSeconds * FNetworkSimTime::GetRealToSimFactor(); // convert input seconds -> msec
-		const FNetworkSimTime AccumulatedSimTimeMS = FNetworkSimTime::FromRealTimeMS(AccumulatedTimeMS);	// truncate (float) accumulated msec to (int32) sim time msec
-
-		NetworkSimTime += AccumulatedSimTimeMS;
-		AccumulatedTimeMS -= AccumulatedSimTimeMS.ToRealTimeMS(); // subtract out the "whole" msec, we are left with the remainder msec
-	}
-
-	void Reset()
-	{
-		AccumulatedTimeMS = 0.f;
-	}
-
-private:
-
-	TRealTime AccumulatedTimeMS = 0.f;
-};
-
-// Specialized version of FixedTicking. This accumulates real time that spills over into NetworkSimTime as it crosses the FixStep threshold
-template<typename TickSettings>
-struct TRealTimeAccumulator<TickSettings, true>
-{
-	using TRealTime = FNetworkSimTime::FRealTime;
-	const TRealTime RealTimeFixedStep = static_cast<TRealTime>(TickSettings::GetFixedStepMS() * FNetworkSimTime::GetSimToRealFactor());
-
-	void Accumulate(FNetworkSimTime& NetworkSimTime, const TRealTime RealTimeSeconds)
-	{
-		AccumulatedTime += RealTimeSeconds;
-		if (AccumulatedTime > RealTimeFixedStep)
-		{
-			const int32 NumFrames = AccumulatedTime / RealTimeFixedStep;
-			AccumulatedTime -= NumFrames * RealTimeFixedStep;
-				
-			if (FMath::Abs<TRealTime>(AccumulatedTime) < SMALL_NUMBER)
-			{
-				AccumulatedTime = TRealTime(0.f);
-			}
-
-			NetworkSimTime += FNetworkSimTime::FromMSec(NumFrames * TickSettings::FixedStepMS);
-		}
-	}
-
-	void Reset()
-	{
-		AccumulatedTime = 0.f;
-	}
-
-private:
-
-	TRealTime AccumulatedTime;
-};
-
-// ----------------------------------------------------------------------------------------------------------------------------------------------
-//	FSimulationTickState: Holds active state for simulation ticking. We track two things: frames and time.
-//
-//	PendingFrame is the next frame we will process: the input/sync/aux state @ PendingFrame will be run through ::SimulationTick and produce the 
-//	next frame's (PendingFrame+1) Sync and possibly Aux state (if it changes). "Out of band" modifications to the sync/aux state should happen
-//	to PendingFrame (e.g, before it is processed. Once a frame has been processed, we won't run it through ::SimulationTick again!).
-//
-//	MaxAllowedFrame is a frame based limiter on simulation updates. This must be incremented to allow the simulation to advance.
-//
-//	Time is also tracked. We keep running total for how much the sim has advanced and how much it is allowed to advance. There is also a historic buffer of
-//	simulation time in SimulationTimeBuffer.
-//
-//	Consider that Frames are essentially client dependent and gaps can happen due to packet loss, etc. Time will always be continuous though.
-//	
-// ----------------------------------------------------------------------------------------------------------------------------------------------
-
-struct FSimulationTickState
-{
-	int32 PendingFrame = 0;
-	int32 MaxAllowedFrame = -1;
-	bool bUpdateInProgress = false;
-	
-	FNetworkSimTime GetTotalProcessedSimulationTime() const 
-	{ 
-		return TotalProcessedSimulationTime; 
-	}
-
-	void SetTotalProcessedSimulationTime(const FNetworkSimTime& SimTime, int32 Frame)
-	{
-		TotalProcessedSimulationTime = SimTime;
-		*SimulationTimeBuffer.WriteFrame(Frame) = SimTime;
-	}
-
-	void IncrementTotalProcessedSimulationTime(const FNetworkSimTime& DeltaSimTime, int32 Frame)
-	{
-		TotalProcessedSimulationTime += DeltaSimTime;
-		*SimulationTimeBuffer.WriteFrame(Frame) = TotalProcessedSimulationTime;
-	}
-		
-	// Historic tracking of simulation time. This allows us to timestamp sync data as its produced
-	TNetworkSimContiguousBuffer<FNetworkSimTime>	SimulationTimeBuffer;
-
-	// How much granted simulation time is left to process
-	FNetworkSimTime GetRemainingAllowedSimulationTime() const
-	{
-		return TotalAllowedSimulationTime - TotalProcessedSimulationTime;
-	}
-
-	FNetworkSimTime GetTotalAllowedSimulationTime() const
-	{
-		return TotalAllowedSimulationTime;
-	}
-
-protected:
-
-	FNetworkSimTime TotalAllowedSimulationTime;	// Total time we have been "given" to process. We cannot process more simulation time than this: doing so would be speed hacking.
-	FNetworkSimTime TotalProcessedSimulationTime;	// How much time we've actually processed. The only way to increment this is to process user commands or receive authoritative state from the network.
-};
-
-// "Ticker" that actually allows us to give the simulation time. This struct will generally not be passed around outside of the core TNetworkedSimulationModel/Replicators
-template<typename TickSettings=TNetworkSimTickSettings<>>
-struct TSimulationTicker : public FSimulationTickState
-{
-	using TSettings = TickSettings;
-
-	void SetTotalAllowedSimulationTime(const FNetworkSimTime& SimTime)
-	{
-		TotalAllowedSimulationTime = SimTime;
-		RealtimeAccumulator.Reset();
-	}
-
-	// "Grants" allowed simulation time to this tick state. That is, we are now allowed to advance the simulation by this amount the next time the sim ticks.
-	// Note the input is RealTime in SECONDS. This is what the rest of the engine uses when dealing with float delta time.
-	void GiveSimulationTime(float RealTimeSeconds)
-	{
-		RealtimeAccumulator.Accumulate(TotalAllowedSimulationTime, RealTimeSeconds);
-	}	
-
-private:
-
-	TRealTimeAccumulator<TSettings>	RealtimeAccumulator;	
-};
-
-// Scoped helper to be used right before entering a call to the sim's ::SimulationTick function.
-// Important to note that this advances the PendingFrame to the output Frame. So that any writes that occur to the buffers during this scope will go to the output frame.
-struct TScopedSimulationTick
-{
-	TScopedSimulationTick(FSimulationTickState& InTicker, const int32& InOutputFrame, const FNetworkSimTime& InDeltaSimTime)
-		: Ticker(InTicker), OutputFrame(InOutputFrame), DeltaSimTime(InDeltaSimTime)
-	{
-		check(Ticker.bUpdateInProgress == false);
-		Ticker.PendingFrame = OutputFrame;
-		Ticker.bUpdateInProgress = true;
-	}
-	~TScopedSimulationTick()
-	{
-		Ticker.IncrementTotalProcessedSimulationTime(DeltaSimTime, OutputFrame);
-		Ticker.bUpdateInProgress = false;
-	}
-	FSimulationTickState& Ticker;
-	const int32& OutputFrame;
-	const FNetworkSimTime& DeltaSimTime;
-};
-
 // ----------------------------------------------------------------------------------------------------------------------------------------------
 //	Accessors - helper structs that provide safe/cleaner access to the underlying NetSim states/events
 //	This is the equivlent of directly calling TNetworkedSimulationModel::GetPendingStateRead() / GetPendingStateWrite()
@@ -573,26 +329,6 @@ private:
 	TFunction<void(bool bForWrite, TState*& OutState, bool& OutSafe)> GetStateFunc;
 };
 
-//TODO
-struct FNetworkSimEventAccessor
-{
-	int32 GetPendingFrame() { return -1; }
-
-	struct FFrameEvents
-	{
-		/** This frame has been received from the authority and will not be rolled back or resimulated again. */
-		DECLARE_MULTICAST_DELEGATE(FOnFrameConfirmed)
-		FOnFrameConfirmed Confirmed;
-
-		/** This frame was previously simulated and the simulation has now been rolled back */
-		DECLARE_MULTICAST_DELEGATE(FOnFrameRolledBack)
-		FOnFrameRolledBack RolledBack;
-
-		DECLARE_MULTICAST_DELEGATE(FOnFrameSimulated)
-		FOnFrameRolledBack Simulated;
-	};
-};
-
 // ----------------------------------------------------------------------------------------------------------------------------------------------
 //	FrameCmd - in variable tick simulations we store the timestep of each frame with the input. TFrameCmd wraps the user struct to do this.
 // ----------------------------------------------------------------------------------------------------------------------------------------------
@@ -665,18 +401,9 @@ public:
 };
 
 // ----------------------------------------------------------------------------------------------------------------------------------------------
-//	SimulationTick parameters. These are the data structures passed into the simulation code each frame
+//	Simulation Input/Output parameters. These are the data structures passed into the simulation code each frame
 // ----------------------------------------------------------------------------------------------------------------------------------------------
 
-struct TNetSimTimeStep
-{
-	// The delta time step for this tick (in MS by default)
-	const FNetworkSimTime& StepMS;
-	// The tick state of the simulation prior to running this tick. E.g, does not "include" the above StepMS that we are simulating now.
-	// The first time ::SimulationTick runs, TickState.GetTotalProcessedSimulationTime() == 0.
-	const FSimulationTickState& TickState;
-};
-
 // Input state: const references to the InputCmd/SyncState/AuxStates
 template<typename TBufferTypes>
 struct TNetSimInput
@@ -700,8 +427,9 @@ struct TNetSimOutput
 {
 	typename TBufferTypes::TSyncState& Sync;
 	const TNetSimLazyWriter<typename TBufferTypes::TAuxState>& Aux;
+	FCueDispatcher& CueDispatch;
 
-	TNetSimOutput(typename TBufferTypes::TSyncState& InSync, const TNetSimLazyWriter<typename TBufferTypes::TAuxState>& InAux)
-		: Sync(InSync), Aux(InAux) { }
+	TNetSimOutput(typename TBufferTypes::TSyncState& InSync, const TNetSimLazyWriter<typename TBufferTypes::TAuxState>& InAux, FCueDispatcher& InCueDispatch)
+		: Sync(InSync), Aux(InAux), CueDispatch(InCueDispatch) { }
 };
 
diff --git a/Engine/Plugins/Runtime/NetworkPrediction/readme.txt b/Engine/Plugins/Runtime/NetworkPrediction/readme.txt
--- a/Engine/Plugins/Runtime/NetworkPrediction/readme.txt
+++ b/Engine/Plugins/Runtime/NetworkPrediction/readme.txt
@@ -108,10 +108,17 @@ Smoothing: Taking the output of the simulation and applying an additional layer
 // Release notes
 // ----------------------------------------------------------------------------------------------------------
 
+Update (11-15-19)
+-NetSimCues initial check in
+-These are the async/deferred events that are outputted during the network simulation tick
+-Different classification in types wrt prediction and replication
+-See FMockAbilityBlinkCue in code, "Misc Notes" at bottom of readme file
+-"Rewindable NetSimCue" api still todo
+
 Update (10-28-19)
 Renaming/cleanup on ::Update function:
 -"TSimulation::Update" is now "TSimulation::SimulationTick". Makes searching through source a bit easier, less generic name
--Collapsed parameters of SimulationTick into 3 containers: TNetSimTimeStep, TNetSimInput, TNetSimOutput
+-Collapsed parameters of SimulationTick into 3 containers: FNetSimTimeStep, TNetSimInput, TNetSimOutput
 -Time passed into simulation is now FNetworkSimTime (instead of forcing float RealTimeSeconds conversion at callsite)
 -TSimulationTickState broken up into FSimulationTickState and TSimulationTicker. 
 -::SimulationTick now also gets a reference to FSimulationTickState. This is setting us up to do timers/countdowns based on sim time within these functions.
@@ -189,7 +196,65 @@ Current State (7-25-19)
 -Flying Movement: a port of FloatingPawnMovement into the new system. Essentially just a basic flying movement system.
 -MockNetworkSimulation: a very basic minimal "simulation" that demonstrates how the system works. Not tied to movement or anything physical: just a running counter/sum.
  
- 
- 
 
+// ----------------------------------------------------------------------------------------------------------
+// Misc Notes
+// ----------------------------------------------------------------------------------------------------------
+
+Notes on "State Transitions"
+-Detecting changes to sync/aux state in FinalizeFrame
+-Can emit events from these transitions
+-Events can't have additional data (thats not in the sync/aux state)
+-Can't really know for sure if transition "just actually happnened in the simulation" vs "I'm just finding out about it now"
+-Biggest advantage: "Reliable", doesn't require SimulationTick to run. (Interpolated proxies can still use them)
+
+Notes on "Events" --> NetSimCues
+
+Basic idea/notes:
+-"Non simulation affecting" events that are emitted during the ::SimulationTick function.
+-Cues are *handled* (by user code) during actor tick. Not during ::SimulationTick. (they are queued in a buffer)
+-Must support arbitrary data payloads. E.g, "Impact Velocity" for a "Landed" cue.
+-Fundamentally unreliable: join in progress, relevancy, temporary network disconnect = missed events.
+	-Seriously! These events should never be used to "track state" at some higher level. Even if simulation code guaruntees 1:1 On/Off events, you may miss them!
+	-Note: see "State transitions" above for how "reliable" events would be done
+
+Implementation progression:
+
+0. Niave approach is to let simulation code call directly into user code at anytime to do this stuff
+	Major issue: rollback/resimulate will double play events	
+	Interpolating proxies would not invoke because they do not run the simulation
+
+[+Simple Invocation Suppressiong mechanism/context]
+1. Simplest approach would be "Weak Cues": "Cues don't replicate. Are only played on 'latest' simulation ticks (supressed during any rewind/resimulate scope)"
+	This causes/leaves two main issues/needs:
+		A. Interpolating simulated proxies (who don't run TickSimulation) will never invoke these cues. 
+			(Note: could just say "yup, thats how it works" and leave replication up to the user at a higher level. Not great though!)
+		B. Predicting Client can mispredict: missed events or can play wrong events. Won't catch this / don't care.
+
+[+Cue type traits/policies, build on invocation supression above. Cue data doesn't matter yet, just local context and cue type]
+[+Time information passed to cue handler]
+2. Next simplest approach would be to add two more cues types (in addition to "Weak Cues"):
+	A. Replicated, non predicted. Always comes from server, never predicted. Everyone gets them.
+		-Downside: we won't predict them since we can't guard against double playing yet.
+		-Downside: cues can fire off "in the past" and now we must include time information to Cue handler (let them figure out how to deal with time make up)
+	B. Replicated (to sim proxy only), predicted (by auto proxy only).
+		-Allow auto proxy to treat them as Weak Cues: always play on simulate (not resimulate), always ignore server replication (so, can mispredict)
+		-Sim proxy will have them replicated explicitly	
+
+[+Cue identification/uniqueness function]
+3. Next complicated approach would be to have identifying/uniqueness function to determine cues that were predicted or not. "Strong Cues"
+	Keep buffer of last N msecs of cues that were invoked locally.
+	When receiving replicated cues, look for matches that are "close enough" (hard to define!). If we already predicted, then ignore from server.
+	+This allows us to combine #2's cue types into a single type: Replicated + predicted (universal)		
+	+Can also now invoke while resimulating by checking if a matching cue was already invoked during the original predicted simulate.
+	
+[+Cue rollback/resimulate callbacks]
+4. Final complicated approach: full rollback/resimulate aware cue system
+	Cues are invoked with callbacks that can be subcribed to: Rewind, Resimulate, Confirmed. This is the only way to support "undoing" cues.
+	+Allows cues that "didn't actually happen" (mispredict) to be undone.
+	+Allows "all cues" to be undown on rollback and then repredicted during resimulate. (How important is this actually given everything else???)
+		-Allows us to "skip" the uniqueness test at the cost of redundantly rolling back -> redoing cues that were not mispredicted
 
+	
+	
+	
\ No newline at end of file
diff --git a/Engine/Plugins/Runtime/NetworkPredictionExtras/Source/NetworkPredictionExtras/Private/MockAbilitySimulation.cpp b/Engine/Plugins/Runtime/NetworkPredictionExtras/Source/NetworkPredictionExtras/Private/MockAbilitySimulation.cpp
--- a/Engine/Plugins/Runtime/NetworkPredictionExtras/Source/NetworkPredictionExtras/Private/MockAbilitySimulation.cpp
+++ b/Engine/Plugins/Runtime/NetworkPredictionExtras/Source/NetworkPredictionExtras/Private/MockAbilitySimulation.cpp
@@ -4,6 +4,8 @@
 #include "DrawDebugHelpers.h"
 #include "VisualLogger/VisualLogger.h"
 
+const FName FMockAbilitySimulation::GroupName(TEXT("Ability"));
+
 namespace MockAbilityCVars
 {
 	NETSIM_DEVCVAR_SHIPCONST_FLOAT(DefaultMaxSpeed, 1200.f, "mockability.DefaultMaxSpeed", "Default Speed");
@@ -15,13 +17,14 @@ namespace MockAbilityCVars
 	NETSIM_DEVCVAR_SHIPCONST_FLOAT(DashAcceleration, 100000.f, "mockability.DashAcceleration", "Acceleration when dashing.");
 }
 
+NETSIMCUE_REGISTER(FMockAbilityBlinkCue, TEXT("MockAbilityBlink"));
+NETSIMCUESET_REGISTER(UMockFlyingAbilityComponent, FMockAbilityCueSet);
+
 // -------------------------------------------------------------------------------------------------------------
 //
 // -------------------------------------------------------------------------------------------------------------
 
-const FName FMockAbilitySimulation::GroupName(TEXT("Ability"));
-
-void FMockAbilitySimulation::SimulationTick(const TNetSimTimeStep& TimeStep, const TNetSimInput<TMockAbilityBufferTypes>& Input, const TNetSimOutput<TMockAbilityBufferTypes>& Output)
+void FMockAbilitySimulation::SimulationTick(const FNetSimTimeStep& TimeStep, const TNetSimInput<TMockAbilityBufferTypes>& Input, const TNetSimOutput<TMockAbilityBufferTypes>& Output)
 {
 	// WIP NOTES:
 	//	-We are creating local copies of the input state that we can mutate and pass on to the parent simulation's tick (LocalCmd, LocalSync, LocalAux). This may seem weird but consider:
@@ -44,8 +47,7 @@ void FMockAbilitySimulation::SimulationTick(const TNetSimTimeStep& TimeStep, con
 	//			writing fragile code. Making code be explicit: "I am writing to the current value this frame AND next frames value" seems ok. Could possibly do everything in a local copy and then see if its changed
 	//			at the end up ::SimulationTick and then copy it into output aux. Probably should avoid comparison operator each frame so would need to wrap in some struct that could track a dirty flag... seems complicating.
 	//
-
-	check(EventHandler);
+	
 	const float DeltaSeconds = TimeStep.StepMS.ToRealTimeSeconds();
 
 	// Local copies of the const input that we will pass into the parent sim as input.
@@ -80,7 +82,6 @@ void FMockAbilitySimulation::SimulationTick(const TNetSimTimeStep& TimeStep, con
 	if (bBlinkActivate)
 	{
 		LocalAux.BlinkWarmupLeft = BlinkWarmupMS;
-		EventHandler->NotifyBlinkStartup();
 	}
 
 	if (LocalAux.BlinkWarmupLeft > 0)
@@ -116,7 +117,9 @@ void FMockAbilitySimulation::SimulationTick(const TNetSimTimeStep& TimeStep, con
 				LocalSync.Stamina = NewStamina;
 				Output.Sync.Stamina = NewStamina;
 
-				EventHandler->NotifyBlinkFinished();
+				// Invoke a NetCue for the blink. This is essentially capturing the start/end location so that all receivers of the event
+				// get the exact coordinates (maybe overkill in practice but key is that we have data that we want to pass out via an event)
+				Output.CueDispatch.Invoke<FMockAbilityBlinkCue>({Input.Sync.Location, DestLocation});
 			}
 		}
 	}
@@ -142,8 +145,6 @@ void FMockAbilitySimulation::SimulationTick(const TNetSimTimeStep& TimeStep, con
 		const float NewStamina = FMath::Max<float>(0.f, LocalSync.Stamina - DashCost);
 		LocalSync.Stamina = NewStamina;
 		Output.Sync.Stamina = NewStamina;
-
-		EventHandler->NotifyDash(true);
 	}
 
 	if (LocalAux.DashTimeLeft > 0)
@@ -158,13 +159,12 @@ void FMockAbilitySimulation::SimulationTick(const TNetSimTimeStep& TimeStep, con
 		LocalCmd.bBlinkPressed = false;
 		LocalCmd.bSprintPressed = false;
 		
-		FlyingMovement::FMovementSimulation::SimulationTick(TimeStep, { LocalCmd, LocalSync, LocalAux }, { Output.Sync, Output.Aux });
+		FlyingMovement::FMovementSimulation::SimulationTick(TimeStep, { LocalCmd, LocalSync, LocalAux }, { Output.Sync, Output.Aux, Output.CueDispatch });
 
 		if (NewDashTimeLeft == 0)
 		{
 			// Stop when dash is over
 			Output.Sync.Velocity.Set(0.f, 0.f, 0.f);
-			EventHandler->NotifyDash(false);
 		}
 
 		return;
@@ -212,11 +212,10 @@ void FMockAbilitySimulation::SimulationTick(const TNetSimTimeStep& TimeStep, con
 	// Update the out aux state and call notifies only if sprinting state as actually changed
 	if (bIsSprinting != Input.Aux.bIsSprinting)
 	{
-		EventHandler->NotifySprint(bIsSprinting);
 		Output.Aux.Get()->bIsSprinting = bIsSprinting;
 	}
 
-	FlyingMovement::FMovementSimulation::SimulationTick(TimeStep, { LocalCmd, LocalSync, LocalAux }, { Output.Sync, Output.Aux });
+	FlyingMovement::FMovementSimulation::SimulationTick(TimeStep, { LocalCmd, LocalSync, LocalAux }, { Output.Sync, Output.Aux, Output.CueDispatch });
 }
 
 // -------------------------------------------------------------------------------------------------------------
@@ -292,43 +291,16 @@ void UMockFlyingAbilityComponent::VisualLog(const FMockAbilityInputCmd* Input, c
 
 // ---------------------------------------------------------------------------------
 
-void UMockFlyingAbilityComponent::NotifySprint(bool bNewIsSprinting)
-{
-	if (GetNetMode() == NM_DedicatedServer)
-	{
-		return;
-	}
-
-	//UE_LOG(LogTemp, Display, TEXT("Sprint: %d"), bNewIsSprinting);
-}
+static float BlinkCueDuration = 1.f;
+static FAutoConsoleVariableRef CVarBindAutomatically(TEXT("NetworkPredictionExtras.FlyingPawn.BlinkCueDuration"),
+	BlinkCueDuration, TEXT("Duration of BlinkCue"), ECVF_Default);
 
-void UMockFlyingAbilityComponent::NotifyDash(bool bNewIsDashing)
+void UMockFlyingAbilityComponent::HandleCue(FMockAbilityBlinkCue& BlinkCue, const FNetworkSimTime& ElapsedTime)
 {
-	if (GetNetMode() == NM_DedicatedServer)
-	{
-		return;
-	}
-
-	//UE_LOG(LogTemp, Display, TEXT("Dash: %d"), bNewIsDashing);
-}
-
-void UMockFlyingAbilityComponent::NotifyBlinkStartup()
-{
-	if (GetNetMode() == NM_DedicatedServer)
-	{
-		return;
-	}
-
-	//UE_LOG(LogTemp, Display, TEXT("Blink Startup"));
-}
-
-void UMockFlyingAbilityComponent::NotifyBlinkFinished()
-{
-	if (GetNetMode() == NM_DedicatedServer)
-	{
-		return;
-	}
-
-	//UE_LOG(LogTemp, Display, TEXT("Blink Finished"));
-}
+	UE_LOG(LogTemp, Display, TEXT("BlinkCue! : <%f, %f, %f> - <%f, %f, %f>. ElapsedTime: %s"), BlinkCue.StartLocation.X, BlinkCue.StartLocation.Y, BlinkCue.StartLocation.Z,
+		BlinkCue.StopLocation.X, BlinkCue.StopLocation.Y, BlinkCue.StopLocation.Z, *ElapsedTime.ToString()); //*BlinkCue.StartLocation.ToString(), *BlinkCue.StopLocation.ToString());
 
+	// Crude compensation for cue firing in the past (note this is not necessary! Some cues not care and need to see the "full" effect regardless of when it happened)
+	float Duration = BlinkCueDuration - ElapsedTime.ToRealTimeSeconds();
+	DrawDebugLine(GetWorld(), BlinkCue.StartLocation, BlinkCue.StopLocation, FColor::Red, false, Duration);
+}
\ No newline at end of file
diff --git a/Engine/Plugins/Runtime/NetworkPredictionExtras/Source/NetworkPredictionExtras/Private/MockNetworkSimulation.cpp b/Engine/Plugins/Runtime/NetworkPredictionExtras/Source/NetworkPredictionExtras/Private/MockNetworkSimulation.cpp
--- a/Engine/Plugins/Runtime/NetworkPredictionExtras/Source/NetworkPredictionExtras/Private/MockNetworkSimulation.cpp
+++ b/Engine/Plugins/Runtime/NetworkPredictionExtras/Source/NetworkPredictionExtras/Private/MockNetworkSimulation.cpp
@@ -61,7 +61,7 @@ static bool ForceMispredict = false;
 
 const FName FMockNetworkSimulation::GroupName(TEXT("Mock"));
 
-void FMockNetworkSimulation::SimulationTick(const TNetSimTimeStep& TimeStep, const TNetSimInput<TMockNetworkSimulationBufferTypes>& Input, const TNetSimOutput<TMockNetworkSimulationBufferTypes>& Output)
+void FMockNetworkSimulation::SimulationTick(const FNetSimTimeStep& TimeStep, const TNetSimInput<TMockNetworkSimulationBufferTypes>& Input, const TNetSimOutput<TMockNetworkSimulationBufferTypes>& Output)
 {
 	Output.Sync.Total = Input.Sync.Total + (Input.Cmd.InputValue * Input.Aux.Multiplier * TimeStep.StepMS.ToRealTimeSeconds());
 
diff --git a/Engine/Plugins/Runtime/NetworkPredictionExtras/Source/NetworkPredictionExtras/Public/MockAbilitySimulation.h b/Engine/Plugins/Runtime/NetworkPredictionExtras/Source/NetworkPredictionExtras/Public/MockAbilitySimulation.h
--- a/Engine/Plugins/Runtime/NetworkPredictionExtras/Source/NetworkPredictionExtras/Public/MockAbilitySimulation.h
+++ b/Engine/Plugins/Runtime/NetworkPredictionExtras/Source/NetworkPredictionExtras/Public/MockAbilitySimulation.h
@@ -39,6 +39,10 @@
 //
 // -------------------------------------------------------------------------------------------------------------------------------
 
+// -------------------------------------------------------
+// MockAbility Data structures
+// -------------------------------------------------------
+
 struct FMockAbilityInputCmd : public FlyingMovement::FInputCmd
 {
 	bool bSprintPressed = false;
@@ -117,19 +121,62 @@ struct FMockAbilityAuxstate : public FlyingMovement::FAuxState
 	}
 };
 
-using TMockAbilityBufferTypes = TNetworkSimBufferTypes<FMockAbilityInputCmd, FMockAbilitySyncState, FMockAbilityAuxstate>;
+// -------------------------------------------------------
+// MockAbility NetSimCues - events emitted by the sim
+// -------------------------------------------------------
 
-// A very hard coded event handling interface for events that are emitted in FMockAbilitySimulation::SimulationTick.
-// These events are for cosmetic, client side work. E.g, we are not implementing sprint/dash/blink here, we are doing
-// visual/audio effects based on those states changing within the core simulation update. These are effectively skipped on the server (in implementation).
-class IMockAbilityEventHandler
+// Cue for blink activation.
+struct FMockAbilityBlinkCue
 {
-public:
-	virtual void NotifySprint(bool bIsSprinting) = 0;
-	virtual void NotifyDash(bool bIsDashing) = 0;
-	virtual void NotifyBlinkStartup() = 0;
-	virtual void NotifyBlinkFinished() = 0;
+	NETSIMCUE_BODY();
+
+	FVector_NetQuantize10 StartLocation;
+	FVector_NetQuantize10 StopLocation;
+
+	void NetSerialize(FArchive& Ar)
+	{
+		bool b = false;
+		StartLocation.NetSerialize(Ar, nullptr, b);
+		StopLocation.NetSerialize(Ar, nullptr, b);
+	}
+
+	static bool Unique(const FMockAbilityBlinkCue& A, const FMockAbilityBlinkCue& B)
+	{
+		const float ErrorTolerance = 1.f;
+		return !A.StartLocation.Equals(B.StartLocation, ErrorTolerance) || !A.StopLocation.Equals(B.StopLocation, ErrorTolerance);
+	}
+};
+
+/*
+template<>
+struct TCueHandlerTraits<FMockAbilityBlinkCue> : public TNetSimCueTraitsBase//<FMockAbilityBlinkCue>
+{
+	static constexpr uint8 InvokeMask { (uint8)ESimulationTickContext::Authority };
+	static constexpr bool Replicate { false };
 };
+*/
+
+template<>
+struct TCueHandlerTraits<FMockAbilityBlinkCue> : public TNetSimCueTraits_ReplicatedNonPredicted { };
+
+
+
+// The set of Cues the MockAbility simulation will invoke
+struct FMockAbilityCueSet
+{
+	template<typename TDispatchTable>
+	static void RegisterNetSimCueTypes(TDispatchTable& DispatchTable)
+	{
+		DispatchTable.template RegisterType<FMockAbilityBlinkCue>();
+	}
+};
+
+
+// -------------------------------------------------------
+// MockAbilitySimulation definition
+// -------------------------------------------------------
+
+using TMockAbilityBufferTypes = TNetworkSimBufferTypes<FMockAbilityInputCmd, FMockAbilitySyncState, FMockAbilityAuxstate>;
 
 class FMockAbilitySimulation : public FlyingMovement::FMovementSimulation
 {
@@ -138,9 +185,7 @@ public:
 	static const FName GroupName;
 
 	/** Main update function */
-	void SimulationTick(const TNetSimTimeStep& TimeStep, const TNetSimInput<TMockAbilityBufferTypes>& Input, const TNetSimOutput<TMockAbilityBufferTypes>& Output);
-
-	IMockAbilityEventHandler* EventHandler = nullptr;
+	void SimulationTick(const FNetSimTimeStep& TimeStep, const TNetSimInput<TMockAbilityBufferTypes>& Input, const TNetSimOutput<TMockAbilityBufferTypes>& Output);
 };
 
 template<int32 InFixedStepMS=0>
@@ -153,7 +198,7 @@ class IMockFlyingAbilitySystemDriver : public TNetworkedSimulationModelDriver<TM
 // -------------------------------------------------------------------------------------------------------------------------------
 
 UCLASS(BlueprintType, meta=(BlueprintSpawnableComponent))
-class NETWORKPREDICTIONEXTRAS_API UMockFlyingAbilityComponent : public UFlyingMovementComponent, public IMockFlyingAbilitySystemDriver, public IMockAbilityEventHandler
+class NETWORKPREDICTIONEXTRAS_API UMockFlyingAbilityComponent : public UFlyingMovementComponent, public IMockFlyingAbilitySystemDriver
 {
 	GENERATED_BODY()
 
@@ -180,11 +225,8 @@ public:
 	using UFlyingMovementComponent::ProduceInput;
 	using UFlyingMovementComponent::FinalizeFrame;
 
-	// IMockAbilityEventHandler
-	virtual void NotifySprint(bool bIsSprinting) override;
-	virtual void NotifyDash(bool bIsDashing) override;
-	virtual void NotifyBlinkStartup() override;
-	virtual void NotifyBlinkFinished() override;
+	// NetSimCues
+	void HandleCue(FMockAbilityBlinkCue& BlinkCue, const FNetworkSimTime& ElapsedTime);
 
 	// -------------------------------------------------------------------------------------
 	//	Ability State and Notifications
@@ -242,8 +284,7 @@ protected:
 	void InitMockAbilitySimulation(TSimulation* Simulation, FMockAbilitySyncState& InitialSyncState, FMockAbilityAuxstate& InitialAuxState)
 	{
 		check(MockAbilitySimulation == nullptr);
-		MockAbilitySimulation = Simulation;
-		MockAbilitySimulation->EventHandler = this;
+		MockAbilitySimulation = Simulation;		
 
 		InitFlyingMovementSimulation(Simulation, InitialSyncState, InitialAuxState);
 	}
diff --git a/Engine/Plugins/Runtime/NetworkPredictionExtras/Source/NetworkPredictionExtras/Public/MockNetworkSimulation.h b/Engine/Plugins/Runtime/NetworkPredictionExtras/Source/NetworkPredictionExtras/Public/MockNetworkSimulation.h
--- a/Engine/Plugins/Runtime/NetworkPredictionExtras/Source/NetworkPredictionExtras/Public/MockNetworkSimulation.h
+++ b/Engine/Plugins/Runtime/NetworkPredictionExtras/Source/NetworkPredictionExtras/Public/MockNetworkSimulation.h
@@ -138,7 +138,7 @@ class FMockNetworkSimulation
 public:
 
 	/** Main update function */
-	void SimulationTick(const TNetSimTimeStep& TimeStep, const TNetSimInput<TMockNetworkSimulationBufferTypes>& Input, const TNetSimOutput<TMockNetworkSimulationBufferTypes>& Output);
+	void SimulationTick(const FNetSimTimeStep& TimeStep, const TNetSimInput<TMockNetworkSimulationBufferTypes>& Input, const TNetSimOutput<TMockNetworkSimulationBufferTypes>& Output);
 
 	/** Tick group the simulation maps to */
 	static const FName GroupName;
