attempted mac linker fix
#rnx

diff --git a/Engine/Plugins/Runtime/NetworkPrediction/Source/NetworkPrediction/Public/NetworkedSimulationModelCues.h b/Engine/Plugins/Runtime/NetworkPrediction/Source/NetworkPrediction/Public/NetworkedSimulationModelCues.h
--- a/Engine/Plugins/Runtime/NetworkPrediction/Source/NetworkPrediction/Public/NetworkedSimulationModelCues.h
+++ b/Engine/Plugins/Runtime/NetworkPrediction/Source/NetworkPrediction/Public/NetworkedSimulationModelCues.h
@@ -191,7 +191,7 @@ private:
 		FString TypeName;
 	};
 
-	static FGlobalCueTypeTable Singleton;
+	NETWORKPREDICTION_API static FGlobalCueTypeTable Singleton;
 
 	TArray<TFunction<void()>> PostFinalizeCallbacks;
 	TArray<FTypeInfo> PendingCueTypes;
