NetworkPrediction Plugin
-Remove use of namespaces in Flying/Parametric movement examples.

#rnx

diff --git a/Engine/Plugins/Runtime/NetworkPrediction/Source/NetworkPrediction/Private/Movement/FlyingMovementComponent.cpp b/Engine/Plugins/Runtime/NetworkPrediction/Source/NetworkPrediction/Private/Movement/FlyingMovementComponent.cpp
--- a/Engine/Plugins/Runtime/NetworkPrediction/Source/NetworkPrediction/Private/Movement/FlyingMovementComponent.cpp
+++ b/Engine/Plugins/Runtime/NetworkPrediction/Source/NetworkPrediction/Private/Movement/FlyingMovementComponent.cpp
@@ -56,13 +56,13 @@ UFlyingMovementComponent::UFlyingMovementComponent()
 INetworkedSimulationModel* UFlyingMovementComponent::InstantiateNetworkedSimulation()
 {
 	// The Simulation
-	FlyingMovement::FMoveState InitialSyncState;
-	FlyingMovement::FAuxState InitialAuxState;
+	FFlyingMovementSyncState InitialSyncState;
+	FFlyingMovementAuxState InitialAuxState;
 
-	InitFlyingMovementSimulation(new FlyingMovement::FMovementSimulation(), InitialSyncState, InitialAuxState);
+	InitFlyingMovementSimulation(new FFlyingMovementSimulation(), InitialSyncState, InitialAuxState);
 
 	// The Model
-	auto NewModel = new FlyingMovement::FMovementSystem<0>(MovementSimulation, this, InitialSyncState, InitialAuxState);
+	auto NewModel = new FFlyingMovementSystem<0>(MovementSimulation, this, InitialSyncState, InitialAuxState);
 	InitFlyingMovementNetSimModel(NewModel);
 	
 	return NewModel;
@@ -90,7 +90,7 @@ void UFlyingMovementComponent::TickComponent(float DeltaTime, enum ELevelTick Ti
 	// Check if we should trip a mispredict. (Note how its not possible to do this inside the Update function!)
 	if (OwnerRole == ROLE_Authority && FlyingMovementCVars::RequestMispredict)
 	{
-		FlyingMovement::FMovementSimulation::ForceMispredict = true;
+		FFlyingMovementSimulation::ForceMispredict = true;
 		FlyingMovementCVars::RequestMispredict = 0;
 	}
 
@@ -100,7 +100,7 @@ void UFlyingMovementComponent::TickComponent(float DeltaTime, enum ELevelTick Ti
 	{
 		if (MovementAuxState->Get()->MaxSpeed != FlyingMovementCVars::MaxSpeed)
 		{
-			MovementAuxState->Modify([](FlyingMovement::FAuxState& Aux)
+			MovementAuxState->Modify([](FFlyingMovementAuxState& Aux)
 			{
 				Aux.MaxSpeed = FlyingMovementCVars::MaxSpeed;
 			});
@@ -119,16 +119,16 @@ void UFlyingMovementComponent::TickComponent(float DeltaTime, enum ELevelTick Ti
 //
 // ----------------------------------------------------------------------------------------------------------
 
-void UFlyingMovementComponent::ProduceInput(const FNetworkSimTime SimTime, FlyingMovement::FInputCmd& Cmd)
+void UFlyingMovementComponent::ProduceInput(const FNetworkSimTime SimTime, FFlyingMovementInputCmd& Cmd)
 {
 	// This isn't ideal. It probably makes sense for the component to do all the input binding rather.
 	ProduceInputDelegate.ExecuteIfBound(SimTime, Cmd);
 }
 
-void UFlyingMovementComponent::FinalizeFrame(const FlyingMovement::FMoveState& SyncState, const FlyingMovement::FAuxState& AuxState)
+void UFlyingMovementComponent::FinalizeFrame(const FFlyingMovementSyncState& SyncState, const FFlyingMovementAuxState& AuxState)
 {
 	// Does checking equality make any sense here? This is unfortunate
-	if (UpdatedComponent->GetComponentLocation().Equals(SyncState.Location) == false || UpdatedComponent->GetComponentQuat().Rotator().Equals(SyncState.Rotation, FlyingMovement::ROTATOR_TOLERANCE) == false)
+	if (UpdatedComponent->GetComponentLocation().Equals(SyncState.Location) == false || UpdatedComponent->GetComponentQuat().Rotator().Equals(SyncState.Rotation, ROTATOR_TOLERANCE) == false)
 	{
 		FTransform Transform(SyncState.Rotation.Quaternion(), SyncState.Location, UpdatedComponent->GetComponentTransform().GetScale3D() );
 		UpdatedComponent->SetWorldTransform(Transform, false, nullptr, ETeleportType::TeleportPhysics);
@@ -147,7 +147,7 @@ const AActor* UFlyingMovementComponent::GetVLogOwner() const
 	return GetOwner();
 }
 
-void UFlyingMovementComponent::VisualLog(const FlyingMovement::FInputCmd* Input, const FlyingMovement::FMoveState* Sync, const FlyingMovement::FAuxState* Aux, const FVisualLoggingParameters& SystemParameters) const
+void UFlyingMovementComponent::VisualLog(const FFlyingMovementInputCmd* Input, const FFlyingMovementSyncState* Sync, const FFlyingMovementAuxState* Aux, const FVisualLoggingParameters& SystemParameters) const
 {
 	FTransform Transform(Sync->Rotation, Sync->Location);
 	FVisualLoggingHelpers::VisualLogActor(GetOwner(), Transform, SystemParameters);	
diff --git a/Engine/Plugins/Runtime/NetworkPrediction/Source/NetworkPrediction/Private/Movement/FlyingMovementSimulation.cpp b/Engine/Plugins/Runtime/NetworkPrediction/Source/NetworkPrediction/Private/Movement/FlyingMovementSimulation.cpp
--- a/Engine/Plugins/Runtime/NetworkPrediction/Source/NetworkPrediction/Private/Movement/FlyingMovementSimulation.cpp
+++ b/Engine/Plugins/Runtime/NetworkPrediction/Source/NetworkPrediction/Private/Movement/FlyingMovementSimulation.cpp
@@ -42,237 +42,234 @@ static FAutoConsoleVariableRef CVarDrawDebugDefaultLifeTime(TEXT("fp.Debug.UseDr
 
 // ============================================================================================
 
-namespace FlyingMovement
-{
-	const FName FMovementSimulation::GroupName(TEXT("FlyingMovement"));
 
-	bool FMovementSimulation::ForceMispredict = false;
-	static FVector ForceMispredictVelocityMagnitude = FVector(2000.f, 0.f, 0.f);
+const FName FFlyingMovementSimulation::GroupName(TEXT("FlyingMovement"));
 
-	bool IsExceedingMaxSpeed(const FVector& Velocity, float InMaxSpeed)
-	{
-		InMaxSpeed = FMath::Max(0.f, InMaxSpeed);
-		const float MaxSpeedSquared = FMath::Square(InMaxSpeed);
+bool FFlyingMovementSimulation::ForceMispredict = false;
+static FVector ForceMispredictVelocityMagnitude = FVector(2000.f, 0.f, 0.f);
+
+bool IsExceedingMaxSpeed(const FVector& Velocity, float InMaxSpeed)
+{
+	InMaxSpeed = FMath::Max(0.f, InMaxSpeed);
+	const float MaxSpeedSquared = FMath::Square(InMaxSpeed);
 	
-		// Allow 1% error tolerance, to account for numeric imprecision.
-		const float OverVelocityPercent = 1.01f;
-		return (Velocity.SizeSquared() > MaxSpeedSquared * OverVelocityPercent);
-	}
+	// Allow 1% error tolerance, to account for numeric imprecision.
+	const float OverVelocityPercent = 1.01f;
+	return (Velocity.SizeSquared() > MaxSpeedSquared * OverVelocityPercent);
+}
+
+FVector ComputeSlideVector(const FVector& Delta, const float Time, const FVector& Normal, const FHitResult& Hit)
+{
+	return FVector::VectorPlaneProject(Delta, Normal) * Time;
+}
 
-	FVector ComputeSlideVector(const FVector& Delta, const float Time, const FVector& Normal, const FHitResult& Hit)
+void TwoWallAdjust(FVector& OutDelta, const FHitResult& Hit, const FVector& OldHitNormal)
+{
+	FVector Delta = OutDelta;
+	const FVector HitNormal = Hit.Normal;
+
+	if ((OldHitNormal | HitNormal) <= 0.f) //90 or less corner, so use cross product for direction
 	{
-		return FVector::VectorPlaneProject(Delta, Normal) * Time;
+		const FVector DesiredDir = Delta;
+		FVector NewDir = (HitNormal ^ OldHitNormal);
+		NewDir = NewDir.GetSafeNormal();
+		Delta = (Delta | NewDir) * (1.f - Hit.Time) * NewDir;
+		if ((DesiredDir | Delta) < 0.f)
+		{
+			Delta = -1.f * Delta;
+		}
 	}
-
-	void TwoWallAdjust(FVector& OutDelta, const FHitResult& Hit, const FVector& OldHitNormal)
+	else //adjust to new wall
 	{
-		FVector Delta = OutDelta;
-		const FVector HitNormal = Hit.Normal;
-
-		if ((OldHitNormal | HitNormal) <= 0.f) //90 or less corner, so use cross product for direction
+		const FVector DesiredDir = Delta;
+		Delta = ComputeSlideVector(Delta, 1.f - Hit.Time, HitNormal, Hit);
+		if ((Delta | DesiredDir) <= 0.f)
 		{
-			const FVector DesiredDir = Delta;
-			FVector NewDir = (HitNormal ^ OldHitNormal);
-			NewDir = NewDir.GetSafeNormal();
-			Delta = (Delta | NewDir) * (1.f - Hit.Time) * NewDir;
-			if ((DesiredDir | Delta) < 0.f)
-			{
-				Delta = -1.f * Delta;
-			}
+			Delta = FVector::ZeroVector;
 		}
-		else //adjust to new wall
+		else if ( FMath::Abs((HitNormal | OldHitNormal) - 1.f) < KINDA_SMALL_NUMBER )
 		{
-			const FVector DesiredDir = Delta;
-			Delta = ComputeSlideVector(Delta, 1.f - Hit.Time, HitNormal, Hit);
-			if ((Delta | DesiredDir) <= 0.f)
-			{
-				Delta = FVector::ZeroVector;
-			}
-			else if ( FMath::Abs((HitNormal | OldHitNormal) - 1.f) < KINDA_SMALL_NUMBER )
-			{
-				// we hit the same wall again even after adjusting to move along it the first time
-				// nudge away from it (this can happen due to precision issues)
-				Delta += HitNormal * 0.01f;
-			}
+			// we hit the same wall again even after adjusting to move along it the first time
+			// nudge away from it (this can happen due to precision issues)
+			Delta += HitNormal * 0.01f;
 		}
-
-		OutDelta = Delta;
 	}
 
-	float FMovementSimulation::SlideAlongSurface(const FVector& Delta, float Time, const FQuat Rotation, const FVector& Normal, FHitResult& Hit, bool bHandleImpact)
+	OutDelta = Delta;
+}
+
+float FFlyingMovementSimulation::SlideAlongSurface(const FVector& Delta, float Time, const FQuat Rotation, const FVector& Normal, FHitResult& Hit, bool bHandleImpact)
+{
+	if (!Hit.bBlockingHit)
 	{
-		if (!Hit.bBlockingHit)
-		{
-			return 0.f;
-		}
+		return 0.f;
+	}
 
-		float PercentTimeApplied = 0.f;
-		const FVector OldHitNormal = Normal;
+	float PercentTimeApplied = 0.f;
+	const FVector OldHitNormal = Normal;
 
-		FVector SlideDelta = ComputeSlideVector(Delta, Time, Normal, Hit);
+	FVector SlideDelta = ComputeSlideVector(Delta, Time, Normal, Hit);
 
-		if ((SlideDelta | Delta) > 0.f)
-		{
-			SafeMoveUpdatedComponent(SlideDelta, Rotation, true, Hit, ETeleportType::None);
+	if ((SlideDelta | Delta) > 0.f)
+	{
+		SafeMoveUpdatedComponent(SlideDelta, Rotation, true, Hit, ETeleportType::None);
 
-			const float FirstHitPercent = Hit.Time;
-			PercentTimeApplied = FirstHitPercent;
-			if (Hit.IsValidBlockingHit())
+		const float FirstHitPercent = Hit.Time;
+		PercentTimeApplied = FirstHitPercent;
+		if (Hit.IsValidBlockingHit())
+		{
+			// Notify first impact
+			if (bHandleImpact)
 			{
-				// Notify first impact
-				if (bHandleImpact)
-				{
-					// !HandleImpact(Hit, FirstHitPercent * Time, SlideDelta);
-				}
+				// !HandleImpact(Hit, FirstHitPercent * Time, SlideDelta);
+			}
+
+			// Compute new slide normal when hitting multiple surfaces.
+			TwoWallAdjust(SlideDelta, Hit, OldHitNormal);
 
-				// Compute new slide normal when hitting multiple surfaces.
-				TwoWallAdjust(SlideDelta, Hit, OldHitNormal);
+			// Only proceed if the new direction is of significant length and not in reverse of original attempted move.
+			if (!SlideDelta.IsNearlyZero(1e-3f) && (SlideDelta | Delta) > 0.f)
+			{
+				// Perform second move
+				SafeMoveUpdatedComponent(SlideDelta, Rotation, true, Hit, ETeleportType::None);
+				const float SecondHitPercent = Hit.Time * (1.f - FirstHitPercent);
+				PercentTimeApplied += SecondHitPercent;
 
-				// Only proceed if the new direction is of significant length and not in reverse of original attempted move.
-				if (!SlideDelta.IsNearlyZero(1e-3f) && (SlideDelta | Delta) > 0.f)
+				// Notify second impact
+				if (bHandleImpact && Hit.bBlockingHit)
 				{
-					// Perform second move
-					SafeMoveUpdatedComponent(SlideDelta, Rotation, true, Hit, ETeleportType::None);
-					const float SecondHitPercent = Hit.Time * (1.f - FirstHitPercent);
-					PercentTimeApplied += SecondHitPercent;
-
-					// Notify second impact
-					if (bHandleImpact && Hit.bBlockingHit)
-					{
-						// !HandleImpact(Hit, SecondHitPercent * Time, SlideDelta);
-					}
+					// !HandleImpact(Hit, SecondHitPercent * Time, SlideDelta);
 				}
 			}
-
-			return FMath::Clamp(PercentTimeApplied, 0.f, 1.f);
 		}
 
-		return 0.f;
+		return FMath::Clamp(PercentTimeApplied, 0.f, 1.f);
 	}
 
-	void FMovementSimulation::SimulationTick(const FNetSimTimeStep& TimeStep, const TNetSimInput<TMovementBufferTypes>& Input, const TNetSimOutput<TMovementBufferTypes>& Output)
-	{
-		Output.Sync = Input.Sync;
+	return 0.f;
+}
 
-		const float DeltaSeconds = TimeStep.StepMS.ToRealTimeSeconds();
+void FFlyingMovementSimulation::SimulationTick(const FNetSimTimeStep& TimeStep, const TNetSimInput<FlyingMovementBufferTypes>& Input, const TNetSimOutput<FlyingMovementBufferTypes>& Output)
+{
+	Output.Sync = Input.Sync;
+
+	const float DeltaSeconds = TimeStep.StepMS.ToRealTimeSeconds();
 
-		// --------------------------------------------------------------
-		//	Rotation Update
-		//	We do the rotational update inside the movement sim so that things like server side teleport will work.
-		//	(We want rotation to be treated the same as location, with respect to how its updated, corrected, etc).
-		//	In this simulation, the rotation update isn't allowed to "fail". We don't expect the collision query to be able to fail the rotational update.
-		// --------------------------------------------------------------
+	// --------------------------------------------------------------
+	//	Rotation Update
+	//	We do the rotational update inside the movement sim so that things like server side teleport will work.
+	//	(We want rotation to be treated the same as location, with respect to how its updated, corrected, etc).
+	//	In this simulation, the rotation update isn't allowed to "fail". We don't expect the collision query to be able to fail the rotational update.
+	// --------------------------------------------------------------
 
-		Output.Sync.Rotation += (Input.Cmd.RotationInput * DeltaSeconds);
-		Output.Sync.Rotation.Normalize();
+	Output.Sync.Rotation += (Input.Cmd.RotationInput * DeltaSeconds);
+	Output.Sync.Rotation.Normalize();
 
-		const FQuat OutputQuat = Output.Sync.Rotation.Quaternion();
+	const FQuat OutputQuat = Output.Sync.Rotation.Quaternion();
 
-		// After the rotation is known, we need to sync our driver to this state. This is unfortunate and probably not best for perf, but since the Driver owned
-		// primitive component is ultimately what does our scene queries, we have no choice: we must get the primitive component in the state we say it should be in.
-		PreSimSync(Output.Sync);
+	// After the rotation is known, we need to sync our driver to this state. This is unfortunate and probably not best for perf, but since the Driver owned
+	// primitive component is ultimately what does our scene queries, we have no choice: we must get the primitive component in the state we say it should be in.
+	PreSimSync(Output.Sync);
 
-		const FVector LocalSpaceMovementInput = Output.Sync.Rotation.RotateVector( Input.Cmd.MovementInput );
+	const FVector LocalSpaceMovementInput = Output.Sync.Rotation.RotateVector( Input.Cmd.MovementInput );
 	   	
-		// --------------------------------------------------------------
-		// Calculate Output.Sync.RelativeVelocity based on Input
-		// --------------------------------------------------------------
-		{
-			const FVector ControlAcceleration = LocalSpaceMovementInput.GetClampedToMaxSize(1.f);
-			FVector Velocity = Input.Sync.Velocity;
+	// --------------------------------------------------------------
+	// Calculate Output.Sync.RelativeVelocity based on Input
+	// --------------------------------------------------------------
+	{
+		const FVector ControlAcceleration = LocalSpaceMovementInput.GetClampedToMaxSize(1.f);
+		FVector Velocity = Input.Sync.Velocity;
 
-			const float AnalogInputModifier = (ControlAcceleration.SizeSquared() > 0.f ? ControlAcceleration.Size() : 0.f);
-			const float MaxPawnSpeed = Input.Aux.MaxSpeed * AnalogInputModifier;
-			const bool bExceedingMaxSpeed = IsExceedingMaxSpeed(Velocity, MaxPawnSpeed);
+		const float AnalogInputModifier = (ControlAcceleration.SizeSquared() > 0.f ? ControlAcceleration.Size() : 0.f);
+		const float MaxPawnSpeed = Input.Aux.MaxSpeed * AnalogInputModifier;
+		const bool bExceedingMaxSpeed = IsExceedingMaxSpeed(Velocity, MaxPawnSpeed);
 
-			if (AnalogInputModifier > 0.f && !bExceedingMaxSpeed)
+		if (AnalogInputModifier > 0.f && !bExceedingMaxSpeed)
+		{
+			// Apply change in velocity direction
+			if (Velocity.SizeSquared() > 0.f)
 			{
-				// Apply change in velocity direction
-				if (Velocity.SizeSquared() > 0.f)
-				{
-					// Change direction faster than only using acceleration, but never increase velocity magnitude.
-					const float TimeScale = FMath::Clamp(DeltaSeconds * Input.Aux.TurningBoost, 0.f, 1.f);
-					Velocity = Velocity + (ControlAcceleration * Velocity.Size() - Velocity) * TimeScale;
-				}
+				// Change direction faster than only using acceleration, but never increase velocity magnitude.
+				const float TimeScale = FMath::Clamp(DeltaSeconds * Input.Aux.TurningBoost, 0.f, 1.f);
+				Velocity = Velocity + (ControlAcceleration * Velocity.Size() - Velocity) * TimeScale;
 			}
-			else
+		}
+		else
+		{
+			// Dampen velocity magnitude based on deceleration.
+			if (Velocity.SizeSquared() > 0.f)
 			{
-				// Dampen velocity magnitude based on deceleration.
-				if (Velocity.SizeSquared() > 0.f)
+				const FVector OldVelocity = Velocity;
+				const float VelSize = FMath::Max(Velocity.Size() - FMath::Abs(Input.Aux.Deceleration) * DeltaSeconds, 0.f);
+				Velocity = Velocity.GetSafeNormal() * VelSize;
+
+				// Don't allow braking to lower us below max speed if we started above it.
+				if (bExceedingMaxSpeed && Velocity.SizeSquared() < FMath::Square(MaxPawnSpeed))
 				{
-					const FVector OldVelocity = Velocity;
-					const float VelSize = FMath::Max(Velocity.Size() - FMath::Abs(Input.Aux.Deceleration) * DeltaSeconds, 0.f);
-					Velocity = Velocity.GetSafeNormal() * VelSize;
-
-					// Don't allow braking to lower us below max speed if we started above it.
-					if (bExceedingMaxSpeed && Velocity.SizeSquared() < FMath::Square(MaxPawnSpeed))
-					{
-						Velocity = OldVelocity.GetSafeNormal() * MaxPawnSpeed;
-					}
+					Velocity = OldVelocity.GetSafeNormal() * MaxPawnSpeed;
 				}
 			}
+		}
 
-			// Apply acceleration and clamp velocity magnitude.
-			const float NewMaxSpeed = (IsExceedingMaxSpeed(Velocity, MaxPawnSpeed)) ? Velocity.Size() : MaxPawnSpeed;
-			Velocity += ControlAcceleration * FMath::Abs(Input.Aux.Acceleration) * DeltaSeconds;
-			Velocity = Velocity.GetClampedToMaxSize(NewMaxSpeed);
+		// Apply acceleration and clamp velocity magnitude.
+		const float NewMaxSpeed = (IsExceedingMaxSpeed(Velocity, MaxPawnSpeed)) ? Velocity.Size() : MaxPawnSpeed;
+		Velocity += ControlAcceleration * FMath::Abs(Input.Aux.Acceleration) * DeltaSeconds;
+		Velocity = Velocity.GetClampedToMaxSize(NewMaxSpeed);
 
-			// Finally, output velocity that we calculated
-			Output.Sync.Velocity = Velocity;
+		// Finally, output velocity that we calculated
+		Output.Sync.Velocity = Velocity;
 		
-			if (FMovementSimulation::ForceMispredict)
-			{
-				Output.Sync.Velocity += ForceMispredictVelocityMagnitude;
-				ForceMispredict = false;
-			}
+		if (FFlyingMovementSimulation::ForceMispredict)
+		{
+			Output.Sync.Velocity += ForceMispredictVelocityMagnitude;
+			ForceMispredict = false;
 		}
+	}
 
-		// ===================================================
+	// ===================================================
 
-		// --------------------------------------------------------------
-		//	Calculate the final movement delta and move the update component
-		// --------------------------------------------------------------
-		FVector Delta = Output.Sync.Velocity * DeltaSeconds;
+	// --------------------------------------------------------------
+	//	Calculate the final movement delta and move the update component
+	// --------------------------------------------------------------
+	FVector Delta = Output.Sync.Velocity * DeltaSeconds;
 
-		if (!Delta.IsNearlyZero(1e-6f))
-		{
-			FHitResult Hit(1.f);
-			SafeMoveUpdatedComponent(Delta, OutputQuat, true, Hit, ETeleportType::None);
-
-			if (Hit.IsValidBlockingHit())
-			{
-				// Try to slide the remaining distance along the surface.
-				SlideAlongSurface(Delta, 1.f-Hit.Time, OutputQuat, Hit.Normal, Hit, true);
-			}
+	if (!Delta.IsNearlyZero(1e-6f))
+	{
+		FHitResult Hit(1.f);
+		SafeMoveUpdatedComponent(Delta, OutputQuat, true, Hit, ETeleportType::None);
 
-			// Update velocity
-			// We don't want position changes to vastly reverse our direction (which can happen due to penetration fixups etc)
-			/*
-			if (!bPositionCorrected)
-			{
-				const FVector NewLocation = UpdatedComponent->GetComponentLocation();
-				Velocity = ((NewLocation - OldLocation) / DeltaTime);
-			}*/
+		if (Hit.IsValidBlockingHit())
+		{
+			// Try to slide the remaining distance along the surface.
+			SlideAlongSurface(Delta, 1.f-Hit.Time, OutputQuat, Hit.Normal, Hit, true);
 		}
 
-		// Finalize. This is unfortunate. The component mirrors our internal motion state and since we call into it to update, at this point, it has the real position.
-		const FTransform UpdateComponentTransform = GetUpdateComponentTransform();
-		Output.Sync.Location = UpdateComponentTransform.GetLocation();
-
-		// Note that we don't pull the rotation out of the final update transform. Converting back from a quat will lead to a different FRotator than what we are storing
-		// here in the simulation layer. This may not be the best choice for all movement simulations, but is ok for this one.
+		// Update velocity
+		// We don't want position changes to vastly reverse our direction (which can happen due to penetration fixups etc)
+		/*
+		if (!bPositionCorrected)
+		{
+			const FVector NewLocation = UpdatedComponent->GetComponentLocation();
+			Velocity = ((NewLocation - OldLocation) / DeltaTime);
+		}*/
 	}
 
-	void FMovementSimulation::PreSimSync(const FMoveState& SyncState)
+	// Finalize. This is unfortunate. The component mirrors our internal motion state and since we call into it to update, at this point, it has the real position.
+	const FTransform UpdateComponentTransform = GetUpdateComponentTransform();
+	Output.Sync.Location = UpdateComponentTransform.GetLocation();
+
+	// Note that we don't pull the rotation out of the final update transform. Converting back from a quat will lead to a different FRotator than what we are storing
+	// here in the simulation layer. This may not be the best choice for all movement simulations, but is ok for this one.
+}
+
+void FFlyingMovementSimulation::PreSimSync(const FFlyingMovementSyncState& SyncState)
+{
+	// Does checking equality make any sense here? This is unfortunate
+	if (UpdatedComponent->GetComponentLocation().Equals(SyncState.Location) == false || UpdatedComponent->GetComponentQuat().Rotator().Equals(SyncState.Rotation, ROTATOR_TOLERANCE) == false)
 	{
-		// Does checking equality make any sense here? This is unfortunate
-		if (UpdatedComponent->GetComponentLocation().Equals(SyncState.Location) == false || UpdatedComponent->GetComponentQuat().Rotator().Equals(SyncState.Rotation, ROTATOR_TOLERANCE) == false)
-		{
-			FTransform Transform(SyncState.Rotation.Quaternion(), SyncState.Location, UpdatedComponent->GetComponentTransform().GetScale3D() );
-			UpdatedComponent->SetWorldTransform(Transform, false, nullptr, ETeleportType::TeleportPhysics);
+		FTransform Transform(SyncState.Rotation.Quaternion(), SyncState.Location, UpdatedComponent->GetComponentTransform().GetScale3D() );
+		UpdatedComponent->SetWorldTransform(Transform, false, nullptr, ETeleportType::TeleportPhysics);
 
-			UpdatedComponent->ComponentVelocity = SyncState.Velocity;
-		}
+		UpdatedComponent->ComponentVelocity = SyncState.Velocity;
 	}
-
-} // end namespace
\ No newline at end of file
+}
\ No newline at end of file
diff --git a/Engine/Plugins/Runtime/NetworkPrediction/Source/NetworkPrediction/Private/Movement/ParametricMovement.cpp b/Engine/Plugins/Runtime/NetworkPrediction/Source/NetworkPrediction/Private/Movement/ParametricMovement.cpp
--- a/Engine/Plugins/Runtime/NetworkPrediction/Source/NetworkPrediction/Private/Movement/ParametricMovement.cpp
+++ b/Engine/Plugins/Runtime/NetworkPrediction/Source/NetworkPrediction/Private/Movement/ParametricMovement.cpp
@@ -37,44 +37,40 @@ static FAutoConsoleVariableRef CVarFixStepMS(TEXT("parametricmover.FixStep"),
 // ParametricMovement
 // -------------------------------------------------------------------------------------------------------------------------------
 
-namespace ParametricMovement
-{
-	const FName FMovementSimulation::GroupName(TEXT("Parametric"));
+const FName FParametricMovementSimulation::GroupName(TEXT("Parametric"));
 
-	void FMovementSimulation::SimulationTick(const FNetSimTimeStep& TimeStep, const TNetSimInput<TMovementBufferTypes>& Input, const TNetSimOutput<TMovementBufferTypes>& Output)
-	{
-		check(Motion != nullptr); // Must set motion mapping prior to running the simulation
-		const float DeltaSeconds = TimeStep.StepMS.ToRealTimeSeconds();
+void FParametricMovementSimulation::SimulationTick(const FNetSimTimeStep& TimeStep, const TNetSimInput<ParametricMovementBufferTypes>& Input, const TNetSimOutput<ParametricMovementBufferTypes>& Output)
+{
+	check(Motion != nullptr); // Must set motion mapping prior to running the simulation
+	const float DeltaSeconds = TimeStep.StepMS.ToRealTimeSeconds();
 
-		// Advance parametric time. This won't always be linear: we could loop, rewind/bounce, etc
-		const float InputPlayRate = Input.Cmd.PlayRate.Get(Input.Sync.PlayRate); // Returns InputCmds playrate if set, else returns previous state's playrate		
-		Motion->AdvanceParametricTime(Input.Sync.Position, InputPlayRate, Output.Sync.Position, Output.Sync.PlayRate, DeltaSeconds);
+	// Advance parametric time. This won't always be linear: we could loop, rewind/bounce, etc
+	const float InputPlayRate = Input.Cmd.PlayRate.Get(Input.Sync.PlayRate); // Returns InputCmds playrate if set, else returns previous state's playrate		
+	Motion->AdvanceParametricTime(Input.Sync.Position, InputPlayRate, Output.Sync.Position, Output.Sync.PlayRate, DeltaSeconds);
 
-		// We have our time that we should be at. We just need to move primitive component to that position.
-		// Again, note that we expect this cannot fail. We move like this so that it can push things, but we don't expect failure.
-		// (I think it would need to be a different simulation that supports this as a failure case. The tricky thing would be working out
-		// the output Position in the case where a Move is blocked (E.g, you move 50% towards the desired location)
+	// We have our time that we should be at. We just need to move primitive component to that position.
+	// Again, note that we expect this cannot fail. We move like this so that it can push things, but we don't expect failure.
+	// (I think it would need to be a different simulation that supports this as a failure case. The tricky thing would be working out
+	// the output Position in the case where a Move is blocked (E.g, you move 50% towards the desired location)
 
-		FTransform StartingTransform;
-		Motion->MapTimeToTransform(Input.Sync.Position, StartingTransform);
+	FTransform StartingTransform;
+	Motion->MapTimeToTransform(Input.Sync.Position, StartingTransform);
 
-		FTransform NewTransform;
-		Motion->MapTimeToTransform(Output.Sync.Position, NewTransform);
+	FTransform NewTransform;
+	Motion->MapTimeToTransform(Output.Sync.Position, NewTransform);
 
-		FHitResult Hit(1.f);
+	FHitResult Hit(1.f);
 
-		FVector Delta = NewTransform.GetLocation() - StartingTransform.GetLocation();
-		MoveUpdatedComponent(Delta, NewTransform.GetRotation(), true, &Hit, ETeleportType::None);
+	FVector Delta = NewTransform.GetLocation() - StartingTransform.GetLocation();
+	MoveUpdatedComponent(Delta, NewTransform.GetRotation(), true, &Hit, ETeleportType::None);
 
-		if (Hit.IsValidBlockingHit())
-		{
-			FTransform ActualTransform = GetUpdateComponentTransform();
-			FVector ActualDelta = NewTransform.GetLocation() - ActualTransform.GetLocation();
-			UE_LOG(LogParametricMovement, Warning, TEXT("Blocking hit occurred when trying to move parametric mover. ActualDelta: %s"), *ActualDelta.ToString());
-		}
+	if (Hit.IsValidBlockingHit())
+	{
+		FTransform ActualTransform = GetUpdateComponentTransform();
+		FVector ActualDelta = NewTransform.GetLocation() - ActualTransform.GetLocation();
+		UE_LOG(LogParametricMovement, Warning, TEXT("Blocking hit occurred when trying to move parametric mover. ActualDelta: %s"), *ActualDelta.ToString());
 	}
-};
-
+}
 
 // -------------------------------------------------------------------------------------------------------------------------------
 // UParametricMovementComponent
@@ -93,20 +89,20 @@ INetworkedSimulationModel* UParametricMovementComponent::InstantiateNetworkedSim
 	}
 
 	// Simulation
-	ParametricMovement::FMoveState InitialSyncState;
-	ParametricMovement::FAuxState InitialAuxState;
+	FParametricSyncState InitialSyncState;
+	FParametricAuxState InitialAuxState;
 
 	InitParametricMovementSimulation(new FParametricMovementSimulation(), InitialSyncState, InitialAuxState);
 
 	// Model
 	if (ParametricMoverCVars::FixStep > 0)
 	{
-		auto *NewModel = new ParametricMovement::FMovementSystem<32>( ParametricMovementSimulation, this);
+		auto *NewModel = new FParametricMovementSystem<32>( ParametricMovementSimulation, this);
 		InitParametricMovementNetSimModel(NewModel);
 		return NewModel;
 	}
 	
-	auto *NewModel = new ParametricMovement::FMovementSystem<>(ParametricMovementSimulation, this);
+	auto *NewModel = new FParametricMovementSystem<>(ParametricMovementSimulation, this);
 	InitParametricMovementNetSimModel(NewModel);
 	return NewModel;
 }
@@ -129,7 +125,7 @@ void UParametricMovementComponent::BeginPlay()
 	ParametricMotion.CachedStartingTransform = UpdatedComponent->GetComponentToWorld();
 }
 
-void UParametricMovementComponent::FinalizeFrame(const ParametricMovement::FMoveState& SyncState, const ParametricMovement::FAuxState& AuxState)
+void UParametricMovementComponent::FinalizeFrame(const FParametricSyncState& SyncState, const FParametricAuxState& AuxState)
 {
 	FTransform NewTransform;
 	ParametricMotion.MapTimeToTransform(SyncState.Position, NewTransform);
@@ -147,7 +143,7 @@ const AActor* UParametricMovementComponent::GetVLogOwner() const
 	return GetOwner();
 }
 
-void UParametricMovementComponent::VisualLog(const ParametricMovement::FInputCmd* Input, const ParametricMovement::FMoveState* Sync, const ParametricMovement::FAuxState* Aux, const FVisualLoggingParameters& SystemParameters) const
+void UParametricMovementComponent::VisualLog(const FParametricInputCmd* Input, const FParametricSyncState* Sync, const FParametricAuxState* Aux, const FVisualLoggingParameters& SystemParameters) const
 {
 	FTransform Transform;
 	ParametricMotion.MapTimeToTransform(Sync->Position, Transform);
@@ -193,7 +189,7 @@ void UParametricMovementComponent::TickComponent(float DeltaTime, enum ELevelTic
 	}
 }
 
-void UParametricMovementComponent::ProduceInput(const FNetworkSimTime DeltaTimeSeconds, ParametricMovement::FInputCmd& Cmd)
+void UParametricMovementComponent::ProduceInput(const FNetworkSimTime DeltaTimeSeconds, FParametricInputCmd& Cmd)
 {
 	Cmd.PlayRate = PendingPlayRate;
 	PendingPlayRate.Reset();
diff --git a/Engine/Plugins/Runtime/NetworkPrediction/Source/NetworkPrediction/Public/Movement/FlyingMovement.h b/Engine/Plugins/Runtime/NetworkPrediction/Source/NetworkPrediction/Public/Movement/FlyingMovement.h
--- a/Engine/Plugins/Runtime/NetworkPrediction/Source/NetworkPrediction/Public/Movement/FlyingMovement.h
+++ b/Engine/Plugins/Runtime/NetworkPrediction/Source/NetworkPrediction/Public/Movement/FlyingMovement.h
@@ -24,161 +24,157 @@
 
 class USceneComponent;
 
-namespace FlyingMovement
+// State the client generates
+struct FFlyingMovementInputCmd
 {
-	// State the client generates
-	struct FInputCmd
-	{
-		// Input: "pure" input for this frame. At this level, frame time has not been accounted for. (E.g., "move straight" would be (1,0,0) regardless of frame time)
-		FRotator RotationInput;
-		FVector MovementInput;
+	// Input: "pure" input for this frame. At this level, frame time has not been accounted for. (E.g., "move straight" would be (1,0,0) regardless of frame time)
+	FRotator RotationInput;
+	FVector MovementInput;
+
+	FFlyingMovementInputCmd()
+		: RotationInput(ForceInitToZero)
+		, MovementInput(ForceInitToZero)
+	{ }
 
-		FInputCmd()
-			: RotationInput(ForceInitToZero)
-			, MovementInput(ForceInitToZero)
-		{ }
+	void NetSerialize(const FNetSerializeParams& P)
+	{
+		P.Ar << RotationInput;
+		P.Ar << MovementInput;
+	}
 
-		void NetSerialize(const FNetSerializeParams& P)
+	void Log(FStandardLoggingParameters& P) const
+	{
+		if (P.Context == EStandardLoggingContext::HeaderOnly)
 		{
-			P.Ar << RotationInput;
-			P.Ar << MovementInput;
+			P.Ar->Logf(TEXT(" %d "), P.Frame);
 		}
-
-		void Log(FStandardLoggingParameters& P) const
+		else if (P.Context == EStandardLoggingContext::Full)
 		{
-			if (P.Context == EStandardLoggingContext::HeaderOnly)
-			{
-				P.Ar->Logf(TEXT(" %d "), P.Frame);
-			}
-			else if (P.Context == EStandardLoggingContext::Full)
-			{
-				P.Ar->Logf(TEXT("Movement: %s"), *MovementInput.ToString());
-				P.Ar->Logf(TEXT("ViewRot: %s"), *RotationInput.ToString());
-			}
+			P.Ar->Logf(TEXT("Movement: %s"), *MovementInput.ToString());
+			P.Ar->Logf(TEXT("ViewRot: %s"), *RotationInput.ToString());
 		}
-	};
+	}
+};
+
+// State we are evolving frame to frame and keeping in sync
+struct FFlyingMovementSyncState
+{
+	FVector Location;
+	FVector Velocity;
+	FRotator Rotation;
+
+	FFlyingMovementSyncState()
+		: Location(ForceInitToZero)
+		, Velocity(ForceInitToZero)
+		, Rotation(ForceInitToZero)
+	{ }
 
-	// State we are evolving frame to frame and keeping in sync
-	struct FMoveState
+	void NetSerialize(const FNetSerializeParams& P)
 	{
-		FVector Location;
-		FVector Velocity;
-		FRotator Rotation;
+		P.Ar << Location;
+		P.Ar << Velocity;
+		P.Ar << Rotation;
+	}
 
-		FMoveState()
-			: Location(ForceInitToZero)
-			, Velocity(ForceInitToZero)
-			, Rotation(ForceInitToZero)
-		{ }
+	// Compare this state with AuthorityState. return true if a reconcile (correction) should happen
+	bool ShouldReconcile(const FFlyingMovementSyncState& AuthorityState)
+	{
+		const float ErrorTolerance = 1.f;
+		return !AuthorityState.Location.Equals(Location, ErrorTolerance);
+	}
 
-		void NetSerialize(const FNetSerializeParams& P)
+	void Log(FStandardLoggingParameters& Params) const
+	{
+		if (Params.Context == EStandardLoggingContext::HeaderOnly)
 		{
-			P.Ar << Location;
-			P.Ar << Velocity;
-			P.Ar << Rotation;
+			Params.Ar->Logf(TEXT(" %d "), Params.Frame);
 		}
-
-		// Compare this state with AuthorityState. return true if a reconcile (correction) should happen
-		bool ShouldReconcile(const FMoveState& AuthorityState)
+		else if (Params.Context == EStandardLoggingContext::Full)
 		{
-			const float ErrorTolerance = 1.f;
-			return !AuthorityState.Location.Equals(Location, ErrorTolerance);
+			Params.Ar->Logf(TEXT("Frame: %d"), Params.Frame);
+			Params.Ar->Logf(TEXT("Loc: %s"), *Location.ToString());
+			Params.Ar->Logf(TEXT("Vel: %s"), *Velocity.ToString());
+			Params.Ar->Logf(TEXT("Rot: %s"), *Rotation.ToString());
 		}
+	}
 
-		void Log(FStandardLoggingParameters& Params) const
-		{
-			if (Params.Context == EStandardLoggingContext::HeaderOnly)
-			{
-				Params.Ar->Logf(TEXT(" %d "), Params.Frame);
-			}
-			else if (Params.Context == EStandardLoggingContext::Full)
-			{
-				Params.Ar->Logf(TEXT("Frame: %d"), Params.Frame);
-				Params.Ar->Logf(TEXT("Loc: %s"), *Location.ToString());
-				Params.Ar->Logf(TEXT("Vel: %s"), *Velocity.ToString());
-				Params.Ar->Logf(TEXT("Rot: %s"), *Rotation.ToString());
-			}
-		}
+	static void Interpolate(const FFlyingMovementSyncState& From, const FFlyingMovementSyncState& To, const float PCT, FFlyingMovementSyncState& OutDest)
+	{
+		OutDest.Location = From.Location + ((To.Location - From.Location) * PCT);
+		OutDest.Rotation = From.Rotation + ((To.Rotation - From.Rotation) * PCT);
+	}
+};
+
+// Auxiliary state that is input into the simulation. Doesn't change during the simulation tick.
+// (It can change and even be predicted but doing so will trigger more bookeeping, etc. Changes will happen "next tick")
+struct FFlyingMovementAuxState
+{	
+	float MaxSpeed = 1200.f;
+	float TurningBoost = 8.f;
+	float Deceleration = 8000.f;
+	float Acceleration = 4000.f;
+
+	void NetSerialize(const FNetSerializeParams& P)
+	{
+		P.Ar << MaxSpeed;
+		P.Ar << TurningBoost;
+		P.Ar << Deceleration;
+		P.Ar << Acceleration;
+	}
 
-		static void Interpolate(const FMoveState& From, const FMoveState& To, const float PCT, FMoveState& OutDest)
-		{
-			OutDest.Location = From.Location + ((To.Location - From.Location) * PCT);
-			OutDest.Rotation = From.Rotation + ((To.Rotation - From.Rotation) * PCT);
-		}
-	};
-
-	// Auxiliary state that is input into the simulation. Doesn't change during the simulation tick.
-	// (It can change and even be predicted but doing so will trigger more bookeeping, etc. Changes will happen "next tick")
-	struct FAuxState
-	{	
-		float MaxSpeed = 1200.f;
-		float TurningBoost = 8.f;
-		float Deceleration = 8000.f;
-		float Acceleration = 4000.f;
-
-		void NetSerialize(const FNetSerializeParams& P)
+	void Log(FStandardLoggingParameters& P) const
+	{
+		if (P.Context == EStandardLoggingContext::HeaderOnly)
 		{
-			P.Ar << MaxSpeed;
-			P.Ar << TurningBoost;
-			P.Ar << Deceleration;
-			P.Ar << Acceleration;
+			P.Ar->Logf(TEXT(" %d "), P.Frame);
 		}
-
-		void Log(FStandardLoggingParameters& P) const
+		else if (P.Context == EStandardLoggingContext::Full)
 		{
-			if (P.Context == EStandardLoggingContext::HeaderOnly)
-			{
-				P.Ar->Logf(TEXT(" %d "), P.Frame);
-			}
-			else if (P.Context == EStandardLoggingContext::Full)
-			{
-				P.Ar->Logf(TEXT("MaxSpeed: %.2f"), MaxSpeed);
-				P.Ar->Logf(TEXT("TurningBoost: %.2f"), TurningBoost);
-				P.Ar->Logf(TEXT("Deceleration: %.2f"), Deceleration);
-				P.Ar->Logf(TEXT("Acceleration: %.2f"), Acceleration);
-			}
+			P.Ar->Logf(TEXT("MaxSpeed: %.2f"), MaxSpeed);
+			P.Ar->Logf(TEXT("TurningBoost: %.2f"), TurningBoost);
+			P.Ar->Logf(TEXT("Deceleration: %.2f"), Deceleration);
+			P.Ar->Logf(TEXT("Acceleration: %.2f"), Acceleration);
 		}
-	};
+	}
+};
 	
-	using TMovementBufferTypes = TNetworkSimBufferTypes<FInputCmd, FMoveState, FAuxState>;
-
-	static FName SimulationGroupName("FlyingMovement");
+using FlyingMovementBufferTypes = TNetworkSimBufferTypes<FFlyingMovementInputCmd, FFlyingMovementSyncState, FFlyingMovementAuxState>;
 
-	class FMovementSimulation : public FBaseMovementSimulation
-	{
-	public:
+static FName SimulationGroupName("FlyingMovement");
 
-		/** Tick group the simulation maps to */
-		static const FName GroupName;
+class FFlyingMovementSimulation : public FBaseMovementSimulation
+{
+public:
 
-		/** Dev tool to force simple mispredict */
-		static bool ForceMispredict;
+	/** Tick group the simulation maps to */
+	static const FName GroupName;
 
-		/** Main update function */
-		NETWORKPREDICTION_API void SimulationTick(const FNetSimTimeStep& TimeStep, const TNetSimInput<TMovementBufferTypes>& Input, const TNetSimOutput<TMovementBufferTypes>& Output);
+	/** Dev tool to force simple mispredict */
+	static bool ForceMispredict;
 
-		// Called prior to running the sim to make sure to make sure the collision component is in the right place. 
-		// This is unfortunate and not good, but is needed to ensure our collision and world position have not been moved out from under us.
-		// Refactoring primitive component movement to allow the sim to do all collision queries outside of the component code would be ideal.
-		void PreSimSync(const FMoveState& SyncState);
+	/** Main update function */
+	NETWORKPREDICTION_API void SimulationTick(const FNetSimTimeStep& TimeStep, const TNetSimInput<FlyingMovementBufferTypes>& Input, const TNetSimOutput<FlyingMovementBufferTypes>& Output);
 
-	protected:
+	// Called prior to running the sim to make sure to make sure the collision component is in the right place. 
+	// This is unfortunate and not good, but is needed to ensure our collision and world position have not been moved out from under us.
+	// Refactoring primitive component movement to allow the sim to do all collision queries outside of the component code would be ideal.
+	void PreSimSync(const FFlyingMovementSyncState& SyncState);
 
-		float SlideAlongSurface(const FVector& Delta, float Time, const FQuat Rotation, const FVector& Normal, FHitResult& Hit, bool bHandleImpact);
-	};
+protected:
 
-	// Actual definition of our network simulation.
-	template<int32 InFixedStepMS=0>
-	using FMovementSystem = TNetworkedSimulationModel<FMovementSimulation, TMovementBufferTypes, TNetworkSimTickSettings<InFixedStepMS>>;
+	float SlideAlongSurface(const FVector& Delta, float Time, const FQuat Rotation, const FVector& Normal, FHitResult& Hit, bool bHandleImpact);
+};
 
-	// general tolerance value for rotation checks
-	static const float ROTATOR_TOLERANCE = (1e-3);
+// Actual definition of our network simulation.
+template<int32 InFixedStepMS=0>
+using FFlyingMovementSystem = TNetworkedSimulationModel<FFlyingMovementSimulation, FlyingMovementBufferTypes, TNetworkSimTickSettings<InFixedStepMS>>;
 
-} // End namespace
+// general tolerance value for rotation checks
+static const float ROTATOR_TOLERANCE = (1e-3);
 
 // Needed to trick UHT into letting UMockNetworkSimulationComponent implement. UHT cannot parse the ::
 // Also needed for forward declaring. Can't just be a typedef/using =
-class IFlyingMovementSystemDriver : public TNetworkedSimulationModelDriver<FlyingMovement::TMovementBufferTypes> { };
+class IFlyingMovementSystemDriver : public TNetworkedSimulationModelDriver<FlyingMovementBufferTypes> { };
 
 // -------------------------------------------------------------------------------------------------------------------------------
 // ActorComponent for running FlyingMovement 
@@ -195,19 +191,19 @@ public:
 	
 	virtual void TickComponent(float DeltaTime, enum ELevelTick TickType, FActorComponentTickFunction *ThisTickFunction) override;
 
-	DECLARE_DELEGATE_TwoParams(FProduceFlyingInput, const FNetworkSimTime /*SimTime*/, FlyingMovement::FInputCmd& /*Cmd*/)
+	DECLARE_DELEGATE_TwoParams(FProduceFlyingInput, const FNetworkSimTime /*SimTime*/, FFlyingMovementInputCmd& /*Cmd*/)
 	FProduceFlyingInput ProduceInputDelegate;
 
-	TNetSimStateAccessor<FlyingMovement::FMoveState> MovementSyncState;
-	TNetSimStateAccessor<FlyingMovement::FAuxState> MovementAuxState;
+	TNetSimStateAccessor<FFlyingMovementSyncState> MovementSyncState;
+	TNetSimStateAccessor<FFlyingMovementAuxState> MovementAuxState;
 
 	// IFlyingMovementSystemDriver
 	FString GetDebugName() const override;
 	const AActor* GetVLogOwner() const override;
-	void VisualLog(const FlyingMovement::FInputCmd* Input, const FlyingMovement::FMoveState* Sync, const FlyingMovement::FAuxState* Aux, const FVisualLoggingParameters& SystemParameters) const override;
+	void VisualLog(const FFlyingMovementInputCmd* Input, const FFlyingMovementSyncState* Sync, const FFlyingMovementAuxState* Aux, const FVisualLoggingParameters& SystemParameters) const override;
 	
-	void ProduceInput(const FNetworkSimTime SimTime, FlyingMovement::FInputCmd& Cmd) override;
-	void FinalizeFrame(const FlyingMovement::FMoveState& SyncState, const FlyingMovement::FAuxState& AuxState) override;
+	void ProduceInput(const FNetworkSimTime SimTime, FFlyingMovementInputCmd& Cmd) override;
+	void FinalizeFrame(const FFlyingMovementSyncState& SyncState, const FFlyingMovementAuxState& AuxState) override;
 
 	UPROPERTY(EditAnywhere, BlueprintReadOnly, Category=FlyingMovementNetworking)
 	bool bEnableInterpolation = true;
@@ -216,11 +212,11 @@ protected:
 
 	// Network Prediction
 	virtual INetworkedSimulationModel* InstantiateNetworkedSimulation() override;
-	FlyingMovement::FMovementSimulation* MovementSimulation = nullptr;
+	FFlyingMovementSimulation* MovementSimulation = nullptr;
 
 	// Init function. This is broken up from ::InstantiateNetworkedSimulation and templated so that subclasses can share the init code
 	template<typename TSimulation>
-	void InitFlyingMovementSimulation(TSimulation* Simulation, FlyingMovement::FMoveState& InitialSyncState, FlyingMovement::FAuxState& InitialAuxState)
+	void InitFlyingMovementSimulation(TSimulation* Simulation, FFlyingMovementSyncState& InitialSyncState, FFlyingMovementAuxState& InitialAuxState)
 	{
 		check(UpdatedComponent);
 		check(MovementSimulation == nullptr); // Reinstantiation not supported
diff --git a/Engine/Plugins/Runtime/NetworkPrediction/Source/NetworkPrediction/Public/Movement/ParametricMovement.h b/Engine/Plugins/Runtime/NetworkPrediction/Source/NetworkPrediction/Public/Movement/ParametricMovement.h
--- a/Engine/Plugins/Runtime/NetworkPrediction/Source/NetworkPrediction/Public/Movement/ParametricMovement.h
+++ b/Engine/Plugins/Runtime/NetworkPrediction/Source/NetworkPrediction/Public/Movement/ParametricMovement.h
@@ -8,127 +8,123 @@
 #include "ParametricMovement.generated.h"
 
 struct FSimpleParametricMotion;
+class IMovementDriver;
 
-namespace ParametricMovement
+// State the client generates
+struct FParametricInputCmd
 {
-	class IMovementDriver;
+	// Input Playrate. This being set can be thought of "telling the simulation what its new playrate should be"
+	TOptional<float> PlayRate;
 
-	// State the client generates
-	struct FInputCmd
+	void NetSerialize(const FNetSerializeParams& P)
 	{
-		// Input Playrate. This being set can be thought of "telling the simulation what its new playrate should be"
-		TOptional<float> PlayRate;
+		P.Ar << PlayRate;
+	}
 
-		void NetSerialize(const FNetSerializeParams& P)
+	void Log(FStandardLoggingParameters& P) const
+	{
+		if (P.Context == EStandardLoggingContext::HeaderOnly)
 		{
-			P.Ar << PlayRate;
+			P.Ar->Logf(TEXT(" %d "), P.Frame);
 		}
-
-		void Log(FStandardLoggingParameters& P) const
+		else if (P.Context == EStandardLoggingContext::Full)
 		{
-			if (P.Context == EStandardLoggingContext::HeaderOnly)
+			if (PlayRate.IsSet())
 			{
-				P.Ar->Logf(TEXT(" %d "), P.Frame);
-			}
-			else if (P.Context == EStandardLoggingContext::Full)
-			{
-				if (PlayRate.IsSet())
-				{
-					P.Ar->Logf(TEXT("PlayRate: %.2f"), PlayRate.GetValue());
-				}
+				P.Ar->Logf(TEXT("PlayRate: %.2f"), PlayRate.GetValue());
 			}
 		}
-	};
+	}
+};
+
+// State we are evolving frame to frame and keeping in sync
+struct FParametricSyncState
+{
+	float Position;
+	float PlayRate;
 
-	// State we are evolving frame to frame and keeping in sync
-	struct FMoveState
+	void NetSerialize(const FNetSerializeParams& P)
 	{
-		float Position;
-		float PlayRate;
+		P.Ar << Position;
+		P.Ar << PlayRate;
+	}
 
-		void NetSerialize(const FNetSerializeParams& P)
-		{
-			P.Ar << Position;
-			P.Ar << PlayRate;
-		}
+	// Compare this state with AuthorityState. return true if a reconcile (correction) should happen
+	bool ShouldReconcile(const FParametricSyncState& AuthorityState)
+	{
+		const float ErrorTolerance = 0.01f;
+		return (FMath::Abs<float>(AuthorityState.Position - Position) > ErrorTolerance) || (FMath::Abs<float>(AuthorityState.PlayRate - PlayRate) > ErrorTolerance);
+	}
 
-		// Compare this state with AuthorityState. return true if a reconcile (correction) should happen
-		bool ShouldReconcile(const FMoveState& AuthorityState)
+	void Log(FStandardLoggingParameters& Params) const
+	{
+		if (Params.Context == EStandardLoggingContext::HeaderOnly)
 		{
-			const float ErrorTolerance = 0.01f;
-			return (FMath::Abs<float>(AuthorityState.Position - Position) > ErrorTolerance) || (FMath::Abs<float>(AuthorityState.PlayRate - PlayRate) > ErrorTolerance);
+			Params.Ar->Logf(TEXT(" %d "), Params.Frame);
 		}
-
-		void Log(FStandardLoggingParameters& Params) const
+		else if (Params.Context == EStandardLoggingContext::Full)
 		{
-			if (Params.Context == EStandardLoggingContext::HeaderOnly)
-			{
-				Params.Ar->Logf(TEXT(" %d "), Params.Frame);
-			}
-			else if (Params.Context == EStandardLoggingContext::Full)
-			{
-				Params.Ar->Logf(TEXT("Frame: %d"), Params.Frame);
-				Params.Ar->Logf(TEXT("Pos: %.2f"), Position);
-				Params.Ar->Logf(TEXT("Rate: %.2f"), PlayRate);
-			}
+			Params.Ar->Logf(TEXT("Frame: %d"), Params.Frame);
+			Params.Ar->Logf(TEXT("Pos: %.2f"), Position);
+			Params.Ar->Logf(TEXT("Rate: %.2f"), PlayRate);
 		}
+	}
 
-		void VisualLog(const FVisualLoggingParameters& Parameters, IMovementDriver* Driver, IMovementDriver* LogDriver) const;
+	void VisualLog(const FVisualLoggingParameters& Parameters, IMovementDriver* Driver, IMovementDriver* LogDriver) const;
 
-		static void Interpolate(const FMoveState& From, const FMoveState& To, const float PCT, FMoveState& OutDest)
-		{
-			OutDest.Position = From.Position + ((To.Position - From.Position) * PCT);
-			OutDest.PlayRate = From.PlayRate + ((To.PlayRate - From.PlayRate) * PCT);
-		}
-	};
+	static void Interpolate(const FParametricSyncState& From, const FParametricSyncState& To, const float PCT, FParametricSyncState& OutDest)
+	{
+		OutDest.Position = From.Position + ((To.Position - From.Position) * PCT);
+		OutDest.PlayRate = From.PlayRate + ((To.PlayRate - From.PlayRate) * PCT);
+	}
+};
 
-	// Auxiliary state that is input into the simulation. Doesn't change during the simulation tick.
-	// (It can change and even be predicted but doing so will trigger more bookeeping, etc. Changes will happen "next tick")
-	struct FAuxState
+// Auxiliary state that is input into the simulation. Doesn't change during the simulation tick.
+// (It can change and even be predicted but doing so will trigger more bookeeping, etc. Changes will happen "next tick")
+struct FParametricAuxState
+{
+	float Multiplier=1;
+	void NetSerialize(const FNetSerializeParams& P)
 	{
-		float Multiplier=1;
-		void NetSerialize(const FNetSerializeParams& P)
+		P.Ar << Multiplier;
+	}
+
+	void Log(FStandardLoggingParameters& Params) const
+	{
+		if (Params.Context == EStandardLoggingContext::HeaderOnly)
 		{
-			P.Ar << Multiplier;
+			Params.Ar->Logf(TEXT(" %d "), Params.Frame);
 		}
-
-		void Log(FStandardLoggingParameters& Params) const
+		else if (Params.Context == EStandardLoggingContext::Full)
 		{
-			if (Params.Context == EStandardLoggingContext::HeaderOnly)
-			{
-				Params.Ar->Logf(TEXT(" %d "), Params.Frame);
-			}
-			else if (Params.Context == EStandardLoggingContext::Full)
-			{
-				Params.Ar->Logf(TEXT("Multiplier: %f"), Multiplier);
-			}
+			Params.Ar->Logf(TEXT("Multiplier: %f"), Multiplier);
 		}
-	};
+	}
+};
 
-	using TMovementBufferTypes = TNetworkSimBufferTypes<FInputCmd, FMoveState, FAuxState>;
+using ParametricMovementBufferTypes = TNetworkSimBufferTypes<FParametricInputCmd, FParametricSyncState, FParametricAuxState>;
 	
-	class FMovementSimulation : public FBaseMovementSimulation
-	{
-	public:
-		static const FName GroupName;
+class FParametricMovementSimulation : public FBaseMovementSimulation
+{
+public:
+	static const FName GroupName;
 
-		void SimulationTick(const FNetSimTimeStep& TimeStep, const TNetSimInput<TMovementBufferTypes>& Input, const TNetSimOutput<TMovementBufferTypes>& Output);
+	void SimulationTick(const FNetSimTimeStep& TimeStep, const TNetSimInput<ParametricMovementBufferTypes>& Input, const TNetSimOutput<ParametricMovementBufferTypes>& Output);
+
+	// Pointer to our static mapping of time->position
+	const FSimpleParametricMotion* Motion = nullptr;
+};
 
-		// Pointer to our static mapping of time->position
-		const FSimpleParametricMotion* Motion = nullptr;
-	};
 
+// Actual definition of our network simulation.
+template<int32 FixedStepMS=0>
+using FParametricMovementSystem = TNetworkedSimulationModel<FParametricMovementSimulation, ParametricMovementBufferTypes, TNetworkSimTickSettings<FixedStepMS>>;
 
-	// Actual definition of our network simulation.
-	template<int32 FixedStepMS=0>
-	using FMovementSystem = TNetworkedSimulationModel<FMovementSimulation, TMovementBufferTypes, TNetworkSimTickSettings<FixedStepMS>>;
 
-} // End namespace
 
 // Needed to trick UHT into letting UMockNetworkSimulationComponent implement. UHT cannot parse the ::
 // Also needed for forward declaring. Can't just be a typedef/using =
-class IParametricMovementDriver : public TNetworkedSimulationModelDriver<ParametricMovement::TMovementBufferTypes> { };
-class FParametricMovementSimulation : public ParametricMovement::FMovementSimulation { };
+class IParametricMovementDriver : public TNetworkedSimulationModelDriver<ParametricMovementBufferTypes> { };
 
 // Extremely simple struct for defining parametric motion. This is editable in UParametricMovementComponent's defaults, and also used by the simulation code above. 
 USTRUCT(BlueprintType)
@@ -176,15 +172,15 @@ class NETWORKPREDICTION_API UParametricMovementComponent : public UBaseMovementC
 	// Base TNetworkModelSimulation driver
 	FString GetDebugName() const override;
 	const AActor* GetVLogOwner() const override;
-	void VisualLog(const ParametricMovement::FInputCmd* Input, const ParametricMovement::FMoveState* Sync, const ParametricMovement::FAuxState* Aux, const FVisualLoggingParameters& SystemParameters) const override;
+	void VisualLog(const FParametricInputCmd* Input, const FParametricSyncState* Sync, const FParametricAuxState* Aux, const FVisualLoggingParameters& SystemParameters) const override;
 
-	void ProduceInput(const FNetworkSimTime SimTime, ParametricMovement::FInputCmd& Cmd);
-	void FinalizeFrame(const ParametricMovement::FMoveState& SyncState, const ParametricMovement::FAuxState& AuxState) override;
+	void ProduceInput(const FNetworkSimTime SimTime, FParametricInputCmd& Cmd);
+	void FinalizeFrame(const FParametricSyncState& SyncState, const FParametricAuxState& AuxState) override;
 
 protected:
 
-	TNetSimStateAccessor<ParametricMovement::FMoveState> MovementSyncState;
-	TNetSimStateAccessor<ParametricMovement::FAuxState> MovementAuxState;
+	TNetSimStateAccessor<FParametricSyncState> MovementSyncState;
+	TNetSimStateAccessor<FParametricAuxState> MovementAuxState;
 
 	virtual INetworkedSimulationModel* InstantiateNetworkedSimulation() override;
 	FNetworkSimulationModelInitParameters GetSimulationInitParameters(ENetRole Role) override;
@@ -192,7 +188,7 @@ protected:
 	FParametricMovementSimulation* ParametricMovementSimulation = nullptr;
 
 	template<typename TSimulation>
-	void InitParametricMovementSimulation(TSimulation* Simulation, ParametricMovement::FMoveState& InitialSyncState, ParametricMovement::FAuxState& InitialAuxState)
+	void InitParametricMovementSimulation(TSimulation* Simulation, FParametricSyncState& InitialSyncState, FParametricAuxState& InitialAuxState)
 	{
 		check(ParametricMovementSimulation == nullptr);
 		ParametricMovementSimulation = Simulation;
diff --git a/Engine/Plugins/Runtime/NetworkPrediction/Source/NetworkPrediction/Public/NetworkedSimulationModel.h b/Engine/Plugins/Runtime/NetworkPrediction/Source/NetworkPrediction/Public/NetworkedSimulationModel.h
--- a/Engine/Plugins/Runtime/NetworkPrediction/Source/NetworkPrediction/Public/NetworkedSimulationModel.h
+++ b/Engine/Plugins/Runtime/NetworkPrediction/Source/NetworkPrediction/Public/NetworkedSimulationModel.h
@@ -35,6 +35,22 @@
 //
 // ----------------------------------------------------------------------------------------------------------------------------------------------------------------------
 
+
+struct TNetSimModelTraitsBase
+{
+
+};
+
+template<typename T>
+struct TNetSimModelTraits : public TNetSimModelTraitsBase
+{
+	using TSimulation = void;
+	using TBufferTypes = void;	
+	using TTickSettings = TNetworkSimTickSettings<>;
+
+};
+
+
 template <
 	typename InTSimulation,								// Final Simulation class
 	typename TUserBufferTypes,							// The user types (input, sync, aux, debug). Note this gets wrapped in TInternalBufferTypes internally.
diff --git a/Engine/Plugins/Runtime/NetworkPredictionExtras/Source/NetworkPredictionExtras/Private/MockAbilitySimulation.cpp b/Engine/Plugins/Runtime/NetworkPredictionExtras/Source/NetworkPredictionExtras/Private/MockAbilitySimulation.cpp
--- a/Engine/Plugins/Runtime/NetworkPredictionExtras/Source/NetworkPredictionExtras/Private/MockAbilitySimulation.cpp
+++ b/Engine/Plugins/Runtime/NetworkPredictionExtras/Source/NetworkPredictionExtras/Private/MockAbilitySimulation.cpp
@@ -159,7 +159,7 @@ void FMockAbilitySimulation::SimulationTick(const FNetSimTimeStep& TimeStep, con
 		LocalCmd.bBlinkPressed = false;
 		LocalCmd.bSprintPressed = false;
 		
-		FlyingMovement::FMovementSimulation::SimulationTick(TimeStep, { LocalCmd, LocalSync, LocalAux }, { Output.Sync, Output.Aux, Output.CueDispatch });
+		FFlyingMovementSimulation::SimulationTick(TimeStep, { LocalCmd, LocalSync, LocalAux }, { Output.Sync, Output.Aux, Output.CueDispatch });
 
 		if (NewDashTimeLeft == 0)
 		{
@@ -215,7 +215,7 @@ void FMockAbilitySimulation::SimulationTick(const FNetSimTimeStep& TimeStep, con
 		Output.Aux.Get()->bIsSprinting = bIsSprinting;
 	}
 
-	FlyingMovement::FMovementSimulation::SimulationTick(TimeStep, { LocalCmd, LocalSync, LocalAux }, { Output.Sync, Output.Aux, Output.CueDispatch });
+	FFlyingMovementSimulation::SimulationTick(TimeStep, { LocalCmd, LocalSync, LocalAux }, { Output.Sync, Output.Aux, Output.CueDispatch });
 }
 
 // -------------------------------------------------------------------------------------------------------------
diff --git a/Engine/Plugins/Runtime/NetworkPredictionExtras/Source/NetworkPredictionExtras/Private/NetworkPredictionExtrasFlyingPawn.cpp b/Engine/Plugins/Runtime/NetworkPredictionExtras/Source/NetworkPredictionExtras/Private/NetworkPredictionExtrasFlyingPawn.cpp
--- a/Engine/Plugins/Runtime/NetworkPredictionExtras/Source/NetworkPredictionExtras/Private/NetworkPredictionExtrasFlyingPawn.cpp
+++ b/Engine/Plugins/Runtime/NetworkPredictionExtras/Source/NetworkPredictionExtras/Private/NetworkPredictionExtrasFlyingPawn.cpp
@@ -210,7 +210,7 @@ void ANetworkPredictionExtrasFlyingPawn::PrintDebug()
 
 float ANetworkPredictionExtrasFlyingPawn::GetMaxMoveSpeed() const
 {
-	if (const FlyingMovement::FAuxState* AuxState = GetAuxStateRead())
+	if (const FFlyingMovementAuxState* AuxState = GetAuxStateRead())
 	{
 		return AuxState->MaxSpeed;
 	}
@@ -219,7 +219,7 @@ float ANetworkPredictionExtrasFlyingPawn::GetMaxMoveSpeed() const
 
 void ANetworkPredictionExtrasFlyingPawn::SetMaxMoveSpeed(float NewMaxMoveSpeed)
 {
-	if (FlyingMovement::FAuxState* AuxState = GetAuxStateWrite())
+	if (FFlyingMovementAuxState* AuxState = GetAuxStateWrite())
 	{
 		AuxState->MaxSpeed = NewMaxMoveSpeed;
 	}
@@ -227,14 +227,14 @@ void ANetworkPredictionExtrasFlyingPawn::SetMaxMoveSpeed(float NewMaxMoveSpeed)
 
 void ANetworkPredictionExtrasFlyingPawn::AddMaxMoveSpeed(float AdditiveMaxMoveSpeed)
 {
-	if (FlyingMovement::FAuxState* AuxState = GetAuxStateWrite())
+	if (FFlyingMovementAuxState* AuxState = GetAuxStateWrite())
 	{
 		//UE_LOG(LogTemp, Warning, TEXT("MaxSpeed: %.2f += %.2f"), AuxState->MaxSpeed, AdditiveMaxMoveSpeed);
 		AuxState->MaxSpeed += AdditiveMaxMoveSpeed;
 	}
 }
 
-const FlyingMovement::FAuxState* ANetworkPredictionExtrasFlyingPawn::GetAuxStateRead() const
+const FFlyingMovementAuxState* ANetworkPredictionExtrasFlyingPawn::GetAuxStateRead() const
 {
 	if (ensure(FlyingMovementComponent))
 	{
@@ -243,7 +243,7 @@ const FlyingMovement::FAuxState* ANetworkPredictionExtrasFlyingPawn::GetAuxState
 	return nullptr;
 }
 
-FlyingMovement::FAuxState* ANetworkPredictionExtrasFlyingPawn::GetAuxStateWrite()
+FFlyingMovementAuxState* ANetworkPredictionExtrasFlyingPawn::GetAuxStateWrite()
 {
 	if (ensure(FlyingMovementComponent))
 	{
@@ -264,7 +264,7 @@ int32 ANetworkPredictionExtrasFlyingPawn::GetPendingFrame() const
 	return 0;
 }
 
-void ANetworkPredictionExtrasFlyingPawn::ProduceInput(const FNetworkSimTime SimTime, FlyingMovement::FInputCmd& Cmd)
+void ANetworkPredictionExtrasFlyingPawn::ProduceInput(const FNetworkSimTime SimTime, FFlyingMovementInputCmd& Cmd)
 {
 	// Generate user commands. Called right before the flying movement simulation will tick (for a locally controlled pawn)
 	// This isn't meant to be the best way of doing a camera system. It is just meant to show a couple of ways it may be done
diff --git a/Engine/Plugins/Runtime/NetworkPredictionExtras/Source/NetworkPredictionExtras/Public/MockAbilitySimulation.h b/Engine/Plugins/Runtime/NetworkPredictionExtras/Source/NetworkPredictionExtras/Public/MockAbilitySimulation.h
--- a/Engine/Plugins/Runtime/NetworkPredictionExtras/Source/NetworkPredictionExtras/Public/MockAbilitySimulation.h
+++ b/Engine/Plugins/Runtime/NetworkPredictionExtras/Source/NetworkPredictionExtras/Public/MockAbilitySimulation.h
@@ -43,7 +43,7 @@
 // MockAbility Data structures
 // -------------------------------------------------------
 
-struct FMockAbilityInputCmd : public FlyingMovement::FInputCmd
+struct FMockAbilityInputCmd : public FFlyingMovementInputCmd
 {
 	bool bSprintPressed = false;
 	bool bDashPressed = false;
@@ -54,12 +54,12 @@ struct FMockAbilityInputCmd : public FlyingMovement::FInputCmd
 		P.Ar << bSprintPressed;
 		P.Ar << bDashPressed;
 		P.Ar << bBlinkPressed;
-		FlyingMovement::FInputCmd::NetSerialize(P);
+		FFlyingMovementInputCmd::NetSerialize(P);
 	}
 
 	void Log(FStandardLoggingParameters& P) const
 	{
-		FlyingMovement::FInputCmd::Log(P);
+		FFlyingMovementInputCmd::Log(P);
 		if (P.Context == EStandardLoggingContext::Full)
 		{
 			P.Ar->Logf(TEXT("bSprintPressed: %d"), bSprintPressed);
@@ -69,19 +69,19 @@ struct FMockAbilityInputCmd : public FlyingMovement::FInputCmd
 	}
 };
 
-struct FMockAbilitySyncState : public FlyingMovement::FMoveState
+struct FMockAbilitySyncState : public FFlyingMovementSyncState
 {
 	float Stamina = 0.f;
 	
 	void NetSerialize(const FNetSerializeParams& P)
 	{
 		P.Ar << Stamina;
-		FlyingMovement::FMoveState::NetSerialize(P);
+		FFlyingMovementSyncState::NetSerialize(P);
 	}
 
 	void Log(FStandardLoggingParameters& P) const
 	{
-		FlyingMovement::FMoveState::Log(P);
+		FFlyingMovementSyncState::Log(P);
 		if (P.Context == EStandardLoggingContext::Full)
 		{
 			P.Ar->Logf(TEXT("Stamina: %.2f"), Stamina);
@@ -89,7 +89,7 @@ struct FMockAbilitySyncState : public FlyingMovement::FMoveState
 	}
 };
 
-struct FMockAbilityAuxstate : public FlyingMovement::FAuxState
+struct FMockAbilityAuxstate : public FFlyingMovementAuxState
 {
 	float MaxStamina = 100.f;
 	float StaminaRegenRate = 20.f;
@@ -104,12 +104,12 @@ struct FMockAbilityAuxstate : public FlyingMovement::FAuxState
 		P.Ar << DashTimeLeft;
 		P.Ar << BlinkWarmupLeft;
 		P.Ar << bIsSprinting;
-		FlyingMovement::FAuxState::NetSerialize(P);
+		FFlyingMovementAuxState::NetSerialize(P);
 	}
 
 	void Log(FStandardLoggingParameters& P) const
 	{
-		FlyingMovement::FAuxState::Log(P);
+		FFlyingMovementAuxState::Log(P);
 		if (P.Context == EStandardLoggingContext::Full)
 		{
 			P.Ar->Logf(TEXT("MaxStamina: %.2f"), MaxStamina);
@@ -178,7 +178,7 @@ struct FMockAbilityCueSet
 
 using TMockAbilityBufferTypes = TNetworkSimBufferTypes<FMockAbilityInputCmd, FMockAbilitySyncState, FMockAbilityAuxstate>;
 
-class FMockAbilitySimulation : public FlyingMovement::FMovementSimulation
+class FMockAbilitySimulation : public FFlyingMovementSimulation
 {
 public:
 	/** Tick group the simulation maps to */
diff --git a/Engine/Plugins/Runtime/NetworkPredictionExtras/Source/NetworkPredictionExtras/Public/NetworkPredictionExtrasFlyingPawn.h b/Engine/Plugins/Runtime/NetworkPredictionExtras/Source/NetworkPredictionExtras/Public/NetworkPredictionExtrasFlyingPawn.h
--- a/Engine/Plugins/Runtime/NetworkPredictionExtras/Source/NetworkPredictionExtras/Public/NetworkPredictionExtrasFlyingPawn.h
+++ b/Engine/Plugins/Runtime/NetworkPredictionExtras/Source/NetworkPredictionExtras/Public/NetworkPredictionExtrasFlyingPawn.h
@@ -82,10 +82,10 @@ public:
 
 protected:
 
-	const FlyingMovement::FAuxState* GetAuxStateRead() const;
-	FlyingMovement::FAuxState* GetAuxStateWrite();
+	const FFlyingMovementAuxState* GetAuxStateRead() const;
+	FFlyingMovementAuxState* GetAuxStateWrite();
 
-	void ProduceInput(const FNetworkSimTime SimTime, FlyingMovement::FInputCmd& Cmd);
+	void ProduceInput(const FNetworkSimTime SimTime, FFlyingMovementInputCmd& Cmd);
 
 	UPROPERTY(Category=Movement, VisibleAnywhere)
 	UFlyingMovementComponent* FlyingMovementComponent;
