NetworkPrediction Plugin:
-Missed file from previous checkin: NetSimCue registration on module load

#rnx

diff --git a/Engine/Plugins/Runtime/NetworkPrediction/Source/NetworkPrediction/Private/NetworkPredictionModule.cpp b/Engine/Plugins/Runtime/NetworkPrediction/Source/NetworkPrediction/Private/NetworkPredictionModule.cpp
--- a/Engine/Plugins/Runtime/NetworkPrediction/Source/NetworkPrediction/Private/NetworkPredictionModule.cpp
+++ b/Engine/Plugins/Runtime/NetworkPrediction/Source/NetworkPrediction/Private/NetworkPredictionModule.cpp
@@ -5,6 +5,8 @@
 #include "Modules/ModuleManager.h"
 #include "Internationalization/Internationalization.h"
 #include "Engine/World.h"
+#include "NetworkedSimulationModelCues.h"
+#include "Misc/CoreDelegates.h"
 
 #define LOCTEXT_NAMESPACE "FNetworkPredictionModule"
 
@@ -17,18 +19,10 @@ class FNetworkPredictionModule : public INetworkPredictionModule
 
 void FNetworkPredictionModule::StartupModule()
 {
-/*
-	FWorldDelegates::OnPostWorldInitialization.AddLambda([](UWorld* World, const UWorld::InitializationValues IVS)
+	FCoreDelegates::OnPostEngineInit.AddLambda([]()
 	{
-		World->AddOnActorSpawnedHandler(FOnActorSpawned::FDelegate::CreateLambda([](AActor* Actor)
-		{
-			if (APawn* Pawn = Cast<APawn>(Actor))
-			{
-
-			}
-		}));
+		FGlobalCueTypeTable::Get().FinalizeTypes();
 	});
-*/
 }
 
 
