Disable false positive PVS warning. Repro case has been submitted to viva64.
#rnx

diff --git a/Engine/Plugins/Runtime/NetworkPrediction/Source/NetworkPrediction/Public/NetworkedSimulationModelCues.h b/Engine/Plugins/Runtime/NetworkPrediction/Source/NetworkPrediction/Public/NetworkedSimulationModelCues.h
--- a/Engine/Plugins/Runtime/NetworkPrediction/Source/NetworkPrediction/Public/NetworkedSimulationModelCues.h
+++ b/Engine/Plugins/Runtime/NetworkPrediction/Source/NetworkPrediction/Public/NetworkedSimulationModelCues.h
@@ -217,7 +217,7 @@ struct FSavedCue
 	template<typename TCue>
 	FSavedCue(FNetworkSimTime InTime, const TCue& SourceCue) : Time(InTime)
 	{
-		ID = TCue::ID;
+		ID = TCue::ID; //-V570
 		CueInstance.Reset(new TNetSimCueWrapper<TCue>(SourceCue)); // Fixme: can the copy be avoided with move semantics?
 	}
 
