NetworkPrediction Plugin
-template refactor: trying to move away from large template parameter lists. There is now a concept of a NetSimModel Definition which holds the user typedefs. The Def is now what is passed around through the various levels within the TNetworkedSimulationModel.
-Moved GroupName, ShouldReconcile, and Interpolate functions to the model def. ShouldReconcile/Interpolate now take both sync/aux state rather than just sync.

#rnx

diff --git a/Engine/Plugins/Runtime/NetworkPrediction/Source/NetworkPrediction/Private/Movement/FlyingMovementComponent.cpp b/Engine/Plugins/Runtime/NetworkPrediction/Source/NetworkPrediction/Private/Movement/FlyingMovementComponent.cpp
--- a/Engine/Plugins/Runtime/NetworkPrediction/Source/NetworkPrediction/Private/Movement/FlyingMovementComponent.cpp
+++ b/Engine/Plugins/Runtime/NetworkPrediction/Source/NetworkPrediction/Private/Movement/FlyingMovementComponent.cpp
@@ -62,9 +62,8 @@ INetworkedSimulationModel* UFlyingMovementComponent::InstantiateNetworkedSimulat
 	InitFlyingMovementSimulation(new FFlyingMovementSimulation(), InitialSyncState, InitialAuxState);
 
 	// The Model
-	auto NewModel = new FFlyingMovementSystem<0>(MovementSimulation, this, InitialSyncState, InitialAuxState);
+	auto* NewModel = new TNetworkedSimulationModel<FFlyingMovementNetSimModelDef>(this->MovementSimulation, this, InitialSyncState, InitialAuxState);
 	InitFlyingMovementNetSimModel(NewModel);
-	
 	return NewModel;
 }
 
@@ -128,7 +127,7 @@ void UFlyingMovementComponent::ProduceInput(const FNetworkSimTime SimTime, FFlyi
 void UFlyingMovementComponent::FinalizeFrame(const FFlyingMovementSyncState& SyncState, const FFlyingMovementAuxState& AuxState)
 {
 	// Does checking equality make any sense here? This is unfortunate
-	if (UpdatedComponent->GetComponentLocation().Equals(SyncState.Location) == false || UpdatedComponent->GetComponentQuat().Rotator().Equals(SyncState.Rotation, ROTATOR_TOLERANCE) == false)
+	if (UpdatedComponent->GetComponentLocation().Equals(SyncState.Location) == false || UpdatedComponent->GetComponentQuat().Rotator().Equals(SyncState.Rotation, FFlyingMovementNetSimModelDef::ROTATOR_TOLERANCE) == false)
 	{
 		FTransform Transform(SyncState.Rotation.Quaternion(), SyncState.Location, UpdatedComponent->GetComponentTransform().GetScale3D() );
 		UpdatedComponent->SetWorldTransform(Transform, false, nullptr, ETeleportType::TeleportPhysics);
diff --git a/Engine/Plugins/Runtime/NetworkPrediction/Source/NetworkPrediction/Private/Movement/FlyingMovementSimulation.cpp b/Engine/Plugins/Runtime/NetworkPrediction/Source/NetworkPrediction/Private/Movement/FlyingMovementSimulation.cpp
--- a/Engine/Plugins/Runtime/NetworkPrediction/Source/NetworkPrediction/Private/Movement/FlyingMovementSimulation.cpp
+++ b/Engine/Plugins/Runtime/NetworkPrediction/Source/NetworkPrediction/Private/Movement/FlyingMovementSimulation.cpp
@@ -43,7 +43,8 @@ static FAutoConsoleVariableRef CVarDrawDebugDefaultLifeTime(TEXT("fp.Debug.UseDr
 // ============================================================================================
 
 
-const FName FFlyingMovementSimulation::GroupName(TEXT("FlyingMovement"));
+const FName FFlyingMovementNetSimModelDef::GroupName(TEXT("FlyingMovement"));
+const float FFlyingMovementNetSimModelDef::ROTATOR_TOLERANCE = (1e-3);
 
 bool FFlyingMovementSimulation::ForceMispredict = false;
 static FVector ForceMispredictVelocityMagnitude = FVector(2000.f, 0.f, 0.f);
@@ -265,7 +266,7 @@ void FFlyingMovementSimulation::SimulationTick(const FNetSimTimeStep& TimeStep,
 void FFlyingMovementSimulation::PreSimSync(const FFlyingMovementSyncState& SyncState)
 {
 	// Does checking equality make any sense here? This is unfortunate
-	if (UpdatedComponent->GetComponentLocation().Equals(SyncState.Location) == false || UpdatedComponent->GetComponentQuat().Rotator().Equals(SyncState.Rotation, ROTATOR_TOLERANCE) == false)
+	if (UpdatedComponent->GetComponentLocation().Equals(SyncState.Location) == false || UpdatedComponent->GetComponentQuat().Rotator().Equals(SyncState.Rotation, FFlyingMovementNetSimModelDef::ROTATOR_TOLERANCE) == false)
 	{
 		FTransform Transform(SyncState.Rotation.Quaternion(), SyncState.Location, UpdatedComponent->GetComponentTransform().GetScale3D() );
 		UpdatedComponent->SetWorldTransform(Transform, false, nullptr, ETeleportType::TeleportPhysics);
diff --git a/Engine/Plugins/Runtime/NetworkPrediction/Source/NetworkPrediction/Private/Movement/ParametricMovement.cpp b/Engine/Plugins/Runtime/NetworkPrediction/Source/NetworkPrediction/Private/Movement/ParametricMovement.cpp
--- a/Engine/Plugins/Runtime/NetworkPrediction/Source/NetworkPrediction/Private/Movement/ParametricMovement.cpp
+++ b/Engine/Plugins/Runtime/NetworkPrediction/Source/NetworkPrediction/Private/Movement/ParametricMovement.cpp
@@ -37,7 +37,7 @@ static FAutoConsoleVariableRef CVarFixStepMS(TEXT("parametricmover.FixStep"),
 // ParametricMovement
 // -------------------------------------------------------------------------------------------------------------------------------
 
-const FName FParametricMovementSimulation::GroupName(TEXT("Parametric"));
+const FName FParametricMovementNetSimModelDef::GroupName(TEXT("Parametric"));
 
 void FParametricMovementSimulation::SimulationTick(const FNetSimTimeStep& TimeStep, const TNetSimInput<ParametricMovementBufferTypes>& Input, const TNetSimOutput<ParametricMovementBufferTypes>& Output)
 {
@@ -97,12 +97,12 @@ INetworkedSimulationModel* UParametricMovementComponent::InstantiateNetworkedSim
 	// Model
 	if (ParametricMoverCVars::FixStep > 0)
 	{
-		auto *NewModel = new FParametricMovementSystem<32>( ParametricMovementSimulation, this);
+		auto* NewModel = new TNetworkedSimulationModel<FParametricMovementNetSimModelDef_Fixed30hz>(ParametricMovementSimulation, this, InitialSyncState, InitialAuxState);
 		InitParametricMovementNetSimModel(NewModel);
 		return NewModel;
 	}
 	
-	auto *NewModel = new FParametricMovementSystem<>(ParametricMovementSimulation, this);
+	auto* NewModel = new TNetworkedSimulationModel<FParametricMovementNetSimModelDef>(ParametricMovementSimulation, this, InitialSyncState, InitialAuxState);
 	InitParametricMovementNetSimModel(NewModel);
 	return NewModel;
 }
@@ -151,7 +151,6 @@ void UParametricMovementComponent::VisualLog(const FParametricInputCmd* Input, c
 	//DrawParams.LogText = FString::Printf(TEXT("[%d] %s. Position: %.4f. Location: %s. Rotation: %s"), Parameters.Frame, *LexToString(Parameters.Context), Position, *Transform.GetLocation().ToString(), *Transform.GetRotation().Rotator().ToString());
 }
 
-
 void UParametricMovementComponent::TickComponent(float DeltaTime, enum ELevelTick TickType, FActorComponentTickFunction* ThisTickFunction)
 {
 	Super::TickComponent(DeltaTime, TickType, ThisTickFunction);
diff --git a/Engine/Plugins/Runtime/NetworkPrediction/Source/NetworkPrediction/Public/Movement/FlyingMovement.h b/Engine/Plugins/Runtime/NetworkPrediction/Source/NetworkPrediction/Public/Movement/FlyingMovement.h
--- a/Engine/Plugins/Runtime/NetworkPrediction/Source/NetworkPrediction/Public/Movement/FlyingMovement.h
+++ b/Engine/Plugins/Runtime/NetworkPrediction/Source/NetworkPrediction/Public/Movement/FlyingMovement.h
@@ -76,13 +76,6 @@ struct FFlyingMovementSyncState
 		P.Ar << Rotation;
 	}
 
-	// Compare this state with AuthorityState. return true if a reconcile (correction) should happen
-	bool ShouldReconcile(const FFlyingMovementSyncState& AuthorityState)
-	{
-		const float ErrorTolerance = 1.f;
-		return !AuthorityState.Location.Equals(Location, ErrorTolerance);
-	}
-
 	void Log(FStandardLoggingParameters& Params) const
 	{
 		if (Params.Context == EStandardLoggingContext::HeaderOnly)
@@ -97,12 +90,6 @@ struct FFlyingMovementSyncState
 			Params.Ar->Logf(TEXT("Rot: %s"), *Rotation.ToString());
 		}
 	}
-
-	static void Interpolate(const FFlyingMovementSyncState& From, const FFlyingMovementSyncState& To, const float PCT, FFlyingMovementSyncState& OutDest)
-	{
-		OutDest.Location = From.Location + ((To.Location - From.Location) * PCT);
-		OutDest.Rotation = From.Rotation + ((To.Rotation - From.Rotation) * PCT);
-	}
 };
 
 // Auxiliary state that is input into the simulation. Doesn't change during the simulation tick.
@@ -140,8 +127,6 @@ struct FFlyingMovementAuxState
 	
 using FlyingMovementBufferTypes = TNetworkSimBufferTypes<FFlyingMovementInputCmd, FFlyingMovementSyncState, FFlyingMovementAuxState>;
 
-static FName SimulationGroupName("FlyingMovement");
-
 class FFlyingMovementSimulation : public FBaseMovementSimulation
 {
 public:
@@ -165,23 +150,57 @@ protected:
 	float SlideAlongSurface(const FVector& Delta, float Time, const FQuat Rotation, const FVector& Normal, FHitResult& Hit, bool bHandleImpact);
 };
 
-// Actual definition of our network simulation.
-template<int32 InFixedStepMS=0>
-using FFlyingMovementSystem = TNetworkedSimulationModel<FFlyingMovementSimulation, FlyingMovementBufferTypes, TNetworkSimTickSettings<InFixedStepMS>>;
+class FFlyingMovementNetSimModelDef : public FNetSimModelDefBase
+{
+public:
+
+	using Simulation = FFlyingMovementSimulation;
+	using BufferTypes = FlyingMovementBufferTypes;
+
+	/** Tick group the simulation maps to */
+	static const FName GroupName;
+
+	// general tolerance value for rotation checks
+	static const float ROTATOR_TOLERANCE;
 
-// general tolerance value for rotation checks
-static const float ROTATOR_TOLERANCE = (1e-3);
+	/** Predicted error testing */
+	static bool ShouldReconcile(const FFlyingMovementSyncState& AuthoritySync, const FFlyingMovementAuxState& AuthorityAux, const FFlyingMovementSyncState& PredictedSync, const FFlyingMovementAuxState& PredictedAux)
+	{
+		const float ErrorTolerance = 1.f;
+		return !AuthoritySync.Location.Equals(PredictedSync.Location, ErrorTolerance);
+	}
 
-// Needed to trick UHT into letting UMockNetworkSimulationComponent implement. UHT cannot parse the ::
-// Also needed for forward declaring. Can't just be a typedef/using =
-class IFlyingMovementSystemDriver : public TNetworkedSimulationModelDriver<FlyingMovementBufferTypes> { };
+	static void Interpolate(const TInterpolatorParameters<FFlyingMovementSyncState, FFlyingMovementAuxState>& Params)
+	{
+		const FVector DeltaLocation = (Params.To.Sync.Location - Params.From.Sync.Location);
+		const float TeleportThresholdSq = 500.f * 500.f; // FIXME: this is not a good way to do this
+		if (DeltaLocation.SizeSquared() > TeleportThresholdSq)
+		{
+			Params.Out.Sync = Params.To.Sync;
+		}
+		else
+		{
+			Params.Out.Sync.Location = Params.From.Sync.Location + (DeltaLocation * Params.InterpolationPCT);
+			Params.Out.Sync.Rotation = Params.From.Sync.Rotation + ((Params.To.Sync.Rotation - Params.From.Sync.Rotation) * Params.InterpolationPCT);
+		}
+	
+		Params.Out.Aux = Params.To.Aux;
+	}
+};
+
+/** Additional specialized types of the Flying Movement NetSimModel */
+class FFlyingMovementNetSimModelDef_Fixed30Hz : public FFlyingMovementNetSimModelDef
+{
+public:
+	using TickSettings = TNetworkSimTickSettings<33>;
+};
 
 // -------------------------------------------------------------------------------------------------------------------------------
 // ActorComponent for running FlyingMovement 
 // -------------------------------------------------------------------------------------------------------------------------------
 
 UCLASS(BlueprintType, meta=(BlueprintSpawnableComponent))
-class NETWORKPREDICTION_API UFlyingMovementComponent : public UBaseMovementComponent, public IFlyingMovementSystemDriver
+class NETWORKPREDICTION_API UFlyingMovementComponent : public UBaseMovementComponent, public TNetworkedSimulationModelDriver<FlyingMovementBufferTypes>
 {
 	GENERATED_BODY()
 
diff --git a/Engine/Plugins/Runtime/NetworkPrediction/Source/NetworkPrediction/Public/Movement/ParametricMovement.h b/Engine/Plugins/Runtime/NetworkPrediction/Source/NetworkPrediction/Public/Movement/ParametricMovement.h
--- a/Engine/Plugins/Runtime/NetworkPrediction/Source/NetworkPrediction/Public/Movement/ParametricMovement.h
+++ b/Engine/Plugins/Runtime/NetworkPrediction/Source/NetworkPrediction/Public/Movement/ParametricMovement.h
@@ -7,8 +7,33 @@
 
 #include "ParametricMovement.generated.h"
 
-struct FSimpleParametricMotion;
-class IMovementDriver;
+// Extremely simple struct for defining parametric motion. This is editable in UParametricMovementComponent's defaults, and also used by the simulation code above. 
+USTRUCT(BlueprintType)
+struct FSimpleParametricMotion
+{
+	GENERATED_BODY()
+
+	// Actually turn the given position into a transform. Again, should be static and not conditional on changing state outside of the network sim
+	void MapTimeToTransform(const float InPosition, FTransform& OutTransform) const;
+
+	// Advance parametric time. This is meant to do simple things like looping/reversing etc.
+	void AdvanceParametricTime(const float InPosition, const float InPlayRate, float &OutPosition, float& OutPlayRate, const float DeltaTimeSeconds) const;
+
+	UPROPERTY(EditAnywhere, BlueprintReadOnly, Category=ParametricMovement)
+	FVector ParametricDelta = FVector(0.f, 0.f, 500.f);
+
+	UPROPERTY(EditAnywhere, BlueprintReadOnly, Category=ParametricMovement)
+	float MinTime = -1.f;
+
+	UPROPERTY(EditAnywhere, BlueprintReadOnly, Category=ParametricMovement)
+	float MaxTime = 1.f;
+
+	FTransform CachedStartingTransform;
+};
+
+// -------------------------------------------------------------------------------------------------------------------------------
+//	Parametric Movement Simulation Types
+// -------------------------------------------------------------------------------------------------------------------------------
 
 // State the client generates
 struct FParametricInputCmd
@@ -49,13 +74,6 @@ struct FParametricSyncState
 		P.Ar << PlayRate;
 	}
 
-	// Compare this state with AuthorityState. return true if a reconcile (correction) should happen
-	bool ShouldReconcile(const FParametricSyncState& AuthorityState)
-	{
-		const float ErrorTolerance = 0.01f;
-		return (FMath::Abs<float>(AuthorityState.Position - Position) > ErrorTolerance) || (FMath::Abs<float>(AuthorityState.PlayRate - PlayRate) > ErrorTolerance);
-	}
-
 	void Log(FStandardLoggingParameters& Params) const
 	{
 		if (Params.Context == EStandardLoggingContext::HeaderOnly)
@@ -69,14 +87,6 @@ struct FParametricSyncState
 			Params.Ar->Logf(TEXT("Rate: %.2f"), PlayRate);
 		}
 	}
-
-	void VisualLog(const FVisualLoggingParameters& Parameters, IMovementDriver* Driver, IMovementDriver* LogDriver) const;
-
-	static void Interpolate(const FParametricSyncState& From, const FParametricSyncState& To, const float PCT, FParametricSyncState& OutDest)
-	{
-		OutDest.Position = From.Position + ((To.Position - From.Position) * PCT);
-		OutDest.PlayRate = From.PlayRate + ((To.PlayRate - From.PlayRate) * PCT);
-	}
 };
 
 // Auxiliary state that is input into the simulation. Doesn't change during the simulation tick.
@@ -102,12 +112,13 @@ struct FParametricAuxState
 	}
 };
 
+/** BufferTypes for ParametricMovement */
 using ParametricMovementBufferTypes = TNetworkSimBufferTypes<FParametricInputCmd, FParametricSyncState, FParametricAuxState>;
-	
+
+/** The actual movement simulation */
 class FParametricMovementSimulation : public FBaseMovementSimulation
 {
 public:
-	static const FName GroupName;
 
 	void SimulationTick(const FNetSimTimeStep& TimeStep, const TNetSimInput<ParametricMovementBufferTypes>& Input, const TNetSimOutput<ParametricMovementBufferTypes>& Output);
 
@@ -115,39 +126,35 @@ public:
 	const FSimpleParametricMotion* Motion = nullptr;
 };
 
-
-// Actual definition of our network simulation.
-template<int32 FixedStepMS=0>
-using FParametricMovementSystem = TNetworkedSimulationModel<FParametricMovementSimulation, ParametricMovementBufferTypes, TNetworkSimTickSettings<FixedStepMS>>;
-
-
-
-// Needed to trick UHT into letting UMockNetworkSimulationComponent implement. UHT cannot parse the ::
-// Also needed for forward declaring. Can't just be a typedef/using =
-class IParametricMovementDriver : public TNetworkedSimulationModelDriver<ParametricMovementBufferTypes> { };
-
-// Extremely simple struct for defining parametric motion. This is editable in UParametricMovementComponent's defaults, and also used by the simulation code above. 
-USTRUCT(BlueprintType)
-struct FSimpleParametricMotion
+/** NetworkedSimulation Model type */
+class FParametricMovementNetSimModelDef : public FNetSimModelDefBase
 {
-	GENERATED_BODY()
-
-	// Actually turn the given position into a transform. Again, should be static and not conditional on changing state outside of the network sim
-	void MapTimeToTransform(const float InPosition, FTransform& OutTransform) const;
+public:
 
-	// Advance parametric time. This is meant to do simple things like looping/reversing etc.
-	void AdvanceParametricTime(const float InPosition, const float InPlayRate, float &OutPosition, float& OutPlayRate, const float DeltaTimeSeconds) const;
+	using Simulation = FParametricMovementSimulation;
+	using BufferTypes = ParametricMovementBufferTypes;
 
-	UPROPERTY(EditAnywhere, BlueprintReadOnly, Category=ParametricMovement)
-	FVector ParametricDelta = FVector(0.f, 0.f, 500.f);
+	static const FName GroupName;
 
-	UPROPERTY(EditAnywhere, BlueprintReadOnly, Category=ParametricMovement)
-	float MinTime = -1.f;
+	// Compare this state with AuthorityState. return true if a reconcile (correction) should happen
+	static bool ShouldReconcile(const FParametricSyncState& AuthoritySync, const FParametricAuxState& AuthorityAux, const FParametricSyncState& PredictedSync, const FParametricAuxState& PredictedAux)
+	{
+		const float ErrorTolerance = 0.01f;
+		return (FMath::Abs<float>(AuthoritySync.Position - PredictedSync.Position) > ErrorTolerance) || (FMath::Abs<float>(AuthoritySync.PlayRate - PredictedSync.PlayRate) > ErrorTolerance);
+	}
 
-	UPROPERTY(EditAnywhere, BlueprintReadOnly, Category=ParametricMovement)
-	float MaxTime = 1.f;
+	static void Interpolate(const TInterpolatorParameters<FParametricSyncState, FParametricAuxState>& Params)
+	{
+		Params.Out.Sync.Position = Params.From.Sync.Position + ((Params.To.Sync.Position - Params.From.Sync.Position) * Params.InterpolationPCT);
+		Params.Out.Sync.PlayRate = Params.From.Sync.PlayRate + ((Params.To.Sync.PlayRate - Params.From.Sync.PlayRate) * Params.InterpolationPCT);
+	}
+};
 
-	FTransform CachedStartingTransform;
+/** Additional specialized types of the Parametric Movement NetSimModel */
+class FParametricMovementNetSimModelDef_Fixed30hz : public FParametricMovementNetSimModelDef
+{ 
+public:
+	using TickSettings = TNetworkSimTickSettings<33>;
 };
 
 // -------------------------------------------------------------------------------------------------------------------------------
@@ -160,7 +167,7 @@ struct FSimpleParametricMotion
 // -------------------------------------------------------------------------------------------------------------------------------
 
 UCLASS(BlueprintType, meta=(BlueprintSpawnableComponent))
-class NETWORKPREDICTION_API UParametricMovementComponent : public UBaseMovementComponent, public IParametricMovementDriver
+class NETWORKPREDICTION_API UParametricMovementComponent : public UBaseMovementComponent, public TNetworkedSimulationModelDriver<ParametricMovementBufferTypes>
 {
 	GENERATED_BODY()
 
diff --git a/Engine/Plugins/Runtime/NetworkPrediction/Source/NetworkPrediction/Public/NetworkedSimulationModel.h b/Engine/Plugins/Runtime/NetworkPrediction/Source/NetworkPrediction/Public/NetworkedSimulationModel.h
--- a/Engine/Plugins/Runtime/NetworkPrediction/Source/NetworkPrediction/Public/NetworkedSimulationModel.h
+++ b/Engine/Plugins/Runtime/NetworkPrediction/Source/NetworkPrediction/Public/NetworkedSimulationModel.h
@@ -6,13 +6,17 @@
 #include "NetworkedSimulationModelTypes.h"
 #include "NetworkedSimulationModelReplicators.h"
 #include "NetworkedSimulationModelDebugger.h"
+#include "NetworkedSimulationModelTraits.h"
 
 // ----------------------------------------------------------------------------------------------------------------------------------------------------------------------
 //	TNetworkedSimulationModel
-//	
-//	* Has all logic for "ticking, advancing buffers, calling Update, calling ServerRPC etc
-//	* Doesn't have anything about update component, movesweep, etc
-//	* Concept of "IDriver" which is the owning object that is driving the network sim. This is the interface to the outside UE4 world.
+//
+//	This is the object that ties all user-defined NetworkPrediction things together:
+//		1. The Simulation: "pure" gameplay simulation: takes input -> produces output. Knows nothing about networking or the general UE4 side of things.
+//		2. The Model Definition: defines user types, tick settings, the simulation class and functions such as ShouldReconcile and Interpolate. E.g, "How do you want the networking model to work".
+//		3. The Driver: required interface for this NetSimModel to function: FinalizeFrame, ProduceInput, Debugging/VisualLogging, etc. Implemented by the Actor/Component/UE4 side.
+//
+//
 //	* Has 4 buffers:
 //		-Input: Generated by a client / not the authority.
 //		-Sync: What we are trying to keep in sync. The state that evolves frame to frame with an Update function.
@@ -25,6 +29,7 @@
 //		-External game code can interact with the system:
 //			-The TNetworkedSimulationModel is mostly public and exposed. It is not recommend to publicly expose to your "user" code (high level scripting, designers, etc).
 //			-TNetkSimStateAccessor is a helper for safely reading/writing to state within the system.
+//		-NetSimCues are how the simulation can invoke events out of the system that are "network aware"
 //
 //	* Notes on ownership and lifetime:
 //		-User code will instantiate the underlying simulation class (TSimulation)
@@ -35,54 +40,66 @@
 //
 // ----------------------------------------------------------------------------------------------------------------------------------------------------------------------
 
-
-struct TNetSimModelTraitsBase
+// Traits for which RepController structs are used internally
+struct TNetSimModelRepControllerTraitsBase
 {
+	template<typename Model>
+	using RepController_Server = TRepController_Server<Model>;
 
-};
-
-template<typename T>
-struct TNetSimModelTraits : public TNetSimModelTraitsBase
-{
-	using TSimulation = void;
-	using TBufferTypes = void;	
-	using TTickSettings = TNetworkSimTickSettings<>;
+	template<typename Model>
+	using RepController_Simulated = TRepController_Simulated<Model>;
 
-};
+	template<typename Model>
+	using RepController_Autonomous = TRepController_Autonomous<Model>;
 
+	template<typename Model>
+	using RepController_Debug = TRepController_Debug<Model>;
 
-template <
-	typename InTSimulation,								// Final Simulation class
-	typename TUserBufferTypes,							// The user types (input, sync, aux, debug). Note this gets wrapped in TInternalBufferTypes internally.
-	typename InTTickSettings=TNetworkSimTickSettings<>, // Defines global rules about time keeping and ticking
+	template<typename Model>
+	using RepController_Replay = TRepController_Empty<Model>;
+};
 
-	// Core proxies that dictate how data replicates and how the simulation evolves for the three main roles
-	typename TRepProxyServerRPC =	TReplicator_Server		<TInternalBufferTypes<TUserBufferTypes, InTTickSettings>,	InTTickSettings >,
-	typename TRepProxyAutonomous =	TReplicator_Autonomous	<TInternalBufferTypes<TUserBufferTypes, InTTickSettings>,	InTTickSettings>,
-	typename TRepProxySimulated =	TReplicator_Simulated	<TInternalBufferTypes<TUserBufferTypes, InTTickSettings>,	InTTickSettings>,
+// Template struct to specialize per-model
+template<typename T> struct TNetSimModelRepControllerTraits : public TNetSimModelRepControllerTraitsBase { };
 
-	// Defines how replication happens on these special channels, but doesn't dictate how simulation evolves
-	typename TRepProxyReplay =		TReplicator_Sequence	<TInternalBufferTypes<TUserBufferTypes, InTTickSettings>,	InTTickSettings, ENetworkSimBufferTypeId::Sync,  3>,
-	typename TRepProxyDebug =		TReplicator_Debug		<TInternalBufferTypes<TUserBufferTypes, InTTickSettings>,	InTTickSettings>
->
-class TNetworkedSimulationModel : public INetworkedSimulationModel
+// The actual NetworkedSimulationModel templated class.
+template <typename InModelDef>
+class TNetworkedSimulationModel : public InModelDef::Base
 {
 public:
 
-	using TSimulation = InTSimulation;
-	using TDriver = TNetworkedSimulationModelDriver<TUserBufferTypes>;
-
-	using TBufferTypes = TInternalBufferTypes<TUserBufferTypes, InTTickSettings>;
-	using TTickSettings = InTTickSettings;
+	// User Model
+	using Model = InModelDef;
+	
+	// "Core" model pieces
+	using TSimulation = typename Model::Simulation;
+	using TTickSettings = typename Model::TickSettings;
+	using TDriver = typename TNetSimModelTraits<Model>::Driver;
+	using TBufferTypes = typename TNetSimModelTraits<Model>::InternalBufferTypes;
 
+	// Buffer types
 	using TInputCmd = typename TBufferTypes::TInputCmd;
 	using TSyncState = typename TBufferTypes::TSyncState;
 	using TAuxState = typename TBufferTypes::TAuxState;
 	using TDebugState = typename TBufferTypes::TDebugState;
 
+	// Time
 	using TSimTime = FNetworkSimTime;
 	using TRealTime = FNetworkSimTime::FRealTime;
 	
+	// Core proxies that dictate how data replicates and how the simulation evolves for the three main roles
+	using TRepProxyServerRPC =	typename TNetSimModelRepControllerTraits<Model>::template RepController_Server<Model>;
+	using TRepProxySimulated =	typename TNetSimModelRepControllerTraits<Model>::template RepController_Simulated<Model>;
+	using TRepProxyAutonomous =	typename TNetSimModelRepControllerTraits<Model>::template RepController_Autonomous<Model>;
+
+	// Defines how replication happens on these special channels, but doesn't dictate how simulation evolves
+	using TRepProxyDebug =	typename TNetSimModelRepControllerTraits<Model>::template RepController_Debug<Model>;
+	using TRepProxyReplay =	typename TNetSimModelRepControllerTraits<Model>::template RepController_Replay<Model>;
+
+	// ----------------------------------------------------------------------------------------------------------
+	//
+	// ----------------------------------------------------------------------------------------------------------
+
 	// Note: Ownership of TSimulation* is taken by this constructor
 	template<typename TInDriver>
 	TNetworkedSimulationModel(TSimulation* InSimulation, TInDriver* InDriver, const TSyncState& InitialSyncState = TSyncState(), const TAuxState& InitialAuxState = TAuxState())
@@ -95,7 +112,7 @@ public:
 		Ticker.SetTotalProcessedSimulationTime(FNetworkSimTime(), 0);
 		DO_NETSIM_MODEL_DEBUG(FNetworkSimulationModelDebuggerManager::Get().RegisterNetworkSimulationModel(this, Driver->GetVLogOwner()));
 
-		ProcessPendingNetSimCuesFunc = [this, InDriver]()
+		this->ProcessPendingNetSimCuesFunc = [this, InDriver]()
 		{ 
 			Buffers.CueDispatcher.template DispatchCueRecord<TInDriver>(*InDriver, Ticker.GetTotalProcessedSimulationTime()); 
 		};
@@ -153,15 +170,15 @@ public:
 		switch (Parameters.Role)
 		{
 			case ROLE_Authority:
-				RepProxy_ServerRPC.template PreSimTick<TDriver>(Driver, Buffers, Ticker, Parameters);
+				RepProxy_ServerRPC.PreSimTick(Driver, Buffers, Ticker, Parameters);
 			break;
 
 			case ROLE_AutonomousProxy:
-				RepProxy_Autonomous.template PreSimTick<TDriver>(Driver, Buffers, Ticker, Parameters);
+				RepProxy_Autonomous.PreSimTick(Driver, Buffers, Ticker, Parameters);
 			break;
 
 			case ROLE_SimulatedProxy:
-				RepProxy_Simulated.template PreSimTick<TDriver>(Driver, Buffers, Ticker, Parameters);
+				RepProxy_Simulated.PreSimTick(Driver, Buffers, Ticker, Parameters);
 			break;
 		}
 
@@ -226,15 +243,15 @@ public:
 		switch (Parameters.Role)
 		{
 			case ROLE_Authority:
-				RepProxy_ServerRPC.template PostSimTick<TDriver>(Driver, Buffers, Ticker, Parameters);
+				RepProxy_ServerRPC.PostSimTick(Driver, Buffers, Ticker, Parameters);
 			break;
 
 			case ROLE_AutonomousProxy:
-				RepProxy_Autonomous.template PostSimTick<TDriver>(Driver, Buffers, Ticker, Parameters);
+				RepProxy_Autonomous.PostSimTick(Driver, Buffers, Ticker, Parameters);
 			break;
 
 			case ROLE_SimulatedProxy:
-				RepProxy_Simulated.template PostSimTick<TDriver>(Driver, Buffers, Ticker, Parameters);
+				RepProxy_Simulated.PostSimTick(Driver, Buffers, Ticker, Parameters);
 			break;
 		}
 
@@ -269,15 +286,15 @@ public:
 		switch (Role)
 		{
 			case ROLE_Authority:
-				RepProxy_ServerRPC.template Reconcile<TSimulation, TDriver>(Simulation.Get(), Driver, Buffers, Ticker);
+				RepProxy_ServerRPC.Reconcile(Simulation.Get(), Driver, Buffers, Ticker);
 			break;
 
 			case ROLE_AutonomousProxy:
-				RepProxy_Autonomous.template Reconcile<TSimulation, TDriver>(Simulation.Get(), Driver, Buffers, Ticker);
+				RepProxy_Autonomous.Reconcile(Simulation.Get(), Driver, Buffers, Ticker);
 			break;
 
 			case ROLE_SimulatedProxy:
-				RepProxy_Simulated.template Reconcile<TSimulation, TDriver>(Simulation.Get(), Driver, Buffers, Ticker);
+				RepProxy_Simulated.Reconcile(Simulation.Get(), Driver, Buffers, Ticker);
 			break;
 		}
 	}
@@ -440,12 +457,12 @@ public:
 
 	void BeginRollback(const FNetworkSimTime& RollbackDeltaTime, const int32 ParentFrame) final override
 	{
-		RepProxy_Simulated.template DependentRollbackBegin<TSimulation, TDriver>(Simulation.Get(), Driver, Buffers, Ticker, RollbackDeltaTime, ParentFrame);
+		RepProxy_Simulated.DependentRollbackBegin(Simulation.Get(), Driver, Buffers, Ticker, RollbackDeltaTime, ParentFrame);
 	}
 
 	void StepRollback(const FNetworkSimTime& Step, const int32 ParentFrame, const bool bFinalStep) final override
 	{
-		RepProxy_Simulated.template DependentRollbackStep<TSimulation, TDriver>(Simulation.Get(), Driver, Buffers, Ticker, Step, ParentFrame, bFinalStep);
+		RepProxy_Simulated.DependentRollbackStep(Simulation.Get(), Driver, Buffers, Ticker, Step, ParentFrame, bFinalStep);
 	}
 
 	void ClearAllDependentSimulations()
@@ -499,7 +516,7 @@ public:
 	}
 
 	// Simulation class should have a static const FName GroupName member
-	FName GetSimulationGroupName() const final override { return TSimulation::GroupName; }
+	FName GetSimulationGroupName() const final override { return Model::GroupName; }
 
 private:
 	float ServerRPCAccumulatedTimeSeconds = 0.f;
diff --git a/Engine/Plugins/Runtime/NetworkPrediction/Source/NetworkPrediction/Public/NetworkedSimulationModelInterpolator.h b/Engine/Plugins/Runtime/NetworkPrediction/Source/NetworkPrediction/Public/NetworkedSimulationModelInterpolator.h
--- a/Engine/Plugins/Runtime/NetworkPrediction/Source/NetworkPrediction/Public/NetworkedSimulationModelInterpolator.h
+++ b/Engine/Plugins/Runtime/NetworkPrediction/Source/NetworkPrediction/Public/NetworkedSimulationModelInterpolator.h
@@ -9,6 +9,7 @@
 #include "Engine/World.h"
 #include "NetworkedSimulationModelTypes.h"
 #include "NetworkedSimulationModelTick.h"
+#include "NetworkedSimulationModelTraits.h"
 
 DEFINE_LOG_CATEGORY_STATIC(LogNetInterpolation, Log, All);
 
@@ -24,14 +25,45 @@ namespace NetworkInterpolationDebugCVars
 	NETSIM_DEVCVAR_SHIPCONST_FLOAT(CatchUpFactor, 1.50, "ni.CatchUpFactor", "Factor we use to catch up");
 }
 
-template<typename TBufferTypes, typename TTickSettings>
-struct TInterpolator
+template<typename TSync, typename TAux>
+struct TInterpolatorParameters
 {
-	using FSimTime = FNetworkSimTime;
-	using FRealTime = FNetworkSimTime::FRealTime;
+	template<typename S, typename A>
+	struct TStatePair
+	{
+		S& Sync;
+		A& Aux;
+	};
+
+	using TInputPair = TStatePair<const TSync, const TAux>;
+	using TOutputPair = TStatePair<TSync, TAux>;
+
+	TInputPair From;
+	TInputPair To;
+	float InterpolationPCT;
+	TOutputPair Out;
+
+	template<typename TS, typename TA>
+	TInterpolatorParameters<TS, TA> Cast() const { return { {From.Sync, From.Aux}, {To.Sync, To.Aux}, InterpolationPCT, {Out.Sync, Out.Aux} }; }
+};
+
+template<typename Model>
+struct TNetSimInterpolator
+{
+	using TBufferTypes = typename TNetSimModelTraits<Model>::InternalBufferTypes;
+	using TTickSettings = typename Model::TickSettings;
 	using TSyncState = typename TBufferTypes::TSyncState;
 	using TAuxState = typename TBufferTypes::TAuxState;
 
+	using FSimTime = FNetworkSimTime;
+	using FRealTime = FNetworkSimTime::FRealTime;
+
+	struct FStatePair
+	{
+		TSyncState Sync;
+		TAuxState Aux;
+	};
+
 	bool bEnableVisualLog = true;
 
 	template<typename TSystemDriver>
@@ -68,7 +100,10 @@ struct TInterpolator
 		{
 			InterpolationTime = SimulationTimeBuffer.TailElement()->ToRealTimeSeconds();
 			InterpolationFrame = SimulationTimeBuffer.TailFrame();
-			InterpolationState = *Buffers.Sync.TailElement();
+
+			auto& FromState = GetFromInterpolationState();
+			FromState.Sync = *Buffers.Sync[InterpolationFrame];
+			FromState.Aux = *Buffers.Aux[InterpolationFrame];
 		}
 
 		// Wait if we were too far ahead
@@ -138,7 +173,7 @@ struct TInterpolator
 
 		// Find "To" frame
 		const TSyncState* ToState = nullptr;
-		const TAuxState* AuxState = nullptr;
+		const TAuxState* ToAuxState = nullptr;
 		FRealTime ToTime = 0.f;
 
 		for (auto It = SimulationTimeBuffer.CreateConstIterator(); It; ++It)
@@ -149,12 +184,12 @@ struct TInterpolator
 				InterpolationFrame = It.Frame();
 				ToTime = ElementSimTime.ToRealTimeSeconds();
 				ToState = Buffers.Sync[It.Frame()];
-				AuxState = Buffers.Aux[It.Frame()];
+				ToAuxState = Buffers.Aux[It.Frame()];
 				break;
 			}
 		}
 
-		if (ensure(ToState && AuxState))
+		if (ensure(ToState && ToAuxState))
 		{
 			const FRealTime FromRealTime = InterpolationTime;
 			const FRealTime ToRealTime = ToTime;
@@ -165,10 +200,11 @@ struct TInterpolator
 				const float InterpolationPCT = (NewInterpolationTime - FromRealTime) / InterpolationInterval;
 				ensureMsgf(InterpolationPCT >= 0.f && InterpolationPCT <= 1.f, TEXT("Calculated InterpolationPCT not in expected range. NewInterpolationTime: %s. From: %s. To: %s"), *LexToString(NewInterpolationTime), *LexToString(FromRealTime), *LexToString(ToRealTime));
 
-				TSyncState NewInterpolatedState;
-				TSyncState::Interpolate(InterpolationState, *ToState, InterpolationPCT, NewInterpolatedState);
+				auto& FromState = GetFromInterpolationState();
+				auto& OutputState = GetNextInterpolationState();
 
-				Driver->FinalizeFrame(NewInterpolatedState, *AuxState);
+				Model::Interpolate({ { FromState.Sync, FromState.Aux }, { *ToState, *ToAuxState }, InterpolationPCT, { OutputState.Sync, OutputState.Aux } });
+				Driver->FinalizeFrame(OutputState.Sync, OutputState.Aux);
 				
 				if (bDoVLog)
 				{
@@ -209,41 +245,42 @@ struct TInterpolator
 
 					{
 						FVisualLoggingParameters VLogParams(EVisualLoggingContext::InterpolationFrom, InterpolationFrame-1, EVisualLoggingLifetime::Transient);
-						Driver->VisualLog(Buffers.Input[InterpolationFrame-1], &InterpolationState, Buffers.Aux[InterpolationFrame-1], VLogParams);
+						Driver->VisualLog(Buffers.Input[InterpolationFrame-1], &FromState.Sync, &FromState.Aux, VLogParams);
 					}
 
 					{
 						FVisualLoggingParameters VLogParams(EVisualLoggingContext::InterpolationTo, InterpolationFrame, EVisualLoggingLifetime::Transient);
-						Driver->VisualLog(Buffers.Input[InterpolationFrame], &InterpolationState, Buffers.Aux[InterpolationFrame], VLogParams);
+						Driver->VisualLog(Buffers.Input[InterpolationFrame], ToState, ToAuxState, VLogParams);
 					}
 
 					{
 						FVisualLoggingParameters VLogParams(LoggingContext, InterpolationFrame, EVisualLoggingLifetime::Transient);
-						Driver->VisualLog(Buffers.Input[InterpolationFrame], &NewInterpolatedState, Buffers.Aux[InterpolationFrame], VLogParams);
+						Driver->VisualLog(Buffers.Input[InterpolationFrame], &OutputState.Sync, &OutputState.Aux, VLogParams);
 					}
 				}
-
-				InterpolationState = NewInterpolatedState;
+				
 				InterpolationTime = NewInterpolationTime;
+				InternalIdx ^= 1;
 			}
 		}
 	}
 
 private:
 
+	FStatePair& GetFromInterpolationState() { return InterpolationState[InternalIdx]; }
+	FStatePair& GetNextInterpolationState() { return InterpolationState[InternalIdx ^ 1]; }
+
 	FRealTime InterpolationTime = 0.f; // SimTime we are currently interpolating at
 	int32 InterpolationFrame = INDEX_NONE; // Frame we are currently/last interpolated at
-
-	TSyncState InterpolationState;
+	FStatePair InterpolationState[2]; // Interpolating "from" state and "out" state
+	int32 InternalIdx = 0; // index into InterpolationState for double buffering pattern
 
 	FRealTime WaitUntilTime = 0.f;	
-
 	FRealTime CatchUpUntilTime = 0.f;
 
-
 	FRealTime DynamicBufferedTime = 1/60.f; // SimTime we are currently interpolating at
 	FRealTime DynamicBufferedTimeStep = 1/60.f;
 
-	FRealTime	MinBufferedTime = 1/120.f;
-	FRealTime	MaxBufferedTime = 1.f;
+	FRealTime MinBufferedTime = 1/120.f;
+	FRealTime MaxBufferedTime = 1.f;
 };
\ No newline at end of file
diff --git a/Engine/Plugins/Runtime/NetworkPrediction/Source/NetworkPrediction/Public/NetworkedSimulationModelReplicators.h b/Engine/Plugins/Runtime/NetworkPrediction/Source/NetworkPrediction/Public/NetworkedSimulationModelReplicators.h
--- a/Engine/Plugins/Runtime/NetworkPrediction/Source/NetworkPrediction/Public/NetworkedSimulationModelReplicators.h
+++ b/Engine/Plugins/Runtime/NetworkPrediction/Source/NetworkPrediction/Public/NetworkedSimulationModelReplicators.h
@@ -51,19 +51,67 @@ struct FNetworkSimulationSerialization
 	}
 };
 
+// -------------------------------------------------------------------------------------------------------
+//	Helpers
+// -------------------------------------------------------------------------------------------------------
+
+/** Helper that writes a new input cmd to the input buffer, at given Frame (usually the sim's PendingFrame). If frame doesn't exist, ProduceInput is called on the driver if bProduceInputViaDriver=true, otherwise the input cmd will be initialized from the previous input cmd. */
+template<typename TDriver, typename TBufferTypes>
+void GenerateLocalInputCmdAtFrame(TDriver* Driver, TNetworkSimBufferContainer<TBufferTypes>& Buffers, const FNetworkSimTime& DeltaSimTime, int32 Frame, bool bProduceInputViaDriver=true)
+{
+	using TInputCmd = typename TBufferTypes::TInputCmd;
+	if (TInputCmd* InputCmd = Buffers.Input[Frame])
+	{
+		InputCmd->SetFrameDeltaTime(DeltaSimTime);
+	}
+	else
+	{
+		InputCmd = Buffers.Input.WriteFrameInitializedFromHead(Frame);
+		InputCmd->SetFrameDeltaTime(DeltaSimTime);
+		if (bProduceInputViaDriver)
+		{
+			Driver->ProduceInput(DeltaSimTime, *InputCmd);
+		}
+	}
+}
+
+/** Helper to generate a local input cmd if we have simulation time to spend and advance the simulation's MaxAllowedFrame so that it can be processed. */
+template<typename TDriver, typename TBufferTypes, typename TTickSettings>
+void TryGenerateLocalInput(TDriver* Driver, TNetworkSimBufferContainer<TBufferTypes>& Buffers, TSimulationTicker<TTickSettings>& Ticker, bool bProduceInputViaDriver=true)
+{
+	using TInputCmd = typename TBufferTypes::TInputCmd;
+
+	FNetworkSimTime DeltaSimTime = Ticker.GetRemainingAllowedSimulationTime();
+	if (DeltaSimTime.IsPositive())
+	{
+		GenerateLocalInputCmdAtFrame(Driver, Buffers, DeltaSimTime, Ticker.PendingFrame, bProduceInputViaDriver);
+		Ticker.MaxAllowedFrame = Ticker.PendingFrame;
+	}
+}
+
+template<typename T>
+struct TNetSimModelTraits;
+
 // ----------------------------------------------------------------------------------------------------------------------------------------------------------------------
-//	Replicators
-//	The Replicators are the pieces of the TNetworkedSimulationModel that make up the role-specific functionality (Server, Autonomous Client, and Simulated Client).
+//	Replication Controllers
+//	
+//	The RepControllers are the pieces of the TNetworkedSimulationModel that make up the role-specific functionality (Server, Autonomous Client, and Simulated Client).
 //	Mainly they NetSerialize, Reconcile, and have Pre/Post sim tick functions. The TNetworkedSimulationModel is still the core piece that ticks the sim, but the Replicators
-//	basically do everything else in a role-specific way.
-//
-//
-//
+//	do everything else in a role-specific way.
 // ----------------------------------------------------------------------------------------------------------------------------------------------------------------------
 
-template<typename TBufferTypes, typename TTickSettings>
-struct TReplicatorEmpty
+// Empty RepController: zero functionality, just stubbed out functions and using definitions for convenience
+template<typename Model>
+struct TRepController_Empty
 { 
+	using TSimulation = typename Model::Simulation;
+	using TTickSettings = typename Model::TickSettings;
+	using TDriver = typename TNetSimModelTraits<Model>::Driver;
+	using TBufferTypes = typename TNetSimModelTraits<Model>::InternalBufferTypes;
+	using TInputCmd = typename TBufferTypes::TInputCmd;
+	using TSyncState = typename TBufferTypes::TSyncState;
+	using TAuxState = typename TBufferTypes::TAuxState;
+
 	// Used for client shadowstate compares. Should just call GetDirtyCount() on the buffer you are replicating
 	int32 GetProxyDirtyCount(TNetworkSimBufferContainer<TBufferTypes>& Buffers) const { return 0; }
 
@@ -71,36 +119,26 @@ struct TReplicatorEmpty
 	void NetSerialize(const FNetSerializeParams& P, TNetworkSimBufferContainer<TBufferTypes>& Buffers, const TSimulationTicker<TTickSettings>& Ticker) { }
 
 	// Reconcile: called after everyone has NetSerialized. "Get right with the server": this function is about reconciling what the server told you vs what you have extrapolated or forward predicted locally
-	template<typename TSimulation, typename TDriver>
-	void Reconcile(TSimulation* Simulation, TDriver* Driver, TNetworkSimBufferContainer<TBufferTypes>& Buffers, TSimulationTicker<TTickSettings>& Ticker) { }
+	void Reconcile(TSimulation* Sim, TDriver* Driver, TNetworkSimBufferContainer<TBufferTypes>& Buffers, TSimulationTicker<TTickSettings>& Ticker) { }
 
 	// Called prior to input processing. This function must update Ticker to allow simulation time (from TickParameters) and to possibly get new input.
-	template<typename TDriver>
 	void PreSimTick(TDriver* Driver, TNetworkSimBufferContainer<TBufferTypes>& Buffers, TSimulationTicker<TTickSettings>& Ticker, const FNetSimTickParameters& TickParameters) { }
 
 	// Called after input processing. Should finalize the frame and do any smoothing/interpolation. This function is not allowed to modify the buffers or tick state, or even call the simulation/Update function.
-	template<typename TDriver>
 	void PostSimTick(TDriver* Driver, const TNetworkSimBufferContainer<TBufferTypes>& Buffers, const TSimulationTicker<TTickSettings>& Ticker, const FNetSimTickParameters& TickParameters) { }
 };
 
-// This is the "templated base class" for replicators but is not required (i.e., this not an official interface used by TNetworkedSimulation model. Just a base implementation you can start with)
-template<typename TBufferTypes, typename TTickSettings>
-struct TReplicatorBase
+// Base implementation with minimal functionality: will grant simulation time and call finalize frame
+template<typename Model, typename TBase=TRepController_Empty<Model>>
+struct TRepController_Base : public TBase
 {
-	using TInputCmd = typename TBufferTypes::TInputCmd;
-
-	// Used for client shadowstate compares. Should just call GetDirtyCount() on the buffer you are replicating
-	int32 GetProxyDirtyCount(TNetworkSimBufferContainer<TBufferTypes>& Buffers) const { return 0; }
-
-	// NetSerialize: just serialize the network data. Don't run simulation steps. Every replicator will be NetSerialized before moving on to Reconcile phase.
-	void NetSerialize(const FNetSerializeParams& P, TNetworkSimBufferContainer<TBufferTypes>& Buffers, const TSimulationTicker<TTickSettings>& Ticker) { }
-
-	// Reconcile: called after everyone has NetSerialized. "Get right with the server": this function is about reconciling what the server told you vs what you have extrapolated or forward predicted locally
-	template<typename TSimulation, typename TDriver>
-	void Reconcile(TSimulation* Simulation, TDriver* Driver, TNetworkSimBufferContainer<TBufferTypes>& Buffers, TSimulationTicker<TTickSettings>& Ticker) { }
+	using typename TBase::TSimulation;
+	using typename TBase::TTickSettings;
+	using typename TBase::TDriver;
+	using typename TBase::TBufferTypes;
+	using typename TBase::TInputCmd;
 
 	// Called prior to input processing. This function must updated Ticker to allow simulation time (from TickParameters) and to possibly get new input.
-	template<typename TDriver>
 	void PreSimTick(TDriver* Driver, TNetworkSimBufferContainer<TBufferTypes>& Buffers, TSimulationTicker<TTickSettings>& Ticker, const FNetSimTickParameters& TickParameters)
 	{
 		// Accumulate local delta time into Ticker
@@ -108,7 +146,6 @@ struct TReplicatorBase
 	}
 
 	// Called after input processing. Should finalize the frame and do any smoothing/interpolation. This function is not allowed to modify the buffers or tick state, or even call the simulation/Update function.
-	template<typename TDriver>
 	void PostSimTick(TDriver* Driver, const TNetworkSimBufferContainer<TBufferTypes>& Buffers, const TSimulationTicker<TTickSettings>& Ticker, const FNetSimTickParameters& TickParameters)
 	{
 		// Sync to latest frame if there is any
@@ -120,13 +157,18 @@ struct TReplicatorBase
 };
 
 // -------------------------------------------------------------------------------------------------------
-//	"Reusable" pieces: mainly for replicating data on specified buffers
+//	Simulation Time replicator
+//	Just serializes simulation time. Can be disabled by Enabled templated parameter. Used as base class in other classes below.
 // -------------------------------------------------------------------------------------------------------
-
-// Just serializes simulation time. Can be disabled by Enabled templated parameter. Used as base class in other classes below.
-template<typename TBufferTypes, typename TTickSettings, typename TBase=TReplicatorBase<TBufferTypes, TTickSettings>, bool Enabled=true>
-struct TReplicator_SimTime : public TBase
+template<typename Model, typename TBase=TRepController_Base<Model>, bool Enabled=true>
+struct TRepController_SimTime : public TBase
 {
+	using typename TBase::TSimulation;
+	using typename TBase::TTickSettings;
+	using typename TBase::TDriver;
+	using typename TBase::TBufferTypes;
+	using typename TBase::TInputCmd;
+
 	void NetSerialize(const FNetSerializeParams& P, TNetworkSimBufferContainer<TBufferTypes>& Buffers, TSimulationTicker<TTickSettings>& Ticker)
 	{
 		TBase::NetSerialize(P, Buffers, Ticker);
@@ -138,19 +180,26 @@ struct TReplicator_SimTime : public TBase
 };
 
 // Enabled=false specialization: do nothing
-template<typename TBufferTypes, typename TTickSettings, typename TBase>
-struct TReplicator_SimTime<TBufferTypes, TTickSettings, TBase, false> : public TBase { };
+template<typename Model, typename TBase>
+struct TRepController_SimTime<Model, TBase, false> : public TBase { };
 
 // -------------------------------------------------------------------------------------------------------
-//	
+//	Sequence replicator
+//
+//	Replicates a sequence of elements. i.e., "the last MaxNumElements".
+//	On the receiving side, we merge the sequence into whatever we have locally
+//	Frames are synchronized. SimTime is also serialized by default (change by changing TBase)
 // -------------------------------------------------------------------------------------------------------
 
-// Replicates a sequence of elements. i.e., "the last MaxNumElements".
-// On the receiving side, we merge the sequence into whatever we have locally
-// Frames are synchronized. SimTime is also serialized by default (change by changing TBase)
-template<typename TBufferTypes, typename TTickSettings, ENetworkSimBufferTypeId BufferId, int32 MaxNumElements=3, typename TBase=TReplicator_SimTime<TBufferTypes, TTickSettings>>
-struct TReplicator_Sequence : public TBase
+template<typename Model, ENetworkSimBufferTypeId BufferId, int32 MaxNumElements=3, typename TBase=TRepController_SimTime<Model>>
+struct TRepController_Sequence : public TBase
 {
+	using typename TBase::TSimulation;
+	using typename TBase::TTickSettings;
+	using typename TBase::TDriver;
+	using typename TBase::TBufferTypes;
+	using typename TBase::TInputCmd;
+
 	int32 GetProxyDirtyCount(TNetworkSimBufferContainer<TBufferTypes>& Buffers) const 
 	{
 		return Buffers.template Get<BufferId>().GetDirtyCount() ^ (TBase::GetProxyDirtyCount(Buffers) << 2); 
@@ -196,87 +245,19 @@ protected:
 	int32 LastSerializedFrame = 0;
 };
 
-// Replicates only the latest single element in the selected buffer.
-// Frame is not synchronized: the new element is just added to head.
-// SimTime is serialized by default (change by changing TBase)
-template<typename TBufferTypes, typename TTickSettings, ENetworkSimBufferTypeId BufferId, typename TBase=TReplicator_SimTime<TBufferTypes, TTickSettings>>
-struct TReplicator_Single : public TBase
-{
-	using TState = typename TBufferTypes::template select_type<BufferId>::type;
-
-	int32 GetProxyDirtyCount(TNetworkSimBufferContainer<TBufferTypes>& Buffers) const
-	{ 
-		return Buffers.template Get<BufferId>().GetDirtyCount() ^ (TBase::GetProxyDirtyCount(Buffers) << 2); ; 
-	}
-
-	void NetSerialize(const FNetSerializeParams& P, TNetworkSimBufferContainer<TBufferTypes>& Buffers, TSimulationTicker<TTickSettings>& Ticker)
-	{
-		TBase::NetSerialize(P, Buffers, Ticker);
-
-		FArchive& Ar = P.Ar;
-
-		auto& Buffer = Buffers.template Get<BufferId>();
-		TState* State = nullptr;
-		
-		if (Ar.IsSaving())
-		{
-			State = Buffer.GetElementFromHead(0);
-			check(State); // We should not be here if the buffer is empty. Want to avoid serializing an "empty" flag at the top here.
-		}
-		else
-		{
-			State = Buffer.GetWriteNext();
-		}
-
-		State->NetSerialize(Ar);
-	}
-};
-
-/** Helper that writes a new input cmd to the input buffer, at given Frame (usually the sim's PendingFrame). If frame doesn't exist, ProduceInput is called on the driver if bProduceInputViaDriver=true, otherwise the input cmd will be initialized from the previous input cmd. */
-template<typename TDriver, typename TBufferTypes>
-void GenerateLocalInputCmdAtFrame(TDriver* Driver, TNetworkSimBufferContainer<TBufferTypes>& Buffers, const FNetworkSimTime& DeltaSimTime, int32 Frame, bool bProduceInputViaDriver=true)
-{
-	using TInputCmd = typename TBufferTypes::TInputCmd;
-	if (TInputCmd* InputCmd = Buffers.Input[Frame])
-	{
-		InputCmd->SetFrameDeltaTime(DeltaSimTime);
-	}
-	else
-	{
-		InputCmd = Buffers.Input.WriteFrameInitializedFromHead(Frame);
-		InputCmd->SetFrameDeltaTime(DeltaSimTime);
-		if (bProduceInputViaDriver)
-		{
-			Driver->ProduceInput(DeltaSimTime, *InputCmd);
-		}
-	}
-}
-
-/** Helper to generate a local input cmd if we have simulation time to spend and advance the simulation's MaxAllowedFrame so that it can be processed. */
-template<typename TDriver, typename TBufferTypes, typename TTickSettings>
-void TryGenerateLocalInput(TDriver* Driver, TNetworkSimBufferContainer<TBufferTypes>& Buffers, TSimulationTicker<TTickSettings>& Ticker, bool bProduceInputViaDriver=true)
-{
-	using TInputCmd = typename TBufferTypes::TInputCmd;
-
-	FNetworkSimTime DeltaSimTime = Ticker.GetRemainingAllowedSimulationTime();
-	if (DeltaSimTime.IsPositive())
-	{
-		GenerateLocalInputCmdAtFrame(Driver, Buffers, DeltaSimTime, Ticker.PendingFrame, bProduceInputViaDriver);
-		Ticker.MaxAllowedFrame = Ticker.PendingFrame;
-	}
-}
-
 // -------------------------------------------------------------------------------------------------------
-//	Role based Replicators: these replicators are meant to serve specific roles wrt how the simulation evolves
+//	Server Replication Controller (FIXME: should this be called something like "client autonomous proxy -> Server RPC" ?
+//
 // -------------------------------------------------------------------------------------------------------
-
-// Default Replicator for the Server: Replicates the InputBuffer client->server
-template<typename TBufferTypes, typename TTickSettings, typename TBase=TReplicator_Sequence<TBufferTypes, TTickSettings, ENetworkSimBufferTypeId::Input, 3>>
-struct TReplicator_Server : public TBase
+template<typename Model, typename TBase=TRepController_Sequence<Model, ENetworkSimBufferTypeId::Input, 3>>
+struct TRepController_Server : public TBase
 {
-	using TInputCmd = typename TBufferTypes::TInputCmd;
-
-	template<typename TSimulation, typename TDriver>
+	using typename TBase::TSimulation;
+	using typename TBase::TTickSettings;
+	using typename TBase::TDriver;
+	using typename TBase::TBufferTypes;
+	using typename TBase::TInputCmd;
+	
 	void Reconcile(TSimulation* Simulation, TDriver* Driver, TNetworkSimBufferContainer<TBufferTypes>& Buffers, TSimulationTicker<TTickSettings>& Ticker)
 	{
 		// After receiving input, server may process up to the latest received frames.
@@ -291,8 +272,7 @@ struct TReplicator_Server : public TBase
 			Ticker.PendingFrame = Buffers.Input.TailFrame();
 		}
 	}
-
-	template<typename T, typename TDriver>
+	
 	void PreSimTick(TDriver* Driver, TNetworkSimBufferContainer<TBufferTypes>& Buffers, TSimulationTicker<TTickSettings>& Ticker, const FNetSimTickParameters& TickParameters)
 	{
 		Ticker.GiveSimulationTime(TickParameters.LocalDeltaTimeSeconds);
@@ -306,12 +286,16 @@ struct TReplicator_Server : public TBase
 // Simulated: "non locally controlled" simulations. We support "Simulation Extrapolation" here (using the sim to fake inputs to advance the sim)
 //	-TODO: this is replicating Input/Sync/Aux which is required to do accurate Simulation Extrapolation. For interpolated simulated proxies, we can skip the input and possibly the aux.
 //	More settings/config options would be nice here
-template<typename TBufferTypes, typename TTickSettings, typename TBase=TReplicatorEmpty<TBufferTypes, TTickSettings>>
-struct TReplicator_Simulated : public TBase
+template<typename Model, typename TBase=TRepController_Empty<Model>>
+struct TRepController_Simulated : public TBase
 {
-	using TInputCmd = typename TBufferTypes::TInputCmd;
-	using TSyncState = typename TBufferTypes::TSyncState;
-	using TAuxState = typename TBufferTypes::TAuxState;
+	using typename TBase::TSimulation;
+	using typename TBase::TTickSettings;
+	using typename TBase::TDriver;
+	using typename TBase::TBufferTypes;
+	using typename TBase::TInputCmd;
+	using typename TBase::TSyncState;
+	using typename TBase::TAuxState;
 
 	// Parent Simulation. If this is set, this simulation will forward predict in sync with this parent sim. The parent sim should be an autonomous proxy driven simulation
 	INetworkedSimulationModel* ParentSimulation = nullptr;
@@ -320,7 +304,7 @@ struct TReplicator_Simulated : public TBase
 	bool bAllowSimulatedExtrapolation = true;
 
 	// Interpolated that will be used if bAllowSimulatedExtrapolation == false && ParentSimulation == nullptr
-	TInterpolator<TBufferTypes, TTickSettings> Interpolator;
+	TNetSimInterpolator<Model> Interpolator;
 
 	// Last Serialized time and state
 	FNetworkSimTime GetLastSerializedSimulationTime() const { return LastSerializedSimulationTime; }
@@ -451,7 +435,6 @@ struct TReplicator_Simulated : public TBase
 		}
 	}
 
-	template<typename TSimulation, typename TDriver>
 	void Reconcile(TSimulation* Simulation, TDriver* Driver, TNetworkSimBufferContainer<TBufferTypes>& Buffers, TSimulationTicker<TTickSettings>& Ticker)
 	{
 		if (ReconcileSimulationTime.IsPositive() == false)
@@ -474,15 +457,14 @@ struct TReplicator_Simulated : public TBase
 			FNetworkSimTime DeltaSimTime = ReconcileSimulationTime - Ticker.GetTotalProcessedSimulationTime();
 			if (DeltaSimTime.IsPositive() && NetworkSimulationModelCVars::EnableSimulatedReconcile())
 			{
-				SimulationExtrapolation<TSimulation, TDriver>(Simulation, Driver, Buffers, Ticker, DeltaSimTime);
+				SimulationExtrapolation(Simulation, Driver, Buffers, Ticker, DeltaSimTime);
 			}
 		}
 		
 		check(Ticker.GetTotalProcessedSimulationTime() <= Ticker.GetTotalAllowedSimulationTime());
 		ReconcileSimulationTime.Reset();
 	}
-
-	template<typename TDriver>
+	
 	void PreSimTick(TDriver* Driver, TNetworkSimBufferContainer<TBufferTypes>& Buffers, TSimulationTicker<TTickSettings>& Ticker, const FNetSimTickParameters& TickParameters)
 	{
 		// Tick if we are dependent simulation or extrapolation is enabled
@@ -499,7 +481,6 @@ struct TReplicator_Simulated : public TBase
 		}
 	}
 
-	template<typename TDriver>
 	void PostSimTick(TDriver* Driver, const TNetworkSimBufferContainer<TBufferTypes>& Buffers, const TSimulationTicker<TTickSettings>& Ticker, const FNetSimTickParameters& TickParameters)
 	{
 		if (GetSimulatedUpdateMode() == ESimulatedUpdateMode::Interpolate)
@@ -516,8 +497,7 @@ struct TReplicator_Simulated : public TBase
 			
 		}
 	}
-
-	template<typename TSimulation, typename TDriver>
+	
 	void DependentRollbackBegin(TSimulation* Simulation, TDriver* Driver, TNetworkSimBufferContainer<TBufferTypes>& Buffers, TSimulationTicker<TTickSettings>& Ticker, const FNetworkSimTime& RollbackDeltaTime, const int32 ParentFrame)
 	{
 		// For now, we make the assumption that our last serialized state and time match the parent simulation.
@@ -540,13 +520,12 @@ struct TReplicator_Simulated : public TBase
 			Driver->VisualLog(nullptr, &LastSerializedSyncState, &LastSerializedAuxState, VLogParams);
 		}
 	}
-
-	template<typename TSimulation, typename TDriver>
+		
 	void DependentRollbackStep(TSimulation* Simulation, TDriver* Driver, TNetworkSimBufferContainer<TBufferTypes>& Buffers, TSimulationTicker<TTickSettings>& Ticker, const FNetworkSimTime& StepTime, const int32 ParentFrame, const bool bFinalStep)
 	{
 		Ticker.SetTotalAllowedSimulationTime( Ticker.GetTotalAllowedSimulationTime() + StepTime );
 
-		SimulationExtrapolation<TSimulation, TDriver>(Simulation, Driver, Buffers, Ticker, StepTime);
+		SimulationExtrapolation(Simulation, Driver, Buffers, Ticker, StepTime);
 
 		TSyncState* SyncState = Buffers.Sync.HeadElement();
 		check(SyncState);
@@ -559,8 +538,7 @@ struct TReplicator_Simulated : public TBase
 	}
 
 private:
-
-	template<typename TSimulation, typename TDriver>
+	
 	void SimulationExtrapolation(TSimulation* Simulation, TDriver* Driver, TNetworkSimBufferContainer<TBufferTypes>& Buffers, TSimulationTicker<TTickSettings>& Ticker, const FNetworkSimTime DeltaSimTime)
 	{
 		// We have extrapolated ahead of the server. The latest network update is now "in the past" from what we rendered last frame.
@@ -598,17 +576,19 @@ private:
 	TInputCmd LastSerializedInputCmd;
 	TSyncState LastSerializedSyncState;
 	TAuxState LastSerializedAuxState;	// Temp? This should be conditional or optional. We want to support not replicating the aux state to simulated proxies
-
 };
 
-
 /** Replicates TSyncState and does basic reconciliation. */
-template<typename TBufferTypes, typename TTickSettings, typename TBase=TReplicatorBase<TBufferTypes, TTickSettings>>
-struct TReplicator_Autonomous : public TBase
+template<typename Model, typename TBase=TRepController_Base<Model>>
+struct TRepController_Autonomous: public TBase
 {
-	using TInputCmd = typename TBufferTypes::TInputCmd;
-	using TSyncState = typename TBufferTypes::TSyncState;
-	using TAuxState = typename TBufferTypes::TAuxState;
+	using typename TBase::TSimulation;
+	using typename TBase::TTickSettings;
+	using typename TBase::TDriver;
+	using typename TBase::TBufferTypes;
+	using typename TBase::TInputCmd;
+	using typename TBase::TSyncState;
+	using typename TBase::TAuxState;
 
 	int32 GetLastSerializedFrame() const { return SerializedFrame; }
 	bool IsReconcileFaultDetected() const { return bReconcileFaultDetected; }
@@ -653,11 +633,15 @@ struct TReplicator_Autonomous : public TBase
 			bReconcileFaultDetected = false;
 			bPendingReconciliation = false;
 
+			TSyncState* ClientSync = Buffers.Sync[SerializedFrame];
+			TAuxState* ClientAux = Buffers.Aux[SerializedFrame];
+
 			// The state the client predicted that corresponds to the state the server just serialized to us
-			if (TSyncState* ClientExistingState = Buffers.Sync[SerializedFrame])
+			if (ClientSync && ClientAux)
 			{
-				// TODO: AuxState->ShouldReconcile
-				if (ClientExistingState->ShouldReconcile(SerializedSyncState) || (NetworkSimulationModelCVars::ForceReconcile() > 0) || (NetworkSimulationModelCVars::ForceReconcileSingle() > 0))
+				const bool bForceReconcile = (NetworkSimulationModelCVars::ForceReconcile() > 0) || (NetworkSimulationModelCVars::ForceReconcileSingle() > 0);
+				
+				if (Model::ShouldReconcile(SerializedSyncState, SerializedAuxState, *ClientSync, *ClientAux))
 				{
 					NetworkSimulationModelCVars::SetForceReconcileSingle(0);
 					UE_CLOG(!Buffers.Input.IsValidFrame(SerializedFrame-1), LogNetworkSim, Error, TEXT("::NetSerialize: Client InputBuffer does not contain data for frame %d. {%s} {%s}"), SerializedFrame, *Buffers.Input.GetBasicDebugStr(), *Buffers.Sync.GetBasicDebugStr());
@@ -684,7 +668,6 @@ struct TReplicator_Autonomous : public TBase
 	// --------------------------------------------------------------------
 	//	Reconcile
 	// --------------------------------------------------------------------
-	template<typename TSimulation, typename TDriver>
 	void Reconcile(TSimulation* Simulation, TDriver* Driver, TNetworkSimBufferContainer<TBufferTypes>& Buffers, TSimulationTicker<TTickSettings>& Ticker)
 	{
 		if (bPendingReconciliation == false && (bDependentSimulationNeedsReconcile == false || SerializedFrame == INDEX_NONE))
@@ -795,7 +778,6 @@ struct TReplicator_Autonomous : public TBase
 	// --------------------------------------------------------------------
 	//	PreSimTick
 	// --------------------------------------------------------------------
-	template<typename TDriver>
 	void PreSimTick(TDriver* Driver, TNetworkSimBufferContainer<TBufferTypes>& Buffers, TSimulationTicker<TTickSettings>& Ticker, const FNetSimTickParameters& TickParameters)
 	{
 		// If we have a reconcile fault, we cannot continue on with the simulation until it clears itself out. This effectively drops the input time and does not sample new inputs
@@ -841,10 +823,15 @@ private:
 };
 
 /** Special replicator for debug buffer, this preserves the local buffer and receives into a replicator-owned buffer (we want these buffers to be distinct/not merged) */
-template<typename TBufferTypes, typename TTickSettings, int32 MaxNumElements=5, typename TBase=TReplicatorEmpty<TBufferTypes, TTickSettings>>
-struct TReplicator_Debug : public TBase
+template<typename Model, int32 MaxNumElements=5, typename TBase=TRepController_Empty<Model>>
+struct TRepController_Debug : public TBase
 {
-	using TDebugState = typename TBufferTypes::TDebugState;
+	using typename TBase::TSimulation;
+	using typename TBase::TTickSettings;
+	using typename TBase::TDriver;
+	using typename TBase::TBufferTypes;
+
+	using TDebugState = typename TBase::TBufferTypes::TDebugState;
 	using TDebugBuffer = typename TNetworkSimBufferContainer<TBufferTypes>::TDebugBuffer;
 
 	int32 GetProxyDirtyCount(TNetworkSimBufferContainer<TBufferTypes>& Buffers) const 
diff --git a/Engine/Plugins/Runtime/NetworkPrediction/Source/NetworkPrediction/Public/NetworkedSimulationModelTraits.h b/Engine/Plugins/Runtime/NetworkPrediction/Source/NetworkPrediction/Public/NetworkedSimulationModelTraits.h
new file mode 100644
--- /dev/null
+++ b/Engine/Plugins/Runtime/NetworkPrediction/Source/NetworkPrediction/Public/NetworkedSimulationModelTraits.h
@@ -0,0 +1,42 @@
+// Copyright 1998-2019 Epic Games, Inc. All Rights Reserved.
+#pragma once
+
+// Base for Network Simulation Model Defs (defines types used by the model)
+class FNetSimModelDefBase
+{
+public:
+
+	// -----------------------------------------------------------------------------------------------
+	//	Must be set by user to compile:
+	// -----------------------------------------------------------------------------------------------
+
+	// The simulation class. This encompasses all state and functions for running the simulation but none of the networking/bookkeeping etc.
+	using Simulation = void;
+
+	// Buffer types: defines your Input/Sync/Aux state (and optional debug state)
+	using BufferTypes = void;	
+
+	// -----------------------------------------------------------------------------------------------
+	//	Optional customization/overrides 
+	// -----------------------------------------------------------------------------------------------
+
+	// Model base class. E.g, common interface all NetworkedSimulationModels can be communicated with
+	using Base = INetworkedSimulationModel;
+	
+	// Tick settings: defines variable vs fix step tick, etc
+	using TickSettings = TNetworkSimTickSettings<>;
+};
+
+// Internal traits that are derived from the user settings (E.g, these can't go in FNetSimModelDefBase because they are dependent on the user types)
+// These can be overridden as well via template specialization but we don't expect this to be common.
+template<typename T>
+struct TNetSimModelTraitsBase
+{
+	// Internal buffer types: the model may wrap the user structs (e.g, input cmds + frame time) for internal usage
+	using InternalBufferTypes = TInternalBufferTypes< typename T::BufferTypes, typename T::TickSettings>;
+	
+	// Driver class: interface that the NetSimModel needs user code to implement in order to function
+	using Driver = TNetworkedSimulationModelDriver<typename T::BufferTypes>;
+};
+
+template<typename T> struct TNetSimModelTraits : public TNetSimModelTraitsBase<T> { };
\ No newline at end of file
diff --git a/Engine/Plugins/Runtime/NetworkPredictionExtras/Source/NetworkPredictionExtras/Private/MockAbilitySimulation.cpp b/Engine/Plugins/Runtime/NetworkPredictionExtras/Source/NetworkPredictionExtras/Private/MockAbilitySimulation.cpp
--- a/Engine/Plugins/Runtime/NetworkPredictionExtras/Source/NetworkPredictionExtras/Private/MockAbilitySimulation.cpp
+++ b/Engine/Plugins/Runtime/NetworkPredictionExtras/Source/NetworkPredictionExtras/Private/MockAbilitySimulation.cpp
@@ -4,7 +4,7 @@
 #include "DrawDebugHelpers.h"
 #include "VisualLogger/VisualLogger.h"
 
-const FName FMockAbilitySimulation::GroupName(TEXT("Ability"));
+const FName FMockAbilityNetSimModelDef::GroupName(TEXT("Ability"));
 
 namespace MockAbilityCVars
 {
@@ -53,7 +53,7 @@ void FMockAbilitySimulation::SimulationTick(const FNetSimTimeStep& TimeStep, con
 	// Local copies of the const input that we will pass into the parent sim as input.
 	FMockAbilityInputCmd LocalCmd = Input.Cmd;
 	FMockAbilitySyncState LocalSync = Input.Sync;
-	FMockAbilityAuxstate LocalAux = Input.Aux;
+	FMockAbilityAuxState LocalAux = Input.Aux;
 
 	const bool bAlreadyBlinking = Input.Aux.BlinkWarmupLeft > 0;
 	const bool bAlreadyDashing = Input.Aux.DashTimeLeft > 0;
@@ -233,13 +233,12 @@ INetworkedSimulationModel* UMockFlyingAbilityComponent::InstantiateNetworkedSimu
 
 	// Simulation
 	FMockAbilitySyncState InitialSyncState;
-	FMockAbilityAuxstate InitialAuxState;
+	FMockAbilityAuxState InitialAuxState;
 	InitMockAbilitySimulation(new FMockAbilitySimulation(), InitialSyncState, InitialAuxState);
 
 	// Model
-	auto NewModel = new FMockAbilitySystem<0>( MockAbilitySimulation, this, InitialSyncState, InitialAuxState);
+	auto* NewModel = new TNetworkedSimulationModel<FMockAbilityNetSimModelDef>( MockAbilitySimulation, this, InitialSyncState, InitialAuxState);
 	InitMockAbilityNetSimModel(NewModel);
-
 	return NewModel;
 }
 
@@ -249,7 +248,7 @@ void UMockFlyingAbilityComponent::ProduceInput(const FNetworkSimTime SimTime, FM
 	ProduceInputDelegate.ExecuteIfBound(SimTime, Cmd);
 }
 
-void UMockFlyingAbilityComponent::FinalizeFrame(const FMockAbilitySyncState& SyncState, const FMockAbilityAuxstate& AuxState)
+void UMockFlyingAbilityComponent::FinalizeFrame(const FMockAbilitySyncState& SyncState, const FMockAbilityAuxState& AuxState)
 {
 	Super::FinalizeFrame(SyncState, AuxState);
 
@@ -284,7 +283,7 @@ const AActor* UMockFlyingAbilityComponent::GetVLogOwner() const
 	return GetOwner();
 }
 
-void UMockFlyingAbilityComponent::VisualLog(const FMockAbilityInputCmd* Input, const FMockAbilitySyncState* Sync, const FMockAbilityAuxstate* Aux, const FVisualLoggingParameters& SystemParameters) const
+void UMockFlyingAbilityComponent::VisualLog(const FMockAbilityInputCmd* Input, const FMockAbilitySyncState* Sync, const FMockAbilityAuxState* Aux, const FVisualLoggingParameters& SystemParameters) const
 {
 	Super::VisualLog(Input, Sync, Aux, SystemParameters);
 }
diff --git a/Engine/Plugins/Runtime/NetworkPredictionExtras/Source/NetworkPredictionExtras/Private/MockNetworkSimulation.cpp b/Engine/Plugins/Runtime/NetworkPredictionExtras/Source/NetworkPredictionExtras/Private/MockNetworkSimulation.cpp
--- a/Engine/Plugins/Runtime/NetworkPredictionExtras/Source/NetworkPredictionExtras/Private/MockNetworkSimulation.cpp
+++ b/Engine/Plugins/Runtime/NetworkPredictionExtras/Source/NetworkPredictionExtras/Private/MockNetworkSimulation.cpp
@@ -59,7 +59,7 @@ static FAutoConsoleVariableRef CVarBindAutomatically(TEXT("mns.BindAutomatically
 
 static bool ForceMispredict = false;
 
-const FName FMockNetworkSimulation::GroupName(TEXT("Mock"));
+const FName FMockNetworkModelDef::GroupName(TEXT("Mock"));
 
 void FMockNetworkSimulation::SimulationTick(const FNetSimTimeStep& TimeStep, const TNetSimInput<TMockNetworkSimulationBufferTypes>& Input, const TNetSimOutput<TMockNetworkSimulationBufferTypes>& Output)
 {
@@ -114,7 +114,7 @@ INetworkedSimulationModel* UMockNetworkSimulationComponent::InstantiateNetworked
 	FMockSyncState InitialSyncState;
 	InitialSyncState.Total = MockValue;
 
-	auto* NewModel = new FMockNetworkModel(MockNetworkSimulation, this, InitialSyncState);
+	auto* NewModel = new TNetworkedSimulationModel<FMockNetworkModelDef>(MockNetworkSimulation, this, InitialSyncState);
 	return NewModel;
 }
 
diff --git a/Engine/Plugins/Runtime/NetworkPredictionExtras/Source/NetworkPredictionExtras/Private/NetworkPredictionExtrasFlyingPawn.cpp b/Engine/Plugins/Runtime/NetworkPredictionExtras/Source/NetworkPredictionExtras/Private/NetworkPredictionExtrasFlyingPawn.cpp
--- a/Engine/Plugins/Runtime/NetworkPredictionExtras/Source/NetworkPredictionExtras/Private/NetworkPredictionExtrasFlyingPawn.cpp
+++ b/Engine/Plugins/Runtime/NetworkPredictionExtras/Source/NetworkPredictionExtras/Private/NetworkPredictionExtrasFlyingPawn.cpp
@@ -529,7 +529,7 @@ float ANetworkPredictionExtrasFlyingPawn_MockAbility::GetMaxStamina() const
 {
 	if (const UMockFlyingAbilityComponent* FlyingAbilityComponent = GetMockFlyingAbilityComponent())
 	{
-		if (const FMockAbilityAuxstate* AuxState = FlyingAbilityComponent->AbilityAuxState.GetStateRead())
+		if (const FMockAbilityAuxState* AuxState = FlyingAbilityComponent->AbilityAuxState.GetStateRead())
 		{
 			return AuxState->MaxStamina;
 		}
diff --git a/Engine/Plugins/Runtime/NetworkPredictionExtras/Source/NetworkPredictionExtras/Public/MockAbilitySimulation.h b/Engine/Plugins/Runtime/NetworkPredictionExtras/Source/NetworkPredictionExtras/Public/MockAbilitySimulation.h
--- a/Engine/Plugins/Runtime/NetworkPredictionExtras/Source/NetworkPredictionExtras/Public/MockAbilitySimulation.h
+++ b/Engine/Plugins/Runtime/NetworkPredictionExtras/Source/NetworkPredictionExtras/Public/MockAbilitySimulation.h
@@ -89,7 +89,7 @@ struct FMockAbilitySyncState : public FFlyingMovementSyncState
 	}
 };
 
-struct FMockAbilityAuxstate : public FFlyingMovementAuxState
+struct FMockAbilityAuxState : public FFlyingMovementAuxState
 {
 	float MaxStamina = 100.f;
 	float StaminaRegenRate = 20.f;
@@ -147,20 +147,9 @@ struct FMockAbilityBlinkCue
 	}
 };
 
-/*
-template<>
-struct TCueHandlerTraits<FMockAbilityBlinkCue> : public TNetSimCueTraitsBase//<FMockAbilityBlinkCue>
-{
-	static constexpr uint8 InvokeMask { (uint8)ESimulationTickContext::Authority };
-	static constexpr bool Replicate { false };
-};
-*/
-
 template<>
 struct TCueHandlerTraits<FMockAbilityBlinkCue> : public TNetSimCueTraits_ReplicatedNonPredicted { };
 
-
-
 // The set of Cues the MockAbility simulation will invoke
 struct FMockAbilityCueSet
 {
@@ -176,29 +165,50 @@ struct FMockAbilityCueSet
 // MockAbilitySimulation definition
 // -------------------------------------------------------
 
-using TMockAbilityBufferTypes = TNetworkSimBufferTypes<FMockAbilityInputCmd, FMockAbilitySyncState, FMockAbilityAuxstate>;
+using TMockAbilityBufferTypes = TNetworkSimBufferTypes<FMockAbilityInputCmd, FMockAbilitySyncState, FMockAbilityAuxState>;
 
 class FMockAbilitySimulation : public FFlyingMovementSimulation
 {
 public:
-	/** Tick group the simulation maps to */
-	static const FName GroupName;
 
 	/** Main update function */
 	void SimulationTick(const FNetSimTimeStep& TimeStep, const TNetSimInput<TMockAbilityBufferTypes>& Input, const TNetSimOutput<TMockAbilityBufferTypes>& Output);
 };
 
-template<int32 InFixedStepMS=0>
-using FMockAbilitySystem = TNetworkedSimulationModel<FMockAbilitySimulation, TMockAbilityBufferTypes, TNetworkSimTickSettings<InFixedStepMS>>;
+class FMockAbilityNetSimModelDef : public FNetSimModelDefBase
+{
+public:
+
+	using Simulation = FMockAbilitySimulation;
+	using BufferTypes = TMockAbilityBufferTypes;
 
-class IMockFlyingAbilitySystemDriver : public TNetworkedSimulationModelDriver<TMockAbilityBufferTypes> { };
+	static const FName GroupName;
+
+	// Compare this state with AuthorityState. return true if a reconcile (correction) should happen
+	static bool ShouldReconcile(const FMockAbilitySyncState& AuthoritySync, const FMockAbilityAuxState& AuthorityAux, const FMockAbilitySyncState& PredictedSync, const FMockAbilityAuxState& PredictedAux)
+	{
+		return FFlyingMovementNetSimModelDef::ShouldReconcile(AuthoritySync, AuthorityAux, PredictedSync, PredictedAux);
+	}
+
+	static void Interpolate(const TInterpolatorParameters<FMockAbilitySyncState, FMockAbilityAuxState>& Params)
+	{
+		FFlyingMovementNetSimModelDef::Interpolate(Params.Cast<FFlyingMovementSyncState, FFlyingMovementAuxState>());
+	}
+};
+
+/** Additional specialized types of the Parametric Movement NetSimModel */
+class FMockAbilityNetSimModelDef_Fixed30hz : public FMockAbilityNetSimModelDef
+{ 
+public:
+	using TickSettings = TNetworkSimTickSettings<33>;
+};
 
 // -------------------------------------------------------------------------------------------------------------------------------
 // ActorComponent for running Mock Ability Simulation 
 // -------------------------------------------------------------------------------------------------------------------------------
 
 UCLASS(BlueprintType, meta=(BlueprintSpawnableComponent))
-class NETWORKPREDICTIONEXTRAS_API UMockFlyingAbilityComponent : public UFlyingMovementComponent, public IMockFlyingAbilitySystemDriver
+class NETWORKPREDICTIONEXTRAS_API UMockFlyingAbilityComponent : public UFlyingMovementComponent, public TNetworkedSimulationModelDriver<TMockAbilityBufferTypes>
 {
 	GENERATED_BODY()
 
@@ -210,15 +220,15 @@ public:
 	FProduceMockAbilityInput ProduceInputDelegate;
 
 	TNetSimStateAccessor<FMockAbilitySyncState> AbilitySyncState;
-	TNetSimStateAccessor<FMockAbilityAuxstate> AbilityAuxState;
+	TNetSimStateAccessor<FMockAbilityAuxState> AbilityAuxState;
 
 	// IMockFlyingAbilitySystemDriver
 	FString GetDebugName() const override;
 	const AActor* GetVLogOwner() const override;
-	void VisualLog(const FMockAbilityInputCmd* Input, const FMockAbilitySyncState* Sync, const FMockAbilityAuxstate* Aux, const FVisualLoggingParameters& SystemParameters) const override;
+	void VisualLog(const FMockAbilityInputCmd* Input, const FMockAbilitySyncState* Sync, const FMockAbilityAuxState* Aux, const FVisualLoggingParameters& SystemParameters) const override;
 
 	void ProduceInput(const FNetworkSimTime SimTime, FMockAbilityInputCmd& Cmd) override;
-	void FinalizeFrame(const FMockAbilitySyncState& SyncState, const FMockAbilityAuxstate& AuxState) override;
+	void FinalizeFrame(const FMockAbilitySyncState& SyncState, const FMockAbilityAuxState& AuxState) override;
 
 	// Required to supress -Woverloaded-virtual warning. We are effectively hiding UFlyingMovementComponent's methods
 	using UFlyingMovementComponent::VisualLog;
@@ -281,11 +291,10 @@ protected:
 	FMockAbilitySimulation* MockAbilitySimulation = nullptr;
 
 	template<typename TSimulation>
-	void InitMockAbilitySimulation(TSimulation* Simulation, FMockAbilitySyncState& InitialSyncState, FMockAbilityAuxstate& InitialAuxState)
+	void InitMockAbilitySimulation(TSimulation* Simulation, FMockAbilitySyncState& InitialSyncState, FMockAbilityAuxState& InitialAuxState)
 	{
 		check(MockAbilitySimulation == nullptr);
-		MockAbilitySimulation = Simulation;		
-
+		MockAbilitySimulation = Simulation;	
 		InitFlyingMovementSimulation(Simulation, InitialSyncState, InitialAuxState);
 	}
 
@@ -294,7 +303,6 @@ protected:
 	{
 		AbilitySyncState.Bind(Model);
 		MovementAuxState.Bind(Model);
-
 		InitFlyingMovementNetSimModel(Model);
 	}
 };
\ No newline at end of file
diff --git a/Engine/Plugins/Runtime/NetworkPredictionExtras/Source/NetworkPredictionExtras/Public/MockNetworkSimulation.h b/Engine/Plugins/Runtime/NetworkPredictionExtras/Source/NetworkPredictionExtras/Public/MockNetworkSimulation.h
--- a/Engine/Plugins/Runtime/NetworkPredictionExtras/Source/NetworkPredictionExtras/Public/MockNetworkSimulation.h
+++ b/Engine/Plugins/Runtime/NetworkPredictionExtras/Source/NetworkPredictionExtras/Public/MockNetworkSimulation.h
@@ -131,32 +131,43 @@ struct FMockAuxState
 
 using TMockNetworkSimulationBufferTypes = TNetworkSimBufferTypes<FMockInputCmd, FMockSyncState, FMockAuxState>;
 
-static FName MockSimulationGroupName("Mock");
-
 class FMockNetworkSimulation
 {
 public:
 
 	/** Main update function */
 	void SimulationTick(const FNetSimTimeStep& TimeStep, const TNetSimInput<TMockNetworkSimulationBufferTypes>& Input, const TNetSimOutput<TMockNetworkSimulationBufferTypes>& Output);
+};
+
+class FMockNetworkModelDef : public FNetSimModelDefBase
+{
+public:
+
+	using Simulation = FMockNetworkSimulation;
+	using BufferTypes = TMockNetworkSimulationBufferTypes;
 
 	/** Tick group the simulation maps to */
 	static const FName GroupName;
-};
 
-// Actual definition of our network simulation.
-using FMockNetworkModel = TNetworkedSimulationModel<FMockNetworkSimulation, TMockNetworkSimulationBufferTypes>;
+	/** Predicted error testing */
+	static bool ShouldReconcile(const FMockSyncState& AuthoritySync, const FMockAuxState& AuthorityAux, const FMockSyncState& PredictedSync, const FMockAuxState& PredictedAux)
+	{
+		return FMath::Abs<float>(AuthoritySync.Total - PredictedSync.Total) > SMALL_NUMBER;
+	}
 
-// Needed to trick UHT into letting UMockNetworkSimulationComponent implement. UHT cannot parse the ::
-// Also needed for forward declaring. Can't just be a typedef/using =
-class IMockSystemDriver : public TNetworkedSimulationModelDriver<TMockNetworkSimulationBufferTypes> { };
+	static void Interpolate(const TInterpolatorParameters<FMockSyncState, FMockAuxState>& Params)
+	{
+		Params.Out.Sync.Total = Params.From.Sync.Total + ((Params.To.Sync.Total - Params.From.Sync.Total) * Params.InterpolationPCT);
+		Params.Out.Aux = Params.To.Aux;
+	}
+};
 
 // -------------------------------------------------------------------------------------------------------------------------------
 // ActorComponent for running a MockNetworkSimulation
 // -------------------------------------------------------------------------------------------------------------------------------
 
 UCLASS(BlueprintType, meta=(BlueprintSpawnableComponent))
-class NETWORKPREDICTIONEXTRAS_API UMockNetworkSimulationComponent : public UNetworkPredictionComponent, public IMockSystemDriver
+class NETWORKPREDICTIONEXTRAS_API UMockNetworkSimulationComponent : public UNetworkPredictionComponent, public TNetworkedSimulationModelDriver<TMockNetworkSimulationBufferTypes>
 {
 	GENERATED_BODY()
 
