NetworkPrediction Plugin
-FCueDispatcher -> FNetSimCueDispatcher + TNetSimCueDispatcher.
-Dispatcher now more time aware. Active RepController can set pruning window and max dispatch time. This allows interpolated controllers to invoke the cue at the right time, so that events align with state.
-Continued refactor of template params to just use a ModelDef to specialize.

#rnx

diff --git a/Engine/Plugins/Runtime/NetworkPrediction/Source/NetworkPrediction/Private/NetworkedSimulationModelCues.cpp b/Engine/Plugins/Runtime/NetworkPrediction/Source/NetworkPrediction/Private/NetworkedSimulationModelCues.cpp
--- a/Engine/Plugins/Runtime/NetworkPrediction/Source/NetworkPrediction/Private/NetworkedSimulationModelCues.cpp
+++ b/Engine/Plugins/Runtime/NetworkPrediction/Source/NetworkPrediction/Private/NetworkedSimulationModelCues.cpp
@@ -6,6 +6,8 @@ DEFINE_LOG_CATEGORY(LogNetSimCues);
 
 FGlobalCueTypeTable FGlobalCueTypeTable::Singleton;
 
+constexpr FNetworkSimTime FCueDispatcherTraitsBase::ReplicationWindow;
+
 // ------------------------------------------------------------------------------------------------------------------------------------------
 //	Mock Cue example
 // ------------------------------------------------------------------------------------------------------------------------------------------
@@ -139,7 +141,7 @@ void TestCues()
 
 	struct FBaseSimulation
 	{
-		static void TickSimulation(FCueDispatcher& Dispatcher)
+		static void TickSimulation(TNetSimCueDispatcher& Dispatcher)
 		{
 			Dispatcher.Invoke<FMockImpactCue>( { FVector(1.f, 2.f, 3.f) } );
 		}
@@ -147,7 +149,7 @@ void TestCues()
 
 	struct FChildSimulation
 	{
-		static void TickSimulation(FCueDispatcher& Dispatcher)
+		static void TickSimulation(TNetSimCueDispatcher& Dispatcher)
 		{
 			FBaseSimulation::TickSimulation(Dispatcher);
 
@@ -156,7 +158,7 @@ void TestCues()
 		}
 	};
 
-	FCueDispatcher ServerDispatcher;
+	TNetSimCueDispatcher ServerDispatcher;
 	FChildSimulation::TickSimulation(ServerDispatcher);
 
 	// -------------------------------------
@@ -171,7 +173,7 @@ void TestCues()
 	// -------------------------------------
 
 	FNetBitReader TempReader(nullptr, TempWriter.GetData(), TempWriter.GetNumBits());
-	FCueDispatcher ClientDispatcher;
+	TNetSimCueDispatcher ClientDispatcher;
 	ClientDispatcher.NetSerializeSavedCues(TempReader, false);
 
 	TestHandlerChild MyObject;
diff --git a/Engine/Plugins/Runtime/NetworkPrediction/Source/NetworkPrediction/Public/NetworkedSimulationModel.h b/Engine/Plugins/Runtime/NetworkPrediction/Source/NetworkPrediction/Public/NetworkedSimulationModel.h
--- a/Engine/Plugins/Runtime/NetworkPrediction/Source/NetworkPrediction/Public/NetworkedSimulationModel.h
+++ b/Engine/Plugins/Runtime/NetworkPrediction/Source/NetworkPrediction/Public/NetworkedSimulationModel.h
@@ -40,6 +40,32 @@
 //
 // ----------------------------------------------------------------------------------------------------------------------------------------------------------------------
 
+// Base for Network Simulation Model Defs (defines types used by the model)
+class FNetSimModelDefBase
+{
+public:
+
+	// -----------------------------------------------------------------------------------------------
+	//	Must be set by user to compile:
+	// -----------------------------------------------------------------------------------------------
+
+	// The simulation class. This encompasses all state and functions for running the simulation but none of the networking/bookkeeping etc.
+	using Simulation = void;
+
+	// Buffer types: defines your Input/Sync/Aux state (and optional debug state)
+	using BufferTypes = void;	
+
+	// -----------------------------------------------------------------------------------------------
+	//	Optional customization/overrides 
+	// -----------------------------------------------------------------------------------------------
+
+	// Model base class. E.g, common interface all NetworkedSimulationModels can be communicated with
+	using Base = INetworkedSimulationModel;
+	
+	// Tick settings: defines variable vs fix step tick, etc
+	using TickSettings = TNetworkSimTickSettings<>;
+};
+
 // Traits for which RepController structs are used internally
 struct TNetSimModelRepControllerTraitsBase
 {
@@ -216,7 +242,7 @@ public:
 
 				{
 					ESimulationTickContext Context = Parameters.Role == ROLE_Authority ? ESimulationTickContext::Authority : ESimulationTickContext::Predict;
-					TScopedSimulationTick UpdateScope(Ticker, Buffers.CueDispatcher, Context, OutputFrame, InputCmd->GetFrameDeltaTime());
+					TScopedSimulationTick<Model> UpdateScope(Ticker, Buffers.CueDispatcher, Context, OutputFrame, InputCmd->GetFrameDeltaTime());
 					
 					Simulation->SimulationTick( 
 						{ InputCmd->GetFrameDeltaTime(), Ticker },
@@ -399,7 +425,7 @@ public:
 	const TState* GetPendingStateRead() const
 	{
 		// Reading is always allowed
-		auto& Buffer = NetSimBufferSelect::Get<TNetworkSimBufferContainer<TBufferTypes>, TState>(Buffers);
+		auto& Buffer = NetSimBufferSelect::Get<TNetworkSimBufferContainer<Model>, TState>(Buffers);
 		return Buffer[Ticker.PendingFrame];
 	}
 
@@ -409,7 +435,7 @@ public:
 		// Writes are always allowed on the authority or during a simulation update for predicting clients
 		if (bHasAuthority || Ticker.bUpdateInProgress)
 		{
-			auto& Buffer = NetSimBufferSelect::Get<TNetworkSimBufferContainer<TBufferTypes>, TState>(Buffers);
+			auto& Buffer = NetSimBufferSelect::Get<TNetworkSimBufferContainer<Model>, TState>(Buffers);
 			return Buffer.WriteFrameInitializedFromHead(Ticker.PendingFrame);
 		}
 		return nullptr;
@@ -481,17 +507,17 @@ public:
 
 	TSimulationTicker<TTickSettings> Ticker;
 
-	TNetworkSimBufferContainer<TBufferTypes> Buffers;
+	TNetworkSimBufferContainer<Model> Buffers;
 
 	TRepProxyServerRPC RepProxy_ServerRPC;
 	TRepProxyAutonomous RepProxy_Autonomous;
 	TRepProxySimulated RepProxy_Simulated;
 	TRepProxyReplay RepProxy_Replay;
 
-	using TInputBuffer = typename TNetworkSimBufferContainer<TBufferTypes>::TInputBuffer;
-	using TSyncBuffer = typename TNetworkSimBufferContainer<TBufferTypes>::TSyncBuffer;
-	using TAuxBuffer = typename TNetworkSimBufferContainer<TBufferTypes>::TAuxBuffer;
-	using TDebugBuffer = typename TNetworkSimBufferContainer<TBufferTypes>::TDebugBuffer;
+	using TInputBuffer = typename TNetworkSimBufferContainer<Model>::TInputBuffer;
+	using TSyncBuffer = typename TNetworkSimBufferContainer<Model>::TSyncBuffer;
+	using TAuxBuffer = typename TNetworkSimBufferContainer<Model>::TAuxBuffer;
+	using TDebugBuffer = typename TNetworkSimBufferContainer<Model>::TDebugBuffer;
 
 	// ------------------------------------------------------------------
 	// RPC Sending helper: provides basic send frequency settings for tracking when the Server RPC can be invoked.
@@ -530,9 +556,9 @@ public:
 #if NETSIM_MODEL_DEBUG
 	TDebugBuffer* GetLocalDebugBuffer() {	return &Buffers.Debug; }
 	TDebugState* GetNextLocalDebugStateWrite() { return Buffers.Debug.WriteFrame( Buffers.Debug.HeadFrame() + 1 ); }
-	TNetworkSimBufferContainer<TBufferTypes>* GetHistoricBuffers(bool bCreate=false)
+	TNetworkSimBufferContainer<Model>* GetHistoricBuffers(bool bCreate=false)
 	{
-		if (HistoricBuffers.IsValid() == false && bCreate) { HistoricBuffers.Reset(new TNetworkSimBufferContainer<TBufferTypes>()); }
+		if (HistoricBuffers.IsValid() == false && bCreate) { HistoricBuffers.Reset(new TNetworkSimBufferContainer<Model>()); }
 		return HistoricBuffers.Get();
 	}
 
@@ -540,7 +566,7 @@ public:
 #else
 	TDebugBuffer* GetLocalDebugBuffer() {	return nullptr; }
 	TDebugState* GetNextLocalDebugStateWrite() { return nullptr; }
-	TNetworkSimBufferContainer<TBufferTypes>* GetHistoricBuffers(bool bCreate=false) { return nullptr; }
+	TNetworkSimBufferContainer<Model>* GetHistoricBuffers(bool bCreate=false) { return nullptr; }
 	TDebugBuffer* GetRemoteDebugBuffer() {	return nullptr; }
 #endif
 
@@ -548,7 +574,7 @@ private:
 
 #if NETSIM_MODEL_DEBUG
 	TRepProxyDebug RepProxy_Debug;
-	TUniquePtr<TNetworkSimBufferContainer<TBufferTypes>> HistoricBuffers;
+	TUniquePtr<TNetworkSimBufferContainer<Model>> HistoricBuffers;
 #endif
 
 };
\ No newline at end of file
diff --git a/Engine/Plugins/Runtime/NetworkPrediction/Source/NetworkPrediction/Public/NetworkedSimulationModelBufferTypes.h b/Engine/Plugins/Runtime/NetworkPrediction/Source/NetworkPrediction/Public/NetworkedSimulationModelBufferTypes.h
new file mode 100644
--- /dev/null
+++ b/Engine/Plugins/Runtime/NetworkPrediction/Source/NetworkPrediction/Public/NetworkedSimulationModelBufferTypes.h
@@ -0,0 +1,156 @@
+// Copyright 1998-2019 Epic Games, Inc. All Rights Reserved.
+#pragma once
+#include "NetworkedSimulationModelTick.h"
+
+// ---------------------------------------------------------------------------------------------------------------------
+//	BufferTypes: template helpers for addressing the different buffer types of the system.
+// ---------------------------------------------------------------------------------------------------------------------
+
+// Enum to refer to buffer type. These are used as template arguments to write generic code that can act on any of the buffers.
+enum class ENetworkSimBufferTypeId : uint8
+{
+	Input,
+	Sync,
+	Aux,
+	Debug
+};
+
+inline FString LexToString(ENetworkSimBufferTypeId A)
+{
+	switch(A)
+	{
+		case ENetworkSimBufferTypeId::Input: return TEXT("Input");
+		case ENetworkSimBufferTypeId::Sync: return TEXT("Sync");
+		case ENetworkSimBufferTypeId::Aux: return TEXT("Aux");
+		case ENetworkSimBufferTypeId::Debug: return TEXT("Debug");
+	};
+	return TEXT("Unknown");
+}
+
+// Helper needed to specialize TNetworkSimBufferTypes::Get (must be done outside of templated struct)
+template<typename TBufferTypes, ENetworkSimBufferTypeId BufferId> 
+struct TSelectTypeHelper
+{
+	using type = void;
+};
+
+template<typename TBufferTypes>
+struct TSelectTypeHelper<TBufferTypes, ENetworkSimBufferTypeId::Input>
+{
+	using type = typename TBufferTypes::TInputCmd;
+};
+
+template<typename TBufferTypes>
+struct TSelectTypeHelper<TBufferTypes, ENetworkSimBufferTypeId::Sync>
+{
+	using type = typename TBufferTypes::TSyncState;
+};
+
+template<typename TBufferTypes>
+struct TSelectTypeHelper<TBufferTypes, ENetworkSimBufferTypeId::Aux>
+{
+	using type = typename TBufferTypes::TAuxState;
+};
+
+template<typename TBufferTypes>
+struct TSelectTypeHelper<TBufferTypes, ENetworkSimBufferTypeId::Debug>
+{
+	using type = typename TBufferTypes::TDebugState;
+};
+
+template<typename T>
+struct HasNetSerialize
+{
+	template<typename U, void (U::*)(const FNetSerializeParams& P)> struct SFINAE {};
+	template<typename U> static char Test(SFINAE<U, &U::NetSerialize>*);
+	template<typename U> static int Test(...);
+	static const bool Has = sizeof(Test<T>(0)) == sizeof(char);
+};
+
+template<typename T>
+struct HasLog
+{
+	template<typename U, void (U::*)(FStandardLoggingParameters& P) const> struct SFINAE {};
+	template<typename U> static char Test(SFINAE<U, &U::Log>*);
+	template<typename U> static int Test(...);
+	static const bool Has = sizeof(Test<T>(0)) == sizeof(char);
+};
+
+// A collection of the system's buffer types. This allows us to collapse the 4 types into a single type to use a template argument elsewhere.
+template<typename InInputCmd, typename InSyncState, typename InAuxState, typename InDebugState = FNetSimProcessedFrameDebugInfo>
+struct TNetworkSimBufferTypes
+{
+	// Quick access to types when you know what you want
+	using TInputCmd = InInputCmd;
+	using TSyncState = InSyncState;
+	using TAuxState = InAuxState;
+	using TDebugState = InDebugState;
+
+	// Template access via ENetworkSimBufferTypeId when "which buffer" is parameterized
+	template<ENetworkSimBufferTypeId Id>
+	struct select_type
+	{
+		using type = typename TSelectTypeHelper< TNetworkSimBufferTypes<TInputCmd, TSyncState, TAuxState, TDebugState>, Id >::type;
+	};
+
+	// Must implement NetSerialize and Log functions. This purposefully does not pass on inherited methods
+	static_assert(HasNetSerialize<InInputCmd>::Has == true, "InputCmd Must implement NetSerialize");
+	static_assert(HasNetSerialize<InSyncState>::Has == true, "SyncState Must implement NetSerialize");
+	static_assert(HasNetSerialize<InAuxState>::Has == true, "AuxState Must implement NetSerialize");
+	static_assert(HasNetSerialize<InDebugState>::Has == true, "DebugState Must implement NetSerialize");
+
+	static_assert(HasLog<InInputCmd>::Has == true, "InputCmd Must implement Log");
+	static_assert(HasLog<InSyncState>::Has == true, "SyncState Must implement Log");
+	static_assert(HasLog<InAuxState>::Has == true, "AuxState Must implement Log");
+	static_assert(HasLog<InDebugState>::Has == true, "DebugState Must implement Log");
+};
+
+// ----------------------------------------------------------------------------------------------------------------------------------------------
+//	FrameCmd - in variable tick simulations we store the timestep of each frame with the input. TFrameCmd wraps the user struct to do this.
+// ----------------------------------------------------------------------------------------------------------------------------------------------
+
+// Wraps an input command (BaseType) in a NetworkSimulation time. 
+template<typename BaseCmdType, typename TickSettings, bool IsFixedTick=(TickSettings::FixedStepMS!=0)>
+struct TFrameCmd : public BaseCmdType
+{
+	FNetworkSimTime GetFrameDeltaTime() const { return FrameDeltaTime; }
+	void SetFrameDeltaTime(const FNetworkSimTime& InTime) { FrameDeltaTime = InTime; }
+	void NetSerialize(const FNetSerializeParams& P)
+	{
+		FrameDeltaTime.NetSerialize(P.Ar);
+		BaseCmdType::NetSerialize(P); 
+	}
+
+	void Log(FStandardLoggingParameters& P) const { BaseCmdType::Log(P); }
+
+private:
+	FNetworkSimTime FrameDeltaTime;
+};
+
+// Fixed tick specialization
+template<typename BaseCmdType, typename TickSettings>
+struct TFrameCmd<BaseCmdType, TickSettings, true> : public BaseCmdType
+{
+	FNetworkSimTime GetFrameDeltaTime() const { return FNetworkSimTime::FromMSec(TickSettings::GetFixedStepMS()); }
+	void SetFrameDeltaTime(const FNetworkSimTime& InTime) { }
+	void NetSerialize(const FNetSerializeParams& P) { BaseCmdType::NetSerialize(P); }
+	void Log(FStandardLoggingParameters& P) const { BaseCmdType::Log(P); }
+};
+
+// ----------------------------------------------------------------------------------------------------------------------------------------------
+// "Internal" Buffers: Helper to turn user supplied buffer types into the "real" buffer types: the InputCmd struct is wrapped in TFrameCmd
+//		-Outside/User code will always use the user supplied buffer types in callbacks/APIs etc
+//		-Internal code to the TNetworkedSimulationModel will use the internal buffer types
+// ----------------------------------------------------------------------------------------------------------------------------------------------
+template<typename TUserBufferTypes, typename TTickSettings>
+struct TInternalBufferTypes : TNetworkSimBufferTypes< 
+	
+	// InputCmds are wrapped in TFrameCmd, which will store an explicit sim delta time if we are not a fixed tick sim
+	TFrameCmd< typename TUserBufferTypes::TInputCmd , TTickSettings>,
+
+	typename TUserBufferTypes::TSyncState,	// SyncState passes through
+	typename TUserBufferTypes::TAuxState,	// Auxstate passes through
+	typename TUserBufferTypes::TDebugState	// Debugstate passes through
+>
+{
+};
\ No newline at end of file
diff --git a/Engine/Plugins/Runtime/NetworkPrediction/Source/NetworkPrediction/Public/NetworkedSimulationModelCues.h b/Engine/Plugins/Runtime/NetworkPrediction/Source/NetworkPrediction/Public/NetworkedSimulationModelCues.h
--- a/Engine/Plugins/Runtime/NetworkPrediction/Source/NetworkPrediction/Public/NetworkedSimulationModelCues.h
+++ b/Engine/Plugins/Runtime/NetworkPrediction/Source/NetworkPrediction/Public/NetworkedSimulationModelCues.h
@@ -3,7 +3,6 @@
 
 #include "Engine/EngineTypes.h"
 #include "NetworkedSimulationModelTime.h"
-#include "Containers/SortedMap.h"
 
 NETWORKPREDICTION_API DECLARE_LOG_CATEGORY_EXTERN(LogNetSimCues, Log, All);
 
@@ -326,13 +325,14 @@ private:
 template<typename TCueHandler>
 TCueDispatchTable<TCueHandler> TCueDispatchTable<TCueHandler>::Singleton;
 
-// ------------------------------------------------------------------------------
+// ------------------------------------------------------------------------------------------------------------------------------------------------------------
 //	CueDispatcher
 //	-Entry point for invoking cues during a SimulationTick
 //	-Holds recorded cue state for replication
-// ------------------------------------------------------------------------------
+// ------------------------------------------------------------------------------------------------------------------------------------------------------------
 
-struct FCueDispatcher
+// Non-templated, "networking model independent" base: this is what the pure simulation code gets to invoke cues. 
+struct FNetSimCueDispatcher
 {
 	template<typename T>
 	void Invoke(T&& Cue)
@@ -367,10 +367,42 @@ struct FCueDispatcher
 				UE_LOG(LogNetSimCues, Log, TEXT("Suppressing Cue Invocation %s. Mask: %d. TickContext: %d"), *FGlobalCueTypeTable::Get().GetTypeName(T::ID), TCueHandlerTraits<T>::InvokeMask, (int32)Context.TickContext);
 			}
 		}
+	}	
+
+protected:
+
+	bool EnsureValidContext()
+	{
+		return ensure(Context.CurrentSimTime.IsPositive() && Context.TickContext != ESimulationTickContext::None);
 	}
 
-	// Serializes the recorded cue record
-	// Note: this is not NetSerializing the individual cues in the record, it is serializing the record which contains the "net bits" of each individual cue
+	// Sim Context: the Sim has to tell the dispatcher what its doing so that it can decide if it should supress Invocations or not
+	struct FContext
+	{
+		FNetworkSimTime CurrentSimTime;
+		ESimulationTickContext TickContext;
+	};
+	
+	TArray<FSavedCue> SavedCues;		// Cues that must be saved for some period of time, either for replication or for uniqueness testing
+	TArray<FSavedCue> TransientCues;	// Cues that are dispatched on this frame and then forgotten about
+
+	FContext Context;
+};
+
+// Traits for TNetSimCueDispatcher
+struct NETWORKPREDICTION_API FCueDispatcherTraitsBase
+{
+	// Window for replicating a NetSimCue. That is, after a cue is invoked, it has ReplicationWindow time before it will be pruned.
+	// If a client does not get a net update for the sim in this window, they will miss the event.
+	static constexpr FNetworkSimTime ReplicationWindow = FNetworkSimTime::FromRealTimeMS(200);
+};
+template<typename T> struct TCueDispatcherTraits : public FCueDispatcherTraitsBase { };
+
+// Templated cue dispatcher that can be specialized per networking model definition. This is what the system actually uses internally, but is not exposed to user code.
+template<typename Model=void>
+struct TNetSimCueDispatcher : public FNetSimCueDispatcher
+{
+	// Serializes all saved cues
 	void NetSerializeSavedCues(FArchive& Ar, bool bIsInterpolatingSim)
 	{
 		FNetSimCueTypeId NumCues = SavedCues.Num();
@@ -422,29 +454,26 @@ struct FCueDispatcher
 		}
 	}
 
+	// Dispatches and prunes saved/transient cues
 	template<typename T>
 	void DispatchCueRecord(T& Handler, FNetworkSimTime CurrentSimTime)
 	{
-		// Prune
-		FNetworkSimTime CutOffTime = CurrentSimTime - FNetworkSimTime::FromMSec(200); // fixme: hardcoded bad. But also, should take into consideration confirmed server frame in prediction case
-		int32 FirstValidElement = 0;
-		for (; FirstValidElement < SavedCues.Num(); ++FirstValidElement)
+		const FNetworkSimTime PruneTime = (MinPruneNetSimTime.IsPositive() ? MinPruneNetSimTime : CurrentSimTime) - TCueDispatcherTraits<Model>::ReplicationWindow;
+
+		// Dispatch
+		int32 PruneIdx = -1;
+		for (int32 SavedCueIdx = 0; SavedCueIdx < SavedCues.Num(); ++ SavedCueIdx)
 		{
-			if ( SavedCues[FirstValidElement].Time > CutOffTime)
+			FSavedCue& SavedCue = SavedCues[SavedCueIdx];
+			if ( SavedCue.Time <= PruneTime)
 			{
-				break;
+				PruneIdx = SavedCueIdx;
 			}
-		}
 
-		if (FirstValidElement > 0)
-		{
-			SavedCues.RemoveAt(0, FirstValidElement, false);
-		}
+			const bool bWithold = MaxDispatchNetSimTime.IsPositive() && SavedCue.Time > MaxDispatchNetSimTime;
+			UE_CLOG(bWithold, LogNetSimCues, Log, TEXT("Withholding Cue %s. %s > %s"), *SavedCue.GetTypeName(), *SavedCue.Time.ToString(), *MaxDispatchNetSimTime.ToString());
 
-		// Dispatch
-		for (FSavedCue& SavedCue : SavedCues)
-		{
-			if (SavedCue.bDispatched == false)
+			if (SavedCue.bDispatched == false && !bWithold)
 			{
 				UE_LOG(LogNetSimCues, Log, TEXT("Dispatching NetSimCue %s."), *SavedCue.GetTypeName());
 				SavedCue.bDispatched = true;
@@ -458,29 +487,37 @@ struct FCueDispatcher
 			TCueDispatchTable<T>::Get().Dispatch(TransientCue, Handler, CurrentSimTime - TransientCue.Time);
 		}
 		TransientCues.Reset();
-	}
 
-	// Sim Context: the Sim has to tell the dispatcher what its doing so that it can decide if it should supress Invocations or not
-	struct FContext
-	{
-		FNetworkSimTime CurrentSimTime;
-		ESimulationTickContext TickContext;
-	};
+		// Prune
+		if (PruneIdx >= 0)
+		{
+#if !(UE_BUILD_SHIPPING || UE_BUILD_TEST)
+			for (int32 i=0; i <= PruneIdx; ++i)
+			{
+				UE_CLOG(!SavedCues[i].bDispatched, LogNetSimCues, Warning, TEXT("Non-Dispatched Cue is about to be pruned! %s. %s"), *SavedCues[i].GetTypeName());
+				UE_LOG(LogNetSimCues, Log, TEXT("Pruning Cue %s. Invoke Time: %s. Current Time: %s"), *SavedCues[i].GetTypeName(), *SavedCues[i].Time.ToString(), *CurrentSimTime.ToString());
+			}
+#endif
+
+			SavedCues.RemoveAt(0, PruneIdx+1, false);
+		}
+	}
 
+	// Push/pop simulation context.
 	void PushContext(const FContext& InContext) { Context = InContext; }
 	void PopContext() { Context = FContext(); }
 
-private:
+	// Set max simulation time to invoke cues for. Clearing = always process latest
+	void SetMaxDispatchTime(const FNetworkSimTime& NewMaxTime) const { MaxDispatchNetSimTime = NewMaxTime; }
+	void ClearMaxDispatchTime() const { MaxDispatchNetSimTime = FNetworkSimTime(); }
 
-	TArray<FSavedCue> SavedCues;		// Cues that must be saved for some period of time, either for replication or for uniqueness testing
-	TArray<FSavedCue> TransientCues;	// Cues that are dispatched on this frame and then forgotten about
+	// Set minimum prune time. System will not prune cues that happen after this time. ReplicationWindow is applied *on top* of this
+	void SetPruneTime(const FNetworkSimTime& NewPruneTime) const { MinPruneNetSimTime = NewPruneTime; }
 
-	FContext Context;
+private:
 
-	bool EnsureValidContext()
-	{
-		return ensure(Context.CurrentSimTime.IsPositive() && Context.TickContext != ESimulationTickContext::None);
-	}
+	mutable FNetworkSimTime MaxDispatchNetSimTime;	// (If set) Max time of a saved cue that can be dispatched. Needed for delaying interpolation / client side buffering.
+	mutable FNetworkSimTime MinPruneNetSimTime;		// (If set) Min time we start at to look for cues to prune. Note ReplicationWindow is always applied on top of this.
 };
 
 // ------------------------------------------------------------------------------------------------------------------------------------------------------------------------------------------------
diff --git a/Engine/Plugins/Runtime/NetworkPrediction/Source/NetworkPrediction/Public/NetworkedSimulationModelInterpolator.h b/Engine/Plugins/Runtime/NetworkPrediction/Source/NetworkPrediction/Public/NetworkedSimulationModelInterpolator.h
--- a/Engine/Plugins/Runtime/NetworkPrediction/Source/NetworkPrediction/Public/NetworkedSimulationModelInterpolator.h
+++ b/Engine/Plugins/Runtime/NetworkPrediction/Source/NetworkPrediction/Public/NetworkedSimulationModelInterpolator.h
@@ -67,7 +67,7 @@ struct TNetSimInterpolator
 	bool bEnableVisualLog = true;
 
 	template<typename TSystemDriver>
-	void PostSimTick(TSystemDriver* Driver, const TNetworkSimBufferContainer<TBufferTypes>& Buffers, const FSimulationTickState& TickInfo, const FNetSimTickParameters& TickParameters)
+	FRealTime PostSimTick(TSystemDriver* Driver, const TNetworkSimBufferContainer<Model>& Buffers, const FSimulationTickState& TickInfo, const FNetSimTickParameters& TickParameters)
 	{
 		const bool bDoVLog = NetworkInterpolationDebugCVars::VLog() && bEnableVisualLog;
 		const float DeltaSeconds = TickParameters.LocalDeltaTimeSeconds;
@@ -84,13 +84,13 @@ struct TNetSimInterpolator
 
 				Driver->FinalizeFrame(*HeadState, *AuxState);
 			}
-			return;
+			return TickInfo.GetTotalProcessedSimulationTime().ToRealTimeSeconds();
 		}
 		
 		if (TickInfo.SimulationTimeBuffer.Num() <= 1)
 		{
 			// Cant interpolate yet	
-			return;
+			return 0.f;
 		}
 
 		auto& SimulationTimeBuffer = TickInfo.SimulationTimeBuffer;
@@ -126,7 +126,7 @@ struct TNetSimInterpolator
 					FVector2D LocalTimeVsInterpolationTime(LogOwner->GetWorld()->GetTimeSeconds(), (int32)(InterpolationTime * 1000.f));
 					UE_VLOG_HISTOGRAM(LogOwner, LogNetInterpolation, Log, "ServerSimulationTimeGraph", LocalInterpolationTimeName, LocalTimeVsInterpolationTime);
 				}
-				return;
+				return InterpolationTime;
 			}
 		}
 
@@ -263,6 +263,8 @@ struct TNetSimInterpolator
 				InternalIdx ^= 1;
 			}
 		}
+
+		return InterpolationTime;
 	}
 
 private:
diff --git a/Engine/Plugins/Runtime/NetworkPrediction/Source/NetworkPrediction/Public/NetworkedSimulationModelReplicators.h b/Engine/Plugins/Runtime/NetworkPrediction/Source/NetworkPrediction/Public/NetworkedSimulationModelReplicators.h
--- a/Engine/Plugins/Runtime/NetworkPrediction/Source/NetworkPrediction/Public/NetworkedSimulationModelReplicators.h
+++ b/Engine/Plugins/Runtime/NetworkPrediction/Source/NetworkPrediction/Public/NetworkedSimulationModelReplicators.h
@@ -56,10 +56,10 @@ struct FNetworkSimulationSerialization
 // -------------------------------------------------------------------------------------------------------
 
 /** Helper that writes a new input cmd to the input buffer, at given Frame (usually the sim's PendingFrame). If frame doesn't exist, ProduceInput is called on the driver if bProduceInputViaDriver=true, otherwise the input cmd will be initialized from the previous input cmd. */
-template<typename TDriver, typename TBufferTypes>
-void GenerateLocalInputCmdAtFrame(TDriver* Driver, TNetworkSimBufferContainer<TBufferTypes>& Buffers, const FNetworkSimTime& DeltaSimTime, int32 Frame, bool bProduceInputViaDriver=true)
+template<typename Model>
+void GenerateLocalInputCmdAtFrame(typename TNetSimModelTraits<Model>::Driver* Driver, TNetworkSimBufferContainer<Model>& Buffers, const FNetworkSimTime& DeltaSimTime, int32 Frame, bool bProduceInputViaDriver=true)
 {
-	using TInputCmd = typename TBufferTypes::TInputCmd;
+	using TInputCmd = typename TNetworkSimBufferContainer<Model>::TInputCmd;
 	if (TInputCmd* InputCmd = Buffers.Input[Frame])
 	{
 		InputCmd->SetFrameDeltaTime(DeltaSimTime);
@@ -76,15 +76,16 @@ void GenerateLocalInputCmdAtFrame(TDriver* Driver, TNetworkSimBufferContainer<TB
 }
 
 /** Helper to generate a local input cmd if we have simulation time to spend and advance the simulation's MaxAllowedFrame so that it can be processed. */
-template<typename TDriver, typename TBufferTypes, typename TTickSettings>
-void TryGenerateLocalInput(TDriver* Driver, TNetworkSimBufferContainer<TBufferTypes>& Buffers, TSimulationTicker<TTickSettings>& Ticker, bool bProduceInputViaDriver=true)
+template<typename Model>
+void TryGenerateLocalInput(typename TNetSimModelTraits<Model>::Driver* Driver, TNetworkSimBufferContainer<Model>& Buffers, TSimulationTicker<typename Model::TickSettings>& Ticker, bool bProduceInputViaDriver=true)
 {
+	using TBufferTypes = typename TNetSimModelTraits<Model>::InternalBufferTypes;
 	using TInputCmd = typename TBufferTypes::TInputCmd;
 
 	FNetworkSimTime DeltaSimTime = Ticker.GetRemainingAllowedSimulationTime();
 	if (DeltaSimTime.IsPositive())
 	{
-		GenerateLocalInputCmdAtFrame(Driver, Buffers, DeltaSimTime, Ticker.PendingFrame, bProduceInputViaDriver);
+		GenerateLocalInputCmdAtFrame<Model>(Driver, Buffers, DeltaSimTime, Ticker.PendingFrame, bProduceInputViaDriver);
 		Ticker.MaxAllowedFrame = Ticker.PendingFrame;
 	}
 }
@@ -113,19 +114,19 @@ struct TRepController_Empty
 	using TAuxState = typename TBufferTypes::TAuxState;
 
 	// Used for client shadowstate compares. Should just call GetDirtyCount() on the buffer you are replicating
-	int32 GetProxyDirtyCount(TNetworkSimBufferContainer<TBufferTypes>& Buffers) const { return 0; }
+	int32 GetProxyDirtyCount(TNetworkSimBufferContainer<Model>& Buffers) const { return 0; }
 
 	// NetSerialize: just serialize the network data. Don't run simulation steps. Every replicator will be NetSerialized before moving on to Reconcile phase.
-	void NetSerialize(const FNetSerializeParams& P, TNetworkSimBufferContainer<TBufferTypes>& Buffers, const TSimulationTicker<TTickSettings>& Ticker) { }
+	void NetSerialize(const FNetSerializeParams& P, TNetworkSimBufferContainer<Model>& Buffers, const TSimulationTicker<TTickSettings>& Ticker) { }
 
 	// Reconcile: called after everyone has NetSerialized. "Get right with the server": this function is about reconciling what the server told you vs what you have extrapolated or forward predicted locally
-	void Reconcile(TSimulation* Sim, TDriver* Driver, TNetworkSimBufferContainer<TBufferTypes>& Buffers, TSimulationTicker<TTickSettings>& Ticker) { }
+	void Reconcile(TSimulation* Sim, TDriver* Driver, TNetworkSimBufferContainer<Model>& Buffers, TSimulationTicker<TTickSettings>& Ticker) { }
 
 	// Called prior to input processing. This function must update Ticker to allow simulation time (from TickParameters) and to possibly get new input.
-	void PreSimTick(TDriver* Driver, TNetworkSimBufferContainer<TBufferTypes>& Buffers, TSimulationTicker<TTickSettings>& Ticker, const FNetSimTickParameters& TickParameters) { }
+	void PreSimTick(TDriver* Driver, TNetworkSimBufferContainer<Model>& Buffers, TSimulationTicker<TTickSettings>& Ticker, const FNetSimTickParameters& TickParameters) { }
 
 	// Called after input processing. Should finalize the frame and do any smoothing/interpolation. This function is not allowed to modify the buffers or tick state, or even call the simulation/Update function.
-	void PostSimTick(TDriver* Driver, const TNetworkSimBufferContainer<TBufferTypes>& Buffers, const TSimulationTicker<TTickSettings>& Ticker, const FNetSimTickParameters& TickParameters) { }
+	void PostSimTick(TDriver* Driver, const TNetworkSimBufferContainer<Model>& Buffers, const TSimulationTicker<TTickSettings>& Ticker, const FNetSimTickParameters& TickParameters) { }
 };
 
 // Base implementation with minimal functionality: will grant simulation time and call finalize frame
@@ -139,14 +140,14 @@ struct TRepController_Base : public TBase
 	using typename TBase::TInputCmd;
 
 	// Called prior to input processing. This function must updated Ticker to allow simulation time (from TickParameters) and to possibly get new input.
-	void PreSimTick(TDriver* Driver, TNetworkSimBufferContainer<TBufferTypes>& Buffers, TSimulationTicker<TTickSettings>& Ticker, const FNetSimTickParameters& TickParameters)
+	void PreSimTick(TDriver* Driver, TNetworkSimBufferContainer<Model>& Buffers, TSimulationTicker<TTickSettings>& Ticker, const FNetSimTickParameters& TickParameters)
 	{
 		// Accumulate local delta time into Ticker
 		Ticker.GiveSimulationTime(TickParameters.LocalDeltaTimeSeconds);
 	}
 
 	// Called after input processing. Should finalize the frame and do any smoothing/interpolation. This function is not allowed to modify the buffers or tick state, or even call the simulation/Update function.
-	void PostSimTick(TDriver* Driver, const TNetworkSimBufferContainer<TBufferTypes>& Buffers, const TSimulationTicker<TTickSettings>& Ticker, const FNetSimTickParameters& TickParameters)
+	void PostSimTick(TDriver* Driver, const TNetworkSimBufferContainer<Model>& Buffers, const TSimulationTicker<TTickSettings>& Ticker, const FNetSimTickParameters& TickParameters)
 	{
 		// Sync to latest frame if there is any
 		if (Buffers.Sync.Num() > 0)
@@ -169,7 +170,7 @@ struct TRepController_SimTime : public TBase
 	using typename TBase::TBufferTypes;
 	using typename TBase::TInputCmd;
 
-	void NetSerialize(const FNetSerializeParams& P, TNetworkSimBufferContainer<TBufferTypes>& Buffers, TSimulationTicker<TTickSettings>& Ticker)
+	void NetSerialize(const FNetSerializeParams& P, TNetworkSimBufferContainer<Model>& Buffers, TSimulationTicker<TTickSettings>& Ticker)
 	{
 		TBase::NetSerialize(P, Buffers, Ticker);
 		SerializedTime = Ticker.GetTotalProcessedSimulationTime();
@@ -200,12 +201,12 @@ struct TRepController_Sequence : public TBase
 	using typename TBase::TBufferTypes;
 	using typename TBase::TInputCmd;
 
-	int32 GetProxyDirtyCount(TNetworkSimBufferContainer<TBufferTypes>& Buffers) const 
+	int32 GetProxyDirtyCount(TNetworkSimBufferContainer<Model>& Buffers) const 
 	{
 		return Buffers.template Get<BufferId>().GetDirtyCount() ^ (TBase::GetProxyDirtyCount(Buffers) << 2); 
 	}
 
-	void NetSerialize(const FNetSerializeParams& P, TNetworkSimBufferContainer<TBufferTypes>& Buffers, TSimulationTicker<TTickSettings>& Ticker)
+	void NetSerialize(const FNetSerializeParams& P, TNetworkSimBufferContainer<Model>& Buffers, TSimulationTicker<TTickSettings>& Ticker)
 	{
 		TBase::NetSerialize(P, Buffers, Ticker);
 
@@ -258,7 +259,7 @@ struct TRepController_Server : public TBase
 	using typename TBase::TBufferTypes;
 	using typename TBase::TInputCmd;
 	
-	void Reconcile(TSimulation* Simulation, TDriver* Driver, TNetworkSimBufferContainer<TBufferTypes>& Buffers, TSimulationTicker<TTickSettings>& Ticker)
+	void Reconcile(TSimulation* Simulation, TDriver* Driver, TNetworkSimBufferContainer<Model>& Buffers, TSimulationTicker<TTickSettings>& Ticker)
 	{
 		// After receiving input, server may process up to the latest received frames.
 		// (If we needed to buffer input server side for whatever reason, we would do it here)
@@ -273,7 +274,7 @@ struct TRepController_Server : public TBase
 		}
 	}
 	
-	void PreSimTick(TDriver* Driver, TNetworkSimBufferContainer<TBufferTypes>& Buffers, TSimulationTicker<TTickSettings>& Ticker, const FNetSimTickParameters& TickParameters)
+	void PreSimTick(TDriver* Driver, TNetworkSimBufferContainer<Model>& Buffers, TSimulationTicker<TTickSettings>& Ticker, const FNetSimTickParameters& TickParameters)
 	{
 		Ticker.GiveSimulationTime(TickParameters.LocalDeltaTimeSeconds);
 		if (TickParameters.bGenerateLocalInputCmds)
@@ -324,12 +325,12 @@ struct TRepController_Simulated : public TBase
 		return ESimulatedUpdateMode::Interpolate;
 	}
 	
-	int32 GetProxyDirtyCount(TNetworkSimBufferContainer<TBufferTypes>& Buffers) const
+	int32 GetProxyDirtyCount(TNetworkSimBufferContainer<Model>& Buffers) const
 	{
 		return Buffers.Sync.GetDirtyCount() ^ (TBase::GetProxyDirtyCount(Buffers) << 2); 
 	}
 
-	void NetSerialize(const FNetSerializeParams& P, TNetworkSimBufferContainer<TBufferTypes>& Buffers, TSimulationTicker<TTickSettings>& Ticker)
+	void NetSerialize(const FNetSerializeParams& P, TNetworkSimBufferContainer<Model>& Buffers, TSimulationTicker<TTickSettings>& Ticker)
 	{
 		FArchive& Ar = P.Ar;
 
@@ -435,7 +436,7 @@ struct TRepController_Simulated : public TBase
 		}
 	}
 
-	void Reconcile(TSimulation* Simulation, TDriver* Driver, TNetworkSimBufferContainer<TBufferTypes>& Buffers, TSimulationTicker<TTickSettings>& Ticker)
+	void Reconcile(TSimulation* Simulation, TDriver* Driver, TNetworkSimBufferContainer<Model>& Buffers, TSimulationTicker<TTickSettings>& Ticker)
 	{
 		if (ReconcileSimulationTime.IsPositive() == false)
 		{
@@ -465,7 +466,7 @@ struct TRepController_Simulated : public TBase
 		ReconcileSimulationTime.Reset();
 	}
 	
-	void PreSimTick(TDriver* Driver, TNetworkSimBufferContainer<TBufferTypes>& Buffers, TSimulationTicker<TTickSettings>& Ticker, const FNetSimTickParameters& TickParameters)
+	void PreSimTick(TDriver* Driver, TNetworkSimBufferContainer<Model>& Buffers, TSimulationTicker<TTickSettings>& Ticker, const FNetSimTickParameters& TickParameters)
 	{
 		// Tick if we are dependent simulation or extrapolation is enabled
 		if (ParentSimulation || (GetSimulatedUpdateMode() == ESimulatedUpdateMode::Extrapolate))
@@ -477,15 +478,19 @@ struct TRepController_Simulated : public TBase
 			}
 			
 			// Generate local input if we are ready to tick. Note that we pass false so we won't call into the Driver to produce the input, we will use the last serialized InputCmd's values
-			TryGenerateLocalInput(Driver, Buffers, Ticker, false);
+			TryGenerateLocalInput<Model>(Driver, Buffers, Ticker, false);
 		}
 	}
 
-	void PostSimTick(TDriver* Driver, const TNetworkSimBufferContainer<TBufferTypes>& Buffers, const TSimulationTicker<TTickSettings>& Ticker, const FNetSimTickParameters& TickParameters)
+	void PostSimTick(TDriver* Driver, const TNetworkSimBufferContainer<Model>& Buffers, const TSimulationTicker<TTickSettings>& Ticker, const FNetSimTickParameters& TickParameters)
 	{
 		if (GetSimulatedUpdateMode() == ESimulatedUpdateMode::Interpolate)
 		{
-			Interpolator.template PostSimTick<TDriver>(Driver, Buffers, Ticker, TickParameters);
+			const FNetworkSimTime::FRealTime InterpolationRealTime = Interpolator.template PostSimTick<TDriver>(Driver, Buffers, Ticker, TickParameters);
+			const FNetworkSimTime InterpolationSimTime = FNetworkSimTime::FromRealTimeSeconds(InterpolationRealTime);
+
+			Buffers.CueDispatcher.SetMaxDispatchTime(InterpolationSimTime);
+			Buffers.CueDispatcher.SetPruneTime(InterpolationSimTime);
 		}
 		else
 		{
@@ -494,11 +499,13 @@ struct TRepController_Simulated : public TBase
 			{
 				Driver->FinalizeFrame(*Buffers.Sync.HeadElement(), *Buffers.Aux.HeadElement());
 			}
-			
+
+			Buffers.CueDispatcher.ClearMaxDispatchTime();
+			Buffers.CueDispatcher.SetPruneTime(LastSerializedSimulationTime);
 		}
 	}
 	
-	void DependentRollbackBegin(TSimulation* Simulation, TDriver* Driver, TNetworkSimBufferContainer<TBufferTypes>& Buffers, TSimulationTicker<TTickSettings>& Ticker, const FNetworkSimTime& RollbackDeltaTime, const int32 ParentFrame)
+	void DependentRollbackBegin(TSimulation* Simulation, TDriver* Driver, TNetworkSimBufferContainer<Model>& Buffers, TSimulationTicker<TTickSettings>& Ticker, const FNetworkSimTime& RollbackDeltaTime, const int32 ParentFrame)
 	{
 		// For now, we make the assumption that our last serialized state and time match the parent simulation.
 		// This would not always be the case with low frequency simulated proxies. But could be handled by replicating the simulations together (at the replication level)
@@ -521,7 +528,7 @@ struct TRepController_Simulated : public TBase
 		}
 	}
 		
-	void DependentRollbackStep(TSimulation* Simulation, TDriver* Driver, TNetworkSimBufferContainer<TBufferTypes>& Buffers, TSimulationTicker<TTickSettings>& Ticker, const FNetworkSimTime& StepTime, const int32 ParentFrame, const bool bFinalStep)
+	void DependentRollbackStep(TSimulation* Simulation, TDriver* Driver, TNetworkSimBufferContainer<Model>& Buffers, TSimulationTicker<TTickSettings>& Ticker, const FNetworkSimTime& StepTime, const int32 ParentFrame, const bool bFinalStep)
 	{
 		Ticker.SetTotalAllowedSimulationTime( Ticker.GetTotalAllowedSimulationTime() + StepTime );
 
@@ -539,7 +546,7 @@ struct TRepController_Simulated : public TBase
 
 private:
 	
-	void SimulationExtrapolation(TSimulation* Simulation, TDriver* Driver, TNetworkSimBufferContainer<TBufferTypes>& Buffers, TSimulationTicker<TTickSettings>& Ticker, const FNetworkSimTime DeltaSimTime)
+	void SimulationExtrapolation(TSimulation* Simulation, TDriver* Driver, TNetworkSimBufferContainer<Model>& Buffers, TSimulationTicker<TTickSettings>& Ticker, const FNetworkSimTime DeltaSimTime)
 	{
 		// We have extrapolated ahead of the server. The latest network update is now "in the past" from what we rendered last frame.
 		// We will insert a new frame to make up the difference from the last known state to where we want to be in the now.
@@ -560,7 +567,7 @@ private:
 
 		// Do the actual update
 		{
-			TScopedSimulationTick UpdateScope(Ticker, Buffers.CueDispatcher, ESimulationTickContext::Resimulate, OutputFrame, NewCmd->GetFrameDeltaTime());
+			TScopedSimulationTick<Model> UpdateScope(Ticker, Buffers.CueDispatcher, ESimulationTickContext::Resimulate, OutputFrame, NewCmd->GetFrameDeltaTime());
 			Simulation->SimulationTick( 
 				{ NewCmd->GetFrameDeltaTime(), Ticker },
 				{ *NewCmd, *PrevSyncState, *AuxState },
@@ -597,7 +604,7 @@ struct TRepController_Autonomous: public TBase
 	TArray<INetworkedSimulationModel*> DependentSimulations;
 	bool bDependentSimulationNeedsReconcile = false;
 
-	int32 GetProxyDirtyCount(TNetworkSimBufferContainer<TBufferTypes>& Buffers) const
+	int32 GetProxyDirtyCount(TNetworkSimBufferContainer<Model>& Buffers) const
 	{
 		return (Buffers.Sync.GetDirtyCount()) ^ (Buffers.Aux.GetDirtyCount() << 1) ^ (TBase::GetProxyDirtyCount(Buffers) << 2); 
 	}
@@ -605,7 +612,7 @@ struct TRepController_Autonomous: public TBase
 	// --------------------------------------------------------------------
 	//	NetSerialize
 	// --------------------------------------------------------------------
-	void NetSerialize(const FNetSerializeParams& P, TNetworkSimBufferContainer<TBufferTypes>& Buffers, TSimulationTicker<TTickSettings>& Ticker)
+	void NetSerialize(const FNetSerializeParams& P, TNetworkSimBufferContainer<Model>& Buffers, TSimulationTicker<TTickSettings>& Ticker)
 	{
 		FArchive& Ar = P.Ar;
 		
@@ -668,7 +675,7 @@ struct TRepController_Autonomous: public TBase
 	// --------------------------------------------------------------------
 	//	Reconcile
 	// --------------------------------------------------------------------
-	void Reconcile(TSimulation* Simulation, TDriver* Driver, TNetworkSimBufferContainer<TBufferTypes>& Buffers, TSimulationTicker<TTickSettings>& Ticker)
+	void Reconcile(TSimulation* Simulation, TDriver* Driver, TNetworkSimBufferContainer<Model>& Buffers, TSimulationTicker<TTickSettings>& Ticker)
 	{
 		if (bPendingReconciliation == false && (bDependentSimulationNeedsReconcile == false || SerializedFrame == INDEX_NONE))
 		{
@@ -753,7 +760,7 @@ struct TRepController_Autonomous: public TBase
 
 			// Do the actual update
 			{
-				TScopedSimulationTick UpdateScope(Ticker, Buffers.CueDispatcher, ESimulationTickContext::Resimulate, OutputFrame, ResimulateCmd->GetFrameDeltaTime());
+				TScopedSimulationTick<Model> UpdateScope(Ticker, Buffers.CueDispatcher, ESimulationTickContext::Resimulate, OutputFrame, ResimulateCmd->GetFrameDeltaTime());
 				Simulation->SimulationTick( 
 					{ ResimulateCmd->GetFrameDeltaTime(), Ticker },
 					{ *ResimulateCmd, *PrevSyncState, *AuxState },
@@ -778,7 +785,7 @@ struct TRepController_Autonomous: public TBase
 	// --------------------------------------------------------------------
 	//	PreSimTick
 	// --------------------------------------------------------------------
-	void PreSimTick(TDriver* Driver, TNetworkSimBufferContainer<TBufferTypes>& Buffers, TSimulationTicker<TTickSettings>& Ticker, const FNetSimTickParameters& TickParameters)
+	void PreSimTick(TDriver* Driver, TNetworkSimBufferContainer<Model>& Buffers, TSimulationTicker<TTickSettings>& Ticker, const FNetSimTickParameters& TickParameters)
 	{
 		// If we have a reconcile fault, we cannot continue on with the simulation until it clears itself out. This effectively drops the input time and does not sample new inputs
 		if (bReconcileFaultDetected)
@@ -792,7 +799,7 @@ struct TRepController_Autonomous: public TBase
 			{
 				// Prediction: add simulation time and generate new commands
 				Ticker.GiveSimulationTime(TickParameters.LocalDeltaTimeSeconds);
-				TryGenerateLocalInput(Driver, Buffers, Ticker);
+				TryGenerateLocalInput<Model>(Driver, Buffers, Ticker);
 			}
 			else
 			{
@@ -803,12 +810,21 @@ struct TRepController_Autonomous: public TBase
 				NonPredictedInputTimeAccumulator.Accumulate(NonPredictedInputTime, TickParameters.LocalDeltaTimeSeconds);
 				if (NonPredictedInputTime.IsPositive())
 				{
-					GenerateLocalInputCmdAtFrame(Driver, Buffers, NonPredictedInputTime, Buffers.Input.HeadFrame() + 1);
+					GenerateLocalInputCmdAtFrame<Model>(Driver, Buffers, NonPredictedInputTime, Buffers.Input.HeadFrame() + 1);
 				}
 			}
 		}
 	}
 
+	void PostSimTick(TDriver* Driver, const TNetworkSimBufferContainer<Model>& Buffers, const TSimulationTicker<TTickSettings>& Ticker, const FNetSimTickParameters& TickParameters)
+	{
+		TBase::PostSimTick(Driver, Buffers, Ticker, TickParameters);
+
+		Buffers.CueDispatcher.ClearMaxDispatchTime();
+		Buffers.CueDispatcher.SetPruneTime(SerializedTime);
+	}
+	
+
 private:
 	
 	TSyncState SerializedSyncState;
@@ -832,14 +848,14 @@ struct TRepController_Debug : public TBase
 	using typename TBase::TBufferTypes;
 
 	using TDebugState = typename TBase::TBufferTypes::TDebugState;
-	using TDebugBuffer = typename TNetworkSimBufferContainer<TBufferTypes>::TDebugBuffer;
+	using TDebugBuffer = typename TNetworkSimBufferContainer<Model>::TDebugBuffer;
 
-	int32 GetProxyDirtyCount(TNetworkSimBufferContainer<TBufferTypes>& Buffers) const 
+	int32 GetProxyDirtyCount(TNetworkSimBufferContainer<Model>& Buffers) const 
 	{
 		return Buffers.Debug.GetDirtyCount() ^ (TBase::GetProxyDirtyCount(Buffers) << 2); 
 	}
 
-	void NetSerialize(const FNetSerializeParams& P, TNetworkSimBufferContainer<TBufferTypes>& Buffers, TSimulationTicker<TTickSettings>& Ticker)
+	void NetSerialize(const FNetSerializeParams& P, TNetworkSimBufferContainer<Model>& Buffers, TSimulationTicker<TTickSettings>& Ticker)
 	{
 		TBase::NetSerialize(P, Buffers, Ticker);
 		FArchive& Ar = P.Ar;
diff --git a/Engine/Plugins/Runtime/NetworkPrediction/Source/NetworkPrediction/Public/NetworkedSimulationModelTick.h b/Engine/Plugins/Runtime/NetworkPrediction/Source/NetworkPrediction/Public/NetworkedSimulationModelTick.h
--- a/Engine/Plugins/Runtime/NetworkPrediction/Source/NetworkPrediction/Public/NetworkedSimulationModelTick.h
+++ b/Engine/Plugins/Runtime/NetworkPrediction/Source/NetworkPrediction/Public/NetworkedSimulationModelTick.h
@@ -1,10 +1,7 @@
 // Copyright 1998-2019 Epic Games, Inc. All Rights Reserved.
 
 #pragma once
-
 #include "NetworkedSimulationModelTime.h"
-#include "NetworkedSimulationModelCues.h"
-
 
 // Holds per-simulation settings about how ticking is supposed to happen.
 template<int32 InFixedStepMS=0, int32 InMaxStepMS=0>
@@ -170,32 +167,6 @@ private:
 	TRealTimeAccumulator<TSettings>	RealtimeAccumulator;	
 };
 
-// Scoped helper to be used right before entering a call to the sim's ::SimulationTick function.
-// Important to note that this advances the PendingFrame to the output Frame. So that any writes that occur to the buffers during this scope will go to the output frame.
-struct TScopedSimulationTick
-{
-	TScopedSimulationTick(FSimulationTickState& InTicker, FCueDispatcher& InDispatcher, ESimulationTickContext TickContext, const int32& InOutputFrame, const FNetworkSimTime& InDeltaSimTime)
-		: Ticker(InTicker), Dispatcher(InDispatcher), OutputFrame(InOutputFrame), DeltaSimTime(InDeltaSimTime)
-	{
-		check(Ticker.bUpdateInProgress == false);
-		Ticker.PendingFrame = OutputFrame;
-		Ticker.bUpdateInProgress = true;
-
-		Dispatcher.PushContext({Ticker.GetTotalProcessedSimulationTime() + DeltaSimTime, TickContext}); // Cues "take place" at the end of the frame
-	}
-	~TScopedSimulationTick()
-	{
-		Ticker.IncrementTotalProcessedSimulationTime(DeltaSimTime, OutputFrame);
-		Ticker.bUpdateInProgress = false;
-
-		Dispatcher.PopContext();
-	}
-	FSimulationTickState& Ticker;
-	FCueDispatcher& Dispatcher;
-	const int32& OutputFrame;
-	const FNetworkSimTime& DeltaSimTime;
-};
-
 // ----------------------------------------------------------------------------------------------------------------------------------------------
 //	FNetSimTimeStep Time data structure that is passed into SimulationTick
 // ----------------------------------------------------------------------------------------------------------------------------------------------
diff --git a/Engine/Plugins/Runtime/NetworkPrediction/Source/NetworkPrediction/Public/NetworkedSimulationModelTime.h b/Engine/Plugins/Runtime/NetworkPrediction/Source/NetworkPrediction/Public/NetworkedSimulationModelTime.h
--- a/Engine/Plugins/Runtime/NetworkPrediction/Source/NetworkPrediction/Public/NetworkedSimulationModelTime.h
+++ b/Engine/Plugins/Runtime/NetworkPrediction/Source/NetworkPrediction/Public/NetworkedSimulationModelTime.h
@@ -1,7 +1,6 @@
 // Copyright 1998-2019 Epic Games, Inc. All Rights Reserved.
 
 #pragma once
-
 #include "NetworkedSimulationModelBuffer.h"
 
 // ----------------------------------------------------------------------------------------------------------------------------------------------
@@ -27,14 +26,14 @@ struct FNetworkSimTime
 	FNetworkSimTime() { }
 
 	// Things get confusing with templated types and overloaded functions. To avoid that, use these funcs to construct from either msec or real time
-	static inline FNetworkSimTime FromMSec(const FSimTime& InTime) { return FNetworkSimTime(InTime); } 
-	static inline FNetworkSimTime FromRealTimeSeconds(const FRealTime& InRealTime) { return FNetworkSimTime( static_cast<FSimTime>(InRealTime * GetRealToSimFactor())); } 
+	static constexpr FNetworkSimTime FromMSec(const FSimTime& InTime) { return FNetworkSimTime(InTime); } 
+	static constexpr FNetworkSimTime FromRealTimeSeconds(const FRealTime& InRealTime) { return FNetworkSimTime( static_cast<FSimTime>(InRealTime * GetRealToSimFactor())); } 
 	FRealTime ToRealTimeSeconds() const { return (Time * GetSimToRealFactor()); }
 
 	// Direct casts to "real time MS" which should be rarely used in practice (TRealTimeAccumulator only current case). All other cases of "real time" imply seconds.
-	static inline FNetworkSimTime FromRealTimeMS(const FRealTime& InRealTime) { return FNetworkSimTime( static_cast<FSimTime>(InRealTime)); } 
+	static constexpr FNetworkSimTime FromRealTimeMS(const FRealTime& InRealTime) { return FNetworkSimTime( static_cast<FSimTime>(InRealTime)); }
+	
 	FRealTime ToRealTimeMS() const { return static_cast<FRealTime>(Time); }
-
 	FString ToString() const { return LexToString(this->Time); }
 
 	bool IsPositive() const { return (Time > 0); }
@@ -64,5 +63,5 @@ private:
 
 	FSimTime Time = 0;
 
-	FNetworkSimTime(const FSimTime& InTime) { Time = InTime; }
+	constexpr FNetworkSimTime(const FSimTime& InTime) { Time = InTime; }
 };
\ No newline at end of file
diff --git a/Engine/Plugins/Runtime/NetworkPrediction/Source/NetworkPrediction/Public/NetworkedSimulationModelTraits.h b/Engine/Plugins/Runtime/NetworkPrediction/Source/NetworkPrediction/Public/NetworkedSimulationModelTraits.h
--- a/Engine/Plugins/Runtime/NetworkPrediction/Source/NetworkPrediction/Public/NetworkedSimulationModelTraits.h
+++ b/Engine/Plugins/Runtime/NetworkPrediction/Source/NetworkPrediction/Public/NetworkedSimulationModelTraits.h
@@ -1,30 +1,30 @@
 // Copyright 1998-2019 Epic Games, Inc. All Rights Reserved.
 #pragma once
+#include "NetworkedSimulationModelBufferTypes.h"
 
-// Base for Network Simulation Model Defs (defines types used by the model)
-class FNetSimModelDefBase
+/** This is the "system driver", it has functions that the TNetworkedSimulationModel needs to call internally, that are specific to the types but not specific to the simulation itself. */
+template<typename TBufferTypes>
+class TNetworkedSimulationModelDriver
 {
 public:
+	using TInputCmd = typename TBufferTypes::TInputCmd;
+	using TSyncState = typename TBufferTypes::TSyncState;
+	using TAuxState= typename TBufferTypes::TAuxState;
 
-	// -----------------------------------------------------------------------------------------------
-	//	Must be set by user to compile:
-	// -----------------------------------------------------------------------------------------------
+	// Debug string that can be used in internal warning/error logs
+	virtual FString GetDebugName() const = 0;
 
-	// The simulation class. This encompasses all state and functions for running the simulation but none of the networking/bookkeeping etc.
-	using Simulation = void;
+	// Owning object for Visual Logs so that the system can emit them internally
+	virtual const AActor* GetVLogOwner() const = 0;
 
-	// Buffer types: defines your Input/Sync/Aux state (and optional debug state)
-	using BufferTypes = void;	
-
-	// -----------------------------------------------------------------------------------------------
-	//	Optional customization/overrides 
-	// -----------------------------------------------------------------------------------------------
-
-	// Model base class. E.g, common interface all NetworkedSimulationModels can be communicated with
-	using Base = INetworkedSimulationModel;
+	// Call to visual log the given states. Note that not all 3 will always be present and you should check for nullptrs.
+	virtual void VisualLog(const TInputCmd* Input, const TSyncState* Sync, const TAuxState* Aux, const FVisualLoggingParameters& SystemParameters) const = 0;
+	
+	// Called whenever the sim is ready to process new local input.
+	virtual void ProduceInput(const FNetworkSimTime SimTime, TInputCmd&) = 0;
 	
-	// Tick settings: defines variable vs fix step tick, etc
-	using TickSettings = TNetworkSimTickSettings<>;
+	// Called from the Network Sim at the end of the sim frame when there is new sync data.
+	virtual void FinalizeFrame(const typename TBufferTypes::TSyncState& SyncState, const typename TBufferTypes::TAuxState& AuxState) = 0;
 };
 
 // Internal traits that are derived from the user settings (E.g, these can't go in FNetSimModelDefBase because they are dependent on the user types)
diff --git a/Engine/Plugins/Runtime/NetworkPrediction/Source/NetworkPrediction/Public/NetworkedSimulationModelTypes.h b/Engine/Plugins/Runtime/NetworkPrediction/Source/NetworkPrediction/Public/NetworkedSimulationModelTypes.h
--- a/Engine/Plugins/Runtime/NetworkPrediction/Source/NetworkPrediction/Public/NetworkedSimulationModelTypes.h
+++ b/Engine/Plugins/Runtime/NetworkPrediction/Source/NetworkPrediction/Public/NetworkedSimulationModelTypes.h
@@ -5,109 +5,7 @@
 #include "NetworkedSimulationModelBuffer.h"
 #include "NetworkedSimulationModelCues.h"
 #include "NetworkedSimulationModelTime.h"
-
-// ---------------------------------------------------------------------------------------------------------------------
-//	BufferTypes: template helpers for addressing the different buffer types of the system.
-// ---------------------------------------------------------------------------------------------------------------------
-
-// Enum to refer to buffer type. These are used as template arguments to write generic code that can act on any of the buffers.
-enum class ENetworkSimBufferTypeId : uint8
-{
-	Input,
-	Sync,
-	Aux,
-	Debug
-};
-
-inline FString LexToString(ENetworkSimBufferTypeId A)
-{
-	switch(A)
-	{
-		case ENetworkSimBufferTypeId::Input: return TEXT("Input");
-		case ENetworkSimBufferTypeId::Sync: return TEXT("Sync");
-		case ENetworkSimBufferTypeId::Aux: return TEXT("Aux");
-		case ENetworkSimBufferTypeId::Debug: return TEXT("Debug");
-	};
-	return TEXT("Unknown");
-}
-
-// Helper needed to specialize TNetworkSimBufferTypes::Get (must be done outside of templated struct)
-template<typename TBufferTypes, ENetworkSimBufferTypeId BufferId> 
-struct TSelectTypeHelper
-{
-	using type = void;
-};
-
-template<typename TBufferTypes>
-struct TSelectTypeHelper<TBufferTypes, ENetworkSimBufferTypeId::Input>
-{
-	using type = typename TBufferTypes::TInputCmd;
-};
-
-template<typename TBufferTypes>
-struct TSelectTypeHelper<TBufferTypes, ENetworkSimBufferTypeId::Sync>
-{
-	using type = typename TBufferTypes::TSyncState;
-};
-
-template<typename TBufferTypes>
-struct TSelectTypeHelper<TBufferTypes, ENetworkSimBufferTypeId::Aux>
-{
-	using type = typename TBufferTypes::TAuxState;
-};
-
-template<typename TBufferTypes>
-struct TSelectTypeHelper<TBufferTypes, ENetworkSimBufferTypeId::Debug>
-{
-	using type = typename TBufferTypes::TDebugState;
-};
-
-template<typename T>
-struct HasNetSerialize
-{
-	template<typename U, void (U::*)(const FNetSerializeParams& P)> struct SFINAE {};
-	template<typename U> static char Test(SFINAE<U, &U::NetSerialize>*);
-	template<typename U> static int Test(...);
-	static const bool Has = sizeof(Test<T>(0)) == sizeof(char);
-};
-
-template<typename T>
-struct HasLog
-{
-	template<typename U, void (U::*)(FStandardLoggingParameters& P) const> struct SFINAE {};
-	template<typename U> static char Test(SFINAE<U, &U::Log>*);
-	template<typename U> static int Test(...);
-	static const bool Has = sizeof(Test<T>(0)) == sizeof(char);
-};
-
-// A collection of the system's buffer types. This allows us to collapse the 4 types into a single type to use a template argument elsewhere.
-template<typename InInputCmd, typename InSyncState, typename InAuxState, typename InDebugState = FNetSimProcessedFrameDebugInfo>
-struct TNetworkSimBufferTypes
-{
-	// Quick access to types when you know what you want
-	using TInputCmd = InInputCmd;
-	using TSyncState = InSyncState;
-	using TAuxState = InAuxState;
-	using TDebugState = InDebugState;
-
-	// Template access via ENetworkSimBufferTypeId when "which buffer" is parameterized
-	template<ENetworkSimBufferTypeId Id>
-	struct select_type
-	{
-		using type = typename TSelectTypeHelper< TNetworkSimBufferTypes<TInputCmd, TSyncState, TAuxState, TDebugState>, Id >::type;
-	};
-
-	// Must implement NetSerialize and Log functions. This purposefully does not pass on inherited methods
-	static_assert(HasNetSerialize<InInputCmd>::Has == true, "InputCmd Must implement NetSerialize");
-	static_assert(HasNetSerialize<InSyncState>::Has == true, "SyncState Must implement NetSerialize");
-	static_assert(HasNetSerialize<InAuxState>::Has == true, "AuxState Must implement NetSerialize");
-	static_assert(HasNetSerialize<InDebugState>::Has == true, "DebugState Must implement NetSerialize");
-
-	static_assert(HasLog<InInputCmd>::Has == true, "InputCmd Must implement Log");
-	static_assert(HasLog<InSyncState>::Has == true, "SyncState Must implement Log");
-	static_assert(HasLog<InAuxState>::Has == true, "AuxState Must implement Log");
-	static_assert(HasLog<InDebugState>::Has == true, "DebugState Must implement Log");
-};
+#include "NetworkedSimulationModelTraits.h"
 
 // ---------------------------------------------------------------------------------------------------------------------
 //	TNetworkSimBufferContainer
@@ -190,34 +88,12 @@ namespace NetSimBufferSelect
 	}
 }
 
-// Struct that encapsulates writing a new element to a buffer. This is used to allow a new Aux state to be created in the ::SimulationTick loop.
-template<typename T>
-struct TLazyStateAccessor
-{
-	TLazyStateAccessor(TUniqueFunction<T*()> && InFunc)
-		: GetWriteNextFunc(MoveTemp(InFunc)) { }
-
-	T* GetWriteNext() const
-	{
-		if (CachedWriteNext == nullptr)
-		{
-			CachedWriteNext = GetWriteNextFunc();
-		}
-		return CachedWriteNext;
-	}
-
-private:
-
-	mutable T* CachedWriteNext = nullptr;
-	TUniqueFunction<T*()> GetWriteNextFunc;
-};
-
 // The main container for all of our buffers
-template<typename InBufferTypes>
+template<typename Model>
 struct TNetworkSimBufferContainer
 {
 	// Collection of types we were assigned
-	using TBufferTypes = InBufferTypes;
+	using TBufferTypes = typename TNetSimModelTraits<Model>::InternalBufferTypes;
 
 	using TInputCmd = typename TBufferTypes::TInputCmd;
 	using TSyncState = typename TBufferTypes::TSyncState;
@@ -253,14 +129,14 @@ struct TNetworkSimBufferContainer
 	TAuxBuffer Aux;
 	TDebugBuffer Debug;
 
-	FCueDispatcher CueDispatcher;
+	TNetSimCueDispatcher<Model> CueDispatcher;
 
 	// Finally, template accessor for getting buffers based on enum. This is really what all this junk is about.
 	// This allows other templated classes in the system to access a specific buffer from another templated argument
 	template<ENetworkSimBufferTypeId BufferId>
 	typename select_buffer_type<BufferId>::type& Get()
 	{
-		return TBufferGetterHelper<TNetworkSimBufferContainer<TBufferTypes>, BufferId>::Get(*this);
+		return TBufferGetterHelper<TNetworkSimBufferContainer<Model>, BufferId>::Get(*this);
 	}
 };
 
@@ -289,7 +165,7 @@ struct TNetSimStateAccessor
 	{
 		GetStateFunc = [NetSimModel](bool bWrite, TState*& OutState, bool& OutSafe)
 		{
-			auto& Buffer = NetSimBufferSelect::Get<TNetworkSimBufferContainer<typename TNetworkSimModel::TBufferTypes>, TState>(NetSimModel->Buffers);
+			auto& Buffer = NetSimBufferSelect::Get<TNetworkSimBufferContainer<typename TNetworkSimModel::Model>, TState>(NetSimModel->Buffers);
 			OutState = bWrite ? Buffer.WriteFrameInitializedFromHead(NetSimModel->Ticker.PendingFrame) : Buffer[NetSimModel->Ticker.PendingFrame];
 			OutSafe = NetSimModel->Ticker.bUpdateInProgress;
 		};
@@ -329,76 +205,6 @@ private:
 	TFunction<void(bool bForWrite, TState*& OutState, bool& OutSafe)> GetStateFunc;
 };
 
-// ----------------------------------------------------------------------------------------------------------------------------------------------
-//	FrameCmd - in variable tick simulations we store the timestep of each frame with the input. TFrameCmd wraps the user struct to do this.
-// ----------------------------------------------------------------------------------------------------------------------------------------------
-
-// Wraps an input command (BaseType) in a NetworkSimulation time. 
-template<typename BaseCmdType, typename TickSettings, bool IsFixedTick=(TickSettings::FixedStepMS!=0)>
-struct TFrameCmd : public BaseCmdType
-{
-	FNetworkSimTime GetFrameDeltaTime() const { return FrameDeltaTime; }
-	void SetFrameDeltaTime(const FNetworkSimTime& InTime) { FrameDeltaTime = InTime; }
-	void NetSerialize(const FNetSerializeParams& P)
-	{
-		FrameDeltaTime.NetSerialize(P.Ar);
-		BaseCmdType::NetSerialize(P); 
-	}
-
-	void Log(FStandardLoggingParameters& P) const { BaseCmdType::Log(P); }
-
-private:
-	FNetworkSimTime FrameDeltaTime;
-};
-
-// Fixed tick specialization
-template<typename BaseCmdType, typename TickSettings>
-struct TFrameCmd<BaseCmdType, TickSettings, true> : public BaseCmdType
-{
-	FNetworkSimTime GetFrameDeltaTime() const { return FNetworkSimTime::FromMSec(TickSettings::GetFixedStepMS()); }
-	void SetFrameDeltaTime(const FNetworkSimTime& InTime) { }
-	void NetSerialize(const FNetSerializeParams& P) { BaseCmdType::NetSerialize(P); }
-	void Log(FStandardLoggingParameters& P) const { BaseCmdType::Log(P); }
-};
-
-// Helper to turn user supplied buffer types into the "real" buffer types: the InputCmd struct is wrapped in TFrameCmd
-template<typename TUserBufferTypes, typename TTickSettings>
-struct TInternalBufferTypes : TNetworkSimBufferTypes< 
-	
-	// InputCmds are wrapped in TFrameCmd, which will store an explicit sim delta time if we are not a fixed tick sim
-	TFrameCmd< typename TUserBufferTypes::TInputCmd , TTickSettings>,
-
-	typename TUserBufferTypes::TSyncState,	// SyncState passes through
-	typename TUserBufferTypes::TAuxState,	// Auxstate passes through
-	typename TUserBufferTypes::TDebugState	// Debugstate passes through
->
-{
-};
-
-/** This is the "system driver", it has functions that the TNetworkedSimulationModel needs to call internally, that are specific to the types but not specific to the simulation itself. */
-template<typename TBufferTypes>
-class TNetworkedSimulationModelDriver
-{
-public:
-	using TInputCmd = typename TBufferTypes::TInputCmd;
-	using TSyncState = typename TBufferTypes::TSyncState;
-	using TAuxState= typename TBufferTypes::TAuxState;
-
-	// Debug string that can be used in internal warning/error logs
-	virtual FString GetDebugName() const = 0;
-
-	// Owning object for Visual Logs so that the system can emit them internally
-	virtual const AActor* GetVLogOwner() const = 0;
-
-	// Call to visual log the given states. Note that not all 3 will always be present and you should check for nullptrs.
-	virtual void VisualLog(const TInputCmd* Input, const TSyncState* Sync, const TAuxState* Aux, const FVisualLoggingParameters& SystemParameters) const = 0;
-	
-	// Called whenever the sim is ready to process new local input.
-	virtual void ProduceInput(const FNetworkSimTime SimTime, TInputCmd&) = 0;
-	
-	// Called from the Network Sim at the end of the sim frame when there is new sync data.
-	virtual void FinalizeFrame(const typename TBufferTypes::TSyncState& SyncState, const typename TBufferTypes::TAuxState& AuxState) = 0;
-};
 
 // ----------------------------------------------------------------------------------------------------------------------------------------------
 //	Simulation Input/Output parameters. These are the data structures passed into the simulation code each frame
@@ -427,9 +233,35 @@ struct TNetSimOutput
 {
 	typename TBufferTypes::TSyncState& Sync;
 	const TNetSimLazyWriter<typename TBufferTypes::TAuxState>& Aux;
-	FCueDispatcher& CueDispatch;
+	FNetSimCueDispatcher& CueDispatch;
 
-	TNetSimOutput(typename TBufferTypes::TSyncState& InSync, const TNetSimLazyWriter<typename TBufferTypes::TAuxState>& InAux, FCueDispatcher& InCueDispatch)
+	TNetSimOutput(typename TBufferTypes::TSyncState& InSync, const TNetSimLazyWriter<typename TBufferTypes::TAuxState>& InAux, FNetSimCueDispatcher& InCueDispatch)
 		: Sync(InSync), Aux(InAux), CueDispatch(InCueDispatch) { }
 };
 
+// Scoped helper to be used right before entering a call to the sim's ::SimulationTick function.
+// Important to note that this advances the PendingFrame to the output Frame. So that any writes that occur to the buffers during this scope will go to the output frame.
+template<typename Model>
+struct TScopedSimulationTick
+{
+	TScopedSimulationTick(FSimulationTickState& InTicker, TNetSimCueDispatcher<Model>& InDispatcher, ESimulationTickContext TickContext, const int32& InOutputFrame, const FNetworkSimTime& InDeltaSimTime)
+		: Ticker(InTicker), Dispatcher(InDispatcher), OutputFrame(InOutputFrame), DeltaSimTime(InDeltaSimTime)
+	{
+		check(Ticker.bUpdateInProgress == false);
+		Ticker.PendingFrame = OutputFrame;
+		Ticker.bUpdateInProgress = true;
+
+		Dispatcher.PushContext({Ticker.GetTotalProcessedSimulationTime() + DeltaSimTime, TickContext}); // Cues "take place" at the end of the frame
+	}
+	~TScopedSimulationTick()
+	{
+		Ticker.IncrementTotalProcessedSimulationTime(DeltaSimTime, OutputFrame);
+		Ticker.bUpdateInProgress = false;
+
+		Dispatcher.PopContext();
+	}
+	FSimulationTickState& Ticker;
+	TNetSimCueDispatcher<Model>& Dispatcher;
+	const int32& OutputFrame;
+	const FNetworkSimTime& DeltaSimTime;
+};
diff --git a/Engine/Plugins/Runtime/NetworkPredictionExtras/Source/NetworkPredictionExtras/Private/MockAbilitySimulation.cpp b/Engine/Plugins/Runtime/NetworkPredictionExtras/Source/NetworkPredictionExtras/Private/MockAbilitySimulation.cpp
--- a/Engine/Plugins/Runtime/NetworkPredictionExtras/Source/NetworkPredictionExtras/Private/MockAbilitySimulation.cpp
+++ b/Engine/Plugins/Runtime/NetworkPredictionExtras/Source/NetworkPredictionExtras/Private/MockAbilitySimulation.cpp
@@ -296,8 +296,12 @@ static FAutoConsoleVariableRef CVarBindAutomatically(TEXT("NetworkPredictionExtr
 
 void UMockFlyingAbilityComponent::HandleCue(FMockAbilityBlinkCue& BlinkCue, const FNetworkSimTime& ElapsedTime)
 {
-	UE_LOG(LogTemp, Display, TEXT("BlinkCue! : <%f, %f, %f> - <%f, %f, %f>. ElapsedTime: %s"), BlinkCue.StartLocation.X, BlinkCue.StartLocation.Y, BlinkCue.StartLocation.Z,
-		BlinkCue.StopLocation.X, BlinkCue.StopLocation.Y, BlinkCue.StopLocation.Z, *ElapsedTime.ToString()); //*BlinkCue.StartLocation.ToString(), *BlinkCue.StopLocation.ToString());
+	FString RoleStr = GetOwnerRole() == ROLE_Authority ? TEXT("Server") : TEXT("Client");
+
+	FVector Delta = GetOwner()->GetActorLocation() - BlinkCue.StopLocation;
+
+	UE_LOG(LogTemp, Display, TEXT("[%s] BlinkCue! : <%f, %f, %f> - <%f, %f, %f>. ElapsedTime: %s. Delta: %.3f"), *RoleStr, BlinkCue.StartLocation.X, BlinkCue.StartLocation.Y, BlinkCue.StartLocation.Z,
+		BlinkCue.StopLocation.X, BlinkCue.StopLocation.Y, BlinkCue.StopLocation.Z, *ElapsedTime.ToString(), Delta.Size()); //*BlinkCue.StartLocation.ToString(), *BlinkCue.StopLocation.ToString());
 
 	// Crude compensation for cue firing in the past (note this is not necessary! Some cues not care and need to see the "full" effect regardless of when it happened)
 	float Duration = BlinkCueDuration - ElapsedTime.ToRealTimeSeconds();
diff --git a/Engine/Plugins/Runtime/NetworkPredictionExtras/Source/NetworkPredictionExtras/Public/MockAbilitySimulation.h b/Engine/Plugins/Runtime/NetworkPredictionExtras/Source/NetworkPredictionExtras/Public/MockAbilitySimulation.h
--- a/Engine/Plugins/Runtime/NetworkPredictionExtras/Source/NetworkPredictionExtras/Public/MockAbilitySimulation.h
+++ b/Engine/Plugins/Runtime/NetworkPredictionExtras/Source/NetworkPredictionExtras/Public/MockAbilitySimulation.h
@@ -192,6 +192,9 @@ public:
 
 	static void Interpolate(const TInterpolatorParameters<FMockAbilitySyncState, FMockAbilityAuxState>& Params)
 	{
+		Params.Out.Sync = Params.To.Sync;
+		Params.Out.Aux = Params.To.Aux;
+
 		FFlyingMovementNetSimModelDef::Interpolate(Params.Cast<FFlyingMovementSyncState, FFlyingMovementAuxState>());
 	}
 };
