NetworkPrediction Plugin
-Support for rewindable NetSimCues
-NetSimCues get emplaced in saved buffer when invoking rather than copied around

#rnx

diff --git a/Engine/Plugins/Runtime/NetworkPrediction/Source/NetworkPrediction/Private/NetworkedSimulationModelCues.cpp b/Engine/Plugins/Runtime/NetworkPrediction/Source/NetworkPrediction/Private/NetworkedSimulationModelCues.cpp
--- a/Engine/Plugins/Runtime/NetworkPrediction/Source/NetworkPrediction/Private/NetworkedSimulationModelCues.cpp
+++ b/Engine/Plugins/Runtime/NetworkPrediction/Source/NetworkPrediction/Private/NetworkedSimulationModelCues.cpp
@@ -24,6 +24,9 @@ struct FMockImpactCue
 {
 	NETSIMCUE_BODY();
 
+	FMockImpactCue() = default;
+	FMockImpactCue(const FVector& InImpact) : ImpactLocation(InImpact) { }
+
 	FVector ImpactLocation;
 
 	void NetSerialize(FArchive& Ar)
@@ -31,10 +34,10 @@ struct FMockImpactCue
 		Ar << ImpactLocation;
 	}
 
-	static bool Unique(const FMockImpactCue& A, const FMockImpactCue& B)
+	bool NetUnique(const FMockImpactCue& Other) const
 	{
 		const float ErrorTolerance = 1.f;
-		return !A.ImpactLocation.Equals(B.ImpactLocation, ErrorTolerance);
+		return !ImpactLocation.Equals(Other.ImpactLocation, ErrorTolerance);
 	}
 
 };
@@ -46,6 +49,10 @@ struct FMockDamageCue
 {
 	NETSIMCUE_BODY();
 
+	FMockDamageCue() = default;
+	FMockDamageCue(const int32 InSourceID, const int32 InDamageType, const FVector& InHitLocation) 
+		: SourceID(InSourceID), DamageType(InDamageType), HitLocation(InHitLocation) { }
+
 	int32 SourceID;
 	int32 DamageType;
 	FVector HitLocation;
@@ -57,10 +64,10 @@ struct FMockDamageCue
 		Ar << HitLocation;
 	}
 
-	static bool Unique(const FMockDamageCue& A, const FMockDamageCue& B)
+	bool NetUnique(const FMockDamageCue& Other) const
 	{
 		const float ErrorTolerance = 1.f;
-		return A.SourceID != B.SourceID || A.DamageType != B.DamageType || !A.HitLocation.Equals(B.HitLocation, ErrorTolerance);
+		return SourceID != Other.SourceID || DamageType != Other.DamageType || !HitLocation.Equals(Other.HitLocation, ErrorTolerance);
 	}
 };
 NETSIMCUE_REGISTER(FMockDamageCue, TEXT("Damage"));
@@ -68,6 +75,10 @@ NETSIMCUE_REGISTER(FMockDamageCue, TEXT("Damage"));
 struct FMockHealingCue
 {
 	NETSIMCUE_BODY();
+	
+	FMockHealingCue() = default;
+	FMockHealingCue(const int32 InSourceID, const int32 InHealingAmount) 
+		: SourceID(InSourceID), HealingAmount(InHealingAmount) { }
 
 	int32 SourceID;
 	int32 HealingAmount;
@@ -78,10 +89,10 @@ struct FMockHealingCue
 		Ar << HealingAmount;
 	}
 
-	static bool Unique(const FMockHealingCue& A, const FMockHealingCue& B)
+	bool NetUnique(const FMockHealingCue& Other) const
 	{
 		const float ErrorTolerance = 1.f;
-		return A.SourceID != B.SourceID || A.HealingAmount != B.HealingAmount;
+		return SourceID != Other.SourceID || HealingAmount != Other.HealingAmount;
 	}
 };
 NETSIMCUE_REGISTER(FMockHealingCue, TEXT("Healing"));
@@ -100,7 +111,7 @@ public:
 		DispatchTable.template RegisterType<FMockImpactCue>();
 	}
 
-	void HandleCue(FMockImpactCue& ImpactCue, const FNetworkSimTime& ElapsedTime)
+	void HandleCue(FMockImpactCue& ImpactCue, const FNetSimCueSystemParamemters& SystemParameters)
 	{
 		UE_LOG(LogTemp, Warning, TEXT("Impact!"));
 	}
@@ -120,12 +131,12 @@ public:
 
 	using TestHandlerBase::HandleCue; // Required for us to "use" the HandleCue methods in our parent class
 
-	void HandleCue(FMockDamageCue& DamageCue, const FNetworkSimTime& ElapsedTime)
+	void HandleCue(FMockDamageCue& DamageCue, const FNetSimCueSystemParamemters& SystemParameters)
 	{
 		UE_LOG(LogTemp, Warning, TEXT("HandleCue: Damage. SourceID: %d DamageType: %d HitLocation: %s"), DamageCue.SourceID, DamageCue.DamageType, *DamageCue.HitLocation.ToString());
 	}
 
-	void HandleCue(FMockHealingCue& HealingCue, const FNetworkSimTime& ElapsedTime)
+	void HandleCue(FMockHealingCue& HealingCue, const FNetSimCueSystemParamemters& SystemParameters)
 	{
 		UE_LOG(LogTemp, Warning, TEXT("HandleCue: Healing. SourceID: %d Healing Amount: %d"), HealingCue.SourceID, HealingCue.HealingAmount);
 	}
@@ -141,24 +152,24 @@ void TestCues()
 
 	struct FBaseSimulation
 	{
-		static void TickSimulation(TNetSimCueDispatcher& Dispatcher)
+		static void TickSimulation(FNetSimCueDispatcher& Dispatcher)
 		{
-			Dispatcher.Invoke<FMockImpactCue>( { FVector(1.f, 2.f, 3.f) } );
+			Dispatcher.Invoke<FMockImpactCue>(FVector(1.f, 2.f, 3.f));
 		}
 	};
 
 	struct FChildSimulation
 	{
-		static void TickSimulation(TNetSimCueDispatcher& Dispatcher)
+		static void TickSimulation(FNetSimCueDispatcher& Dispatcher)
 		{
 			FBaseSimulation::TickSimulation(Dispatcher);
 
-			Dispatcher.Invoke<FMockDamageCue>( { 1, 2, FVector(1.f, 2.f, 3.f) } );
-			Dispatcher.Invoke<FMockHealingCue>( { 10, 32 } );
+			Dispatcher.Invoke<FMockDamageCue>(1, 2, FVector(1.f, 2.f, 3.f));
+			Dispatcher.Invoke<FMockHealingCue>(10, 32);
 		}
 	};
 
-	TNetSimCueDispatcher ServerDispatcher;
+	TNetSimCueDispatcher<void> ServerDispatcher;
 	FChildSimulation::TickSimulation(ServerDispatcher);
 
 	// -------------------------------------
@@ -173,7 +184,7 @@ void TestCues()
 	// -------------------------------------
 
 	FNetBitReader TempReader(nullptr, TempWriter.GetData(), TempWriter.GetNumBits());
-	TNetSimCueDispatcher ClientDispatcher;
+	TNetSimCueDispatcher<void> ClientDispatcher;
 	ClientDispatcher.NetSerializeSavedCues(TempReader, false);
 
 	TestHandlerChild MyObject;
diff --git a/Engine/Plugins/Runtime/NetworkPrediction/Source/NetworkPrediction/Public/NetworkedSimulationModelCues.h b/Engine/Plugins/Runtime/NetworkPrediction/Source/NetworkPrediction/Public/NetworkedSimulationModelCues.h
--- a/Engine/Plugins/Runtime/NetworkPrediction/Source/NetworkPrediction/Public/NetworkedSimulationModelCues.h
+++ b/Engine/Plugins/Runtime/NetworkPrediction/Source/NetworkPrediction/Public/NetworkedSimulationModelCues.h
@@ -30,44 +30,55 @@ enum class ENetSimCueReplicationTarget : uint8
 };
 
 // ----------------------------------------------------------------------------------------------------------------------------------------------
-//	NetSimCue Traits
+//	NetSimCue Traits: how to configure how a NetSimCue will dispatch within the NetworkedSimulationModel
 // ----------------------------------------------------------------------------------------------------------------------------------------------
 
 struct TNetSimCueTraitsBase
 {
-	// Who can Invoke this Cue in their simulation
+	// Who can Invoke this Cue in their simulation (if this test fails, the Invoke call is supressed locally)
 	static constexpr uint8 InvokeMask { (uint8)ESimulationTickContext::All };
 
-	// Does the cue replicate? (from authority)
+	// Does the cue replicate? (from authority). This will also determine if the cue needs to be saved locally for NetUnique tests (to avoid double playing)
 	static constexpr ENetSimCueReplicationTarget ReplicationTarget { ENetSimCueReplicationTarget::All };
+
+	// Does the cue support rollback callbacks (in contexts where applicable. e.g, when we might get rolled back (!Authority))
+	static constexpr bool Rollbackable { true };
 };
 
 // Preset: non replicated cue that only plays during "latest" simulate. Will not be played during rewind/resimulate.
+// Lightest weight cue. Best used for cosmetic, non critical events. Footsteps, impact effects, etc.
 struct TNetSimCueTraits_Weak
 {
 	static constexpr uint8 InvokeMask { (uint8)ESimulationTickContext::Authority | (uint8)ESimulationTickContext::Predict };
 	static constexpr ENetSimCueReplicationTarget ReplicationTarget { ENetSimCueReplicationTarget::None };
+	static constexpr bool Rollbackable { false };
 };
 
 // Preset: Replicated, non predicted. Only invoked on authority and will replicate to everyone else.
+// Best for events that are critical that cannot be rolled back/undown and do not need to be predicted.
 struct TNetSimCueTraits_ReplicatedNonPredicted
 {
 	static constexpr uint8 InvokeMask { (uint8)ESimulationTickContext::Authority };
 	static constexpr ENetSimCueReplicationTarget ReplicationTarget { ENetSimCueReplicationTarget::All };
+	static constexpr bool Rollbackable { false };
 };
 
-// Preset:  Replicated to simulated proxy, predicted by autonomous proxy
+// Preset: Replicated to interpolating proxies, predicted by autonomous proxy
+// Best for events you want everyone to see but don't need to get perfect in the predicting cases: doesn't need to rollback and cheap on cpu (no unique tests on predicted path)
 struct TNetSimCueTraits_ReplicatedXOrPredicted
 {
 	static constexpr uint8 InvokeMask { (uint8)ESimulationTickContext::Authority | (uint8)ESimulationTickContext::Predict };
 	static constexpr ENetSimCueReplicationTarget ReplicationTarget { ENetSimCueReplicationTarget::Interpolators };
+	static constexpr bool Rollbackable { false };
 };
 
-// Preset: Invoked and replicated to all. Uniqueness testing to avoid double playing etc
+// Preset: Invoked and replicated to all. Uniqueness testing to avoid double playing, rollbackable so that it can (re)play during resimulates
+// Most expensive (bandwidth and CPU for uniqueness testing) and requires rollback callbacks to be implemented to be correct
 struct TNetSimCueTraits_Strong
 {
 	static constexpr uint8 InvokeMask { (uint8)ESimulationTickContext::All };
 	static constexpr ENetSimCueReplicationTarget ReplicationTarget { ENetSimCueReplicationTarget::All };
+	static constexpr bool Rollbackable { true };
 };
 
 // Actual trait struct that we use to look up traits. User cues must specialize this
@@ -84,28 +95,31 @@ struct FNetSimCueWrapperBase
 {
 	virtual ~FNetSimCueWrapperBase() { }
 	virtual void NetSerialize(FArchive& Ar) = 0;
-	virtual bool Unique(const FNetSimCueWrapperBase* Other) const = 0;
+	virtual bool NetUnique(const FNetSimCueWrapperBase* Other) const = 0;
 	virtual void* CueData() const = 0;
 	virtual ENetSimCueReplicationTarget GetReplicationTarget() const = 0;
+	virtual bool Rollbackable() const = 0;
 };
 
 template<typename TCue>
 struct TNetSimCueWrapper : FNetSimCueWrapperBase
 {
 	TNetSimCueWrapper() = default;
-	TNetSimCueWrapper(const TCue& Source)
-	{
-		Instance = Source; // Fixme: can we use move semantics to avoid this copy?
-	}
+
+	template <typename... ArgsType>
+	TNetSimCueWrapper(ArgsType&&... Args)
+		: Instance( MoveTempIfPossible(Forward<ArgsType>(Args))... ) { }
 
 	void NetSerialize(FArchive& Ar) override final
 	{
+		// Cue types must implement NetSerialize(FArchive& Ar)
 		Instance.NetSerialize(Ar);
 	}
 
-	bool Unique(const FNetSimCueWrapperBase* Other) const override final
+	bool NetUnique(const FNetSimCueWrapperBase* Other) const override final
 	{
-		return TCue::Unique(Instance, *((TCue*)Other->CueData()));
+		// Cue types must implement bool NetUnique(const TMyCueType& Other) const
+		return Instance.NetUnique(*((TCue*)Other->CueData()));
 	}
 
 	void* CueData() const override final
@@ -118,9 +132,42 @@ struct TNetSimCueWrapper : FNetSimCueWrapperBase
 		return TCueHandlerTraits<TCue>::ReplicationTarget;
 	}
 
+	bool Rollbackable() const override final
+	{
+		return TCueHandlerTraits<TCue>::Rollbackable;
+	}
+
 	TCue Instance;
 };
 
+// ----------------------------------------------------------------------------------------------------------------------------------------------
+//	Callbacks
+// ----------------------------------------------------------------------------------------------------------------------------------------------
+
+struct FNetSimCueCallbacks
+{
+	/** Rollback: the keyframe this cue was invoked on was rolled back (and will be resimulated) */
+	DECLARE_MULTICAST_DELEGATE(FOnRollback)
+	FOnRollback	OnRollback;
+	
+	/** Confirmed: the keyframe this was was predictively invoked on has been confirmed. The frame will not be rolled back now. */
+	DECLARE_MULTICAST_DELEGATE(FOnConfirmed)
+	FOnConfirmed OnConfirmed;
+
+	void ClearAll() { OnRollback.Clear(); OnConfirmed.Clear(); }
+	bool IsBound() const { return OnRollback.IsBound() || OnConfirmed.IsBound(); }
+};
+
+/** System parameters for NetSimCue events */
+struct FNetSimCueSystemParamemters
+{
+	// How much simulation time has passed since this cue was invoked. This will be 0 in authority/predict contexts, but when invoked via replication this will tell you how long ago it happened, relative to local simulation time.
+	const FNetworkSimTime& TimeSinceInvocation;
+
+	// Callback structure if applicable. This will be null on non-rewindable cues as well as execution contexts where rollbacks wont happen (e.g, authority).
+	FNetSimCueCallbacks* Callbacks;	
+};
+
 // ----------------------------------------------------------------------------------------------------------------------------------------------
 // GlobalCueTypeTable: Cue types register with this to get a Type ID assigned. That ID is used in net serialization to talk about types.
 // ----------------------------------------------------------------------------------------------------------------------------------------------
@@ -212,14 +259,12 @@ struct FSavedCue
 
 	FSavedCue(FSavedCue&&) = default;
 	FSavedCue& operator=(FSavedCue&&) = default;
-
-	template<typename TCue>
-	FSavedCue(FNetworkSimTime InTime, const TCue& SourceCue) : Time(InTime)
+	
+	FSavedCue(const FNetSimCueTypeId& InId, const FNetworkSimTime& InTime, FNetSimCueWrapperBase* Cue, const bool& bInAllowRollback)
+		: ID(InId), Time(InTime), CueInstance(Cue), bAllowRollback(bInAllowRollback)
 	{
-		ID = TCue::ID; //-V570
-		CueInstance.Reset(new TNetSimCueWrapper<TCue>(SourceCue)); // Fixme: can the copy be avoided with move semantics?
 	}
-
+	
 	void NetSerialize(FArchive& Ar)
 	{
 		if (Ar.IsSaving())
@@ -240,14 +285,14 @@ struct FSavedCue
 		}
 	}
 
-	bool Unique(FSavedCue& OtherCue) const
+	bool NetUnique(FSavedCue& OtherCue) const
 	{
 		if (ID != OtherCue.ID)
 		{
 			return false;
 		}
 
-		return CueInstance->Unique(OtherCue.CueInstance.Get());
+		return CueInstance->NetUnique(OtherCue.CueInstance.Get());
 	}
 
 	FString GetTypeName() const
@@ -259,6 +304,7 @@ struct FSavedCue
 	FNetworkSimTime Time;
 	TUniquePtr<FNetSimCueWrapperBase> CueInstance;
 	bool bDispatched = false;
+	bool bAllowRollback = false;
 };
 
 
@@ -285,9 +331,9 @@ public:
 			FCueTypeInfo& CueTypeInfo = CueTypeInfoMap.Add(TCue::ID);
 
 			// The actual Dispatch func that gets invoked
-			CueTypeInfo.Dispatch = [](FNetSimCueWrapperBase* Cue, TCueHandler& Handler, const FNetworkSimTime& ElapsedTime)
+			CueTypeInfo.Dispatch = [](FNetSimCueWrapperBase* Cue, TCueHandler& Handler, const FNetSimCueSystemParamemters& SystemParameters)
 			{
-				Handler.HandleCue( *static_cast<TCue*>(Cue->CueData()), ElapsedTime );
+				Handler.HandleCue( *static_cast<TCue*>(Cue->CueData()), SystemParameters );
 			};
 		};
 		
@@ -301,12 +347,12 @@ public:
 		}
 	}
 
-	void Dispatch(FSavedCue& SavedCue, TCueHandler& Handler, const FNetworkSimTime& ElapsedTime)
+	void Dispatch(FSavedCue& SavedCue, TCueHandler& Handler, const FNetSimCueSystemParamemters& SystemParameters)
 	{
 		if (FCueTypeInfo* TypeInfo = CueTypeInfoMap.Find(SavedCue.ID))
 		{
 			check(TypeInfo->Dispatch);
-			TypeInfo->Dispatch(SavedCue.CueInstance.Get(), Handler, ElapsedTime);
+			TypeInfo->Dispatch(SavedCue.CueInstance.Get(), Handler, SystemParameters);
 		}
 	}
 
@@ -316,7 +362,7 @@ private:
 
 	struct FCueTypeInfo
 	{
-		TFunction<void(FNetSimCueWrapperBase* Cue, TCueHandler& Handler, const FNetworkSimTime& ElapsedTime)> Dispatch;
+		TFunction<void(FNetSimCueWrapperBase* Cue, TCueHandler& Handler, const FNetSimCueSystemParamemters& SystemParameters)> Dispatch;
 	};
 
 	TMap<FNetSimCueTypeId, FCueTypeInfo> CueTypeInfoMap;
@@ -334,8 +380,17 @@ TCueDispatchTable<TCueHandler> TCueDispatchTable<TCueHandler>::Singleton;
 // Non-templated, "networking model independent" base: this is what the pure simulation code gets to invoke cues. 
 struct FNetSimCueDispatcher
 {
-	template<typename T>
-	void Invoke(T&& Cue)
+	// Invoke - this is how to invoke a cue from simulation code. This will construct the CueType T emplace in the saved cue record.
+	// 
+	// Best way to call:
+	//	Invoke<FMyCue>(a, b, c); // a, b, c are constructor parameters
+	//
+	// This works too, but will cause a move (if possible) or copy
+	//	FMyCue MyCue(a,b,c);
+	//	Invoke<FMyCue>(MyCue);	
+
+	template<typename T, typename... ArgsType>
+	void Invoke(ArgsType&&... Args)
 	{
 		if (EnsureValidContext())
 		{
@@ -344,22 +399,20 @@ struct FNetSimCueDispatcher
 				// There is an implicit contract that we will invoke events going forward OR will receive explicit rollback notification (TODO)
 				ensure(SavedCues.Num() == 0 || SavedCues.Last().Time <= Context.CurrentSimTime);
 
+				// Whether this cue should be dispatched with rollback callbacks. This is a trait of the cue type + non authority (authority will never rollback)
+				const bool bAllowRollback = TCueHandlerTraits<T>::Rollbackable && Context.TickContext != ESimulationTickContext::Authority; 
 
-				bool bTransient = false;
-				if (Context.TickContext == ESimulationTickContext::Authority && TCueHandlerTraits<T>::ReplicationTarget == ENetSimCueReplicationTarget::None)
-				{
-					bTransient = true;
-				}
+				// Whether we go in the transient list. Transient cues are dumped after dispatching (not saved over multiple frames for uniqueness comparisons)
+				const bool bTransient = TCueHandlerTraits<T>::ReplicationTarget == ENetSimCueReplicationTarget::None;
 
 				UE_LOG(LogNetSimCues, Log, TEXT("Invoking Cue %s. Transient: %d. Mask: %d. ReplicationTarget: %d"), *FGlobalCueTypeTable::Get().GetTypeName(T::ID), bTransient, TCueHandlerTraits<T>::InvokeMask, TCueHandlerTraits<T>::ReplicationTarget);
-
 				if (bTransient)
 				{
-					TransientCues.Emplace(Context.CurrentSimTime, Cue);
+					TransientCues.Emplace(T::ID, Context.CurrentSimTime, new TNetSimCueWrapper<T>(Forward<ArgsType>(Args)...), bAllowRollback);
 				}
 				else
 				{
-					SavedCues.Emplace(Context.CurrentSimTime, Cue);
+					SavedCues.Emplace(T::ID, Context.CurrentSimTime, new TNetSimCueWrapper<T>(Forward<ArgsType>(Args)...), bAllowRollback);
 				}
 			}
 			else
@@ -367,7 +420,7 @@ struct FNetSimCueDispatcher
 				UE_LOG(LogNetSimCues, Log, TEXT("Suppressing Cue Invocation %s. Mask: %d. TickContext: %d"), *FGlobalCueTypeTable::Get().GetTypeName(T::ID), TCueHandlerTraits<T>::InvokeMask, (int32)Context.TickContext);
 			}
 		}
-	}	
+	}
 
 protected:
 
@@ -417,6 +470,12 @@ struct TNetSimCueDispatcher : public FNetSimCueDispatcher
 		}
 		else
 		{
+			// This is quite inefficient right now. 
+			//	-We are replicating cues in the last X seconds (ReplicationWindow) redundantly
+			//	-Client has to deserialize them (+ heap allocation) and check for uniqueness (have they already processed)
+			//	-If already processed (quite common), they are thrown out.
+			//	-Would be better if we maybe serialized "net hash" and could skip ahead in the bunch of already processed
+
 			int32 StartingNum = SavedCues.Num();
 
 			for (int32 CueIdx=0; CueIdx < NumCues; ++CueIdx)
@@ -437,7 +496,7 @@ struct TNetSimCueDispatcher : public FNetSimCueDispatcher
 				bool bUniqueCue = true;
 				for (int32 ExistingIdx=0; ExistingIdx < StartingNum; ++ExistingIdx)
 				{
-					if (SerializedCue.Unique(SavedCues[ExistingIdx]) == false)
+					if (SerializedCue.NetUnique(SavedCues[ExistingIdx]) == false)
 					{
 						// These cues are not unique ("close enough") so we are skipping receiving this one
 						UE_LOG(LogNetSimCues, Log, TEXT("Discarding replicated NetSimCue %s because we've already processed it."), *SerializedCue.GetTypeName());
@@ -458,48 +517,106 @@ struct TNetSimCueDispatcher : public FNetSimCueDispatcher
 	template<typename T>
 	void DispatchCueRecord(T& Handler, FNetworkSimTime CurrentSimTime)
 	{
-		const FNetworkSimTime PruneTime = (MinPruneNetSimTime.IsPositive() ? MinPruneNetSimTime : CurrentSimTime) - TCueDispatcherTraits<Model>::ReplicationWindow;
+		const FNetworkSimTime& ConfirmedTime = UserConfirmedTime.IsPositive() ? UserConfirmedTime : CurrentSimTime;
+		const FNetworkSimTime SavedCuePruneTime = ConfirmedTime - TCueDispatcherTraits<Model>::ReplicationWindow;
+
+		const FNetworkSimTime& DispatchTime = UserMaxDispatchTime.IsPositive() ? UserMaxDispatchTime : CurrentSimTime;
+		
+		int32 SavedCuePruneIdx = -1;
+
+		// ------------------------------------------------------------------------
+		// User callbacks
+		// ------------------------------------------------------------------------
+		
+		ensure(UserConfirmedTime.IsPositive() || InvocationCallbacks.Num() == 0); // Should not have callbacks when PruneTime is zero
+
+		for (FInvocationCallbackContainer& Callback : InvocationCallbacks)
+		{
+			if (RollbackTime.IsPositive() && Callback.InvocationTime > RollbackTime)
+			{
+				Callback.UserCallbacks.OnRollback.Broadcast();
+				Callback.UserCallbacks.OnRollback.Clear();
+			}
+
+			if (Callback.UserCallbacks.OnConfirmed.IsBound() && Callback.InvocationTime <= ConfirmedTime)
+			{
+				Callback.UserCallbacks.OnConfirmed.Broadcast();
+				Callback.UserCallbacks.ClearAll(); // Confirmed is the end of the road. Clear all and removal will happen at the bottom of this function
+			}
+		}
+		RollbackTime.Reset();
 
-		// Dispatch
-		int32 PruneIdx = -1;
+		// ------------------------------------------------------------------------
+		// Dispatch (call ::HandleCue)
+		// ------------------------------------------------------------------------
+		
 		for (int32 SavedCueIdx = 0; SavedCueIdx < SavedCues.Num(); ++ SavedCueIdx)
 		{
 			FSavedCue& SavedCue = SavedCues[SavedCueIdx];
-			if ( SavedCue.Time <= PruneTime)
+			if ( SavedCue.Time <= SavedCuePruneTime)
 			{
-				PruneIdx = SavedCueIdx;
+				SavedCuePruneIdx = SavedCueIdx;
 			}
 
-			const bool bWithold = MaxDispatchNetSimTime.IsPositive() && SavedCue.Time > MaxDispatchNetSimTime;
-			UE_CLOG(bWithold, LogNetSimCues, Log, TEXT("Withholding Cue %s. %s > %s"), *SavedCue.GetTypeName(), *SavedCue.Time.ToString(), *MaxDispatchNetSimTime.ToString());
+			const bool bWithold = SavedCue.Time > DispatchTime;
+			UE_CLOG(bWithold, LogNetSimCues, Log, TEXT("Withholding Cue %s. %s > %s"), *SavedCue.GetTypeName(), *SavedCue.Time.ToString(), *DispatchTime.ToString());
 
 			if (SavedCue.bDispatched == false && !bWithold)
 			{
 				UE_LOG(LogNetSimCues, Log, TEXT("Dispatching NetSimCue %s."), *SavedCue.GetTypeName());
 				SavedCue.bDispatched = true;
-				TCueDispatchTable<T>::Get().Dispatch(SavedCue, Handler, CurrentSimTime - SavedCue.Time);
+				TCueDispatchTable<T>::Get().Dispatch(SavedCue, Handler, {DispatchTime - SavedCue.Time, GetUserCallbackPtr(SavedCue)} );
 			}
 		}
 
 		for (FSavedCue& TransientCue : TransientCues)
 		{
 			UE_LOG(LogNetSimCues, Log, TEXT("Dispatching transient NetSimCue %s."), *TransientCue.GetTypeName());
-			TCueDispatchTable<T>::Get().Dispatch(TransientCue, Handler, CurrentSimTime - TransientCue.Time);
+			TCueDispatchTable<T>::Get().Dispatch(TransientCue, Handler, {DispatchTime - TransientCue.Time, GetUserCallbackPtr(TransientCue)} );
 		}
 		TransientCues.Reset();
 
+		// ------------------------------------------------------------------------
 		// Prune
-		if (PruneIdx >= 0)
+		// ------------------------------------------------------------------------
+		
+		// Remove Cues we know longer need to keep around
+		if (SavedCuePruneIdx >= 0)
 		{
 #if !(UE_BUILD_SHIPPING || UE_BUILD_TEST)
-			for (int32 i=0; i <= PruneIdx; ++i)
+			for (int32 i=0; i <= SavedCuePruneIdx; ++i)
 			{
 				UE_CLOG(!SavedCues[i].bDispatched, LogNetSimCues, Warning, TEXT("Non-Dispatched Cue is about to be pruned! %s. %s"), *SavedCues[i].GetTypeName());
-				UE_LOG(LogNetSimCues, Log, TEXT("Pruning Cue %s. Invoke Time: %s. Current Time: %s"), *SavedCues[i].GetTypeName(), *SavedCues[i].Time.ToString(), *CurrentSimTime.ToString());
+				UE_LOG(LogNetSimCues, Log, TEXT("Pruning Cue %s. Invoke Time: %s. Current Time: %s"), *SavedCues[i].GetTypeName(), *SavedCues[i].Time.ToString(), *DispatchTime.ToString());
 			}
 #endif
 
-			SavedCues.RemoveAt(0, PruneIdx+1, false);
+			SavedCues.RemoveAt(0, SavedCuePruneIdx+1, false);
+		}
+
+		// Remove any unbounded callbacks
+		for (int32 idx=InvocationCallbacks.Num()-1; idx >= 0; --idx)
+		{
+			FInvocationCallbackContainer& Callback = InvocationCallbacks[idx];
+			if (!Callback.UserCallbacks.IsBound())
+			{
+				UE_LOG(LogNetSimCues, Log, TEXT("Pruned Callbacks at time=%s"), *Callback.InvocationTime.ToString());
+				InvocationCallbacks.RemoveAt(idx, 1, false);
+			}
+		}
+	}
+
+	// Tell dispatcher that we've rolled back to a new simulation time (resimulate steps to follow, most likely)
+	void NotifyRollback(const FNetworkSimTime& InRollbackTime)
+	{
+		// Just cache off the notify. We want to invoke the callbacks in DispatchCueRecord, not right now (in the middle of simulation tick/reconcile)
+		if (RollbackTime.IsPositive() == false)
+		{
+			RollbackTime = InRollbackTime;
+		}
+		else
+		{
+			ensure(RollbackTime < InRollbackTime);
 		}
 	}
 
@@ -507,17 +624,60 @@ struct TNetSimCueDispatcher : public FNetSimCueDispatcher
 	void PushContext(const FContext& InContext) { Context = InContext; }
 	void PopContext() { Context = FContext(); }
 
-	// Set max simulation time to invoke cues for. Clearing = always process latest
-	void SetMaxDispatchTime(const FNetworkSimTime& NewMaxTime) const { MaxDispatchNetSimTime = NewMaxTime; }
-	void ClearMaxDispatchTime() const { MaxDispatchNetSimTime = FNetworkSimTime(); }
+	// Set max simulation time to invoke cues for. Clearing = always process latest. This is used for interpolation or client side buffering delaying cues until "its time"
+	void SetMaxDispatchTime(const FNetworkSimTime& NewMaxTime) const { UserMaxDispatchTime = NewMaxTime; }
+	void ClearMaxDispatchTime() const { UserMaxDispatchTime = FNetworkSimTime(); }
 
-	// Set minimum prune time. System will not prune cues that happen after this time. ReplicationWindow is applied *on top* of this
-	void SetPruneTime(const FNetworkSimTime& NewPruneTime) const { MinPruneNetSimTime = NewPruneTime; }
+	// Set Confirmed time. If this is never called, we assume we are always confirmed/authority.
+	void SetConfirmedTime(const FNetworkSimTime& NewConfirmedTime) const { UserConfirmedTime = NewConfirmedTime; }
 
 private:
 
-	mutable FNetworkSimTime MaxDispatchNetSimTime;	// (If set) Max time of a saved cue that can be dispatched. Needed for delaying interpolation / client side buffering.
-	mutable FNetworkSimTime MinPruneNetSimTime;		// (If set) Min time we start at to look for cues to prune. Note ReplicationWindow is always applied on top of this.
+	mutable FNetworkSimTime UserMaxDispatchTime;	// (If set) Max time of a saved cue that can be dispatched. Needed for delaying interpolation / client side buffering.
+	mutable FNetworkSimTime UserConfirmedTime;		// (If set) latest confirmed simulation time. If not set, we assume we are authority and will never rollback.
+	
+	struct FInvocationCallbackContainer
+	{
+		FInvocationCallbackContainer(const FNetworkSimTime& InTime) : InvocationTime(InTime) { }
+
+		FNetworkSimTime InvocationTime;
+		FNetSimCueCallbacks UserCallbacks;
+	};
+
+	TArray<FInvocationCallbackContainer> InvocationCallbacks;
+	FNetworkSimTime RollbackTime;
+
+	FNetSimCueCallbacks* GetUserCallbackPtr(const FSavedCue& SavedCue)
+	{
+		FNetSimCueCallbacks* CallbackPtr = nullptr;
+		if (SavedCue.bAllowRollback)
+		{
+			for (int32 CallbacksIdx = InvocationCallbacks.Num()-1; CallbacksIdx >=0; --CallbacksIdx)
+			{
+				FInvocationCallbackContainer& CallbackContainer = InvocationCallbacks[CallbacksIdx];
+
+				if (CallbackContainer.InvocationTime == SavedCue.Time)
+				{
+					CallbackPtr = &CallbackContainer.UserCallbacks;
+					break;
+				}
+				if (CallbackContainer.InvocationTime < SavedCue.Time)
+				{
+					UE_LOG(LogNetSimCues, Log, TEXT("Added Callbacks for time=%s at idx=%d"), *SavedCue.Time.ToString(), CallbacksIdx);
+					CallbackPtr = &InvocationCallbacks.EmplaceAt_GetRef( CallbacksIdx+1, SavedCue.Time ).UserCallbacks;
+					break;
+				}
+			}
+
+			if (!CallbackPtr)
+			{
+				UE_LOG(LogNetSimCues, Log, TEXT("Added Callbacks for time=%s"));
+				CallbackPtr = &InvocationCallbacks.Emplace_GetRef( SavedCue.Time ).UserCallbacks;
+			}
+
+		}
+		return CallbackPtr;
+	}
 };
 
 // ------------------------------------------------------------------------------------------------------------------------------------------------------------------------------------------------
diff --git a/Engine/Plugins/Runtime/NetworkPrediction/Source/NetworkPrediction/Public/NetworkedSimulationModelReplicators.h b/Engine/Plugins/Runtime/NetworkPrediction/Source/NetworkPrediction/Public/NetworkedSimulationModelReplicators.h
--- a/Engine/Plugins/Runtime/NetworkPrediction/Source/NetworkPrediction/Public/NetworkedSimulationModelReplicators.h
+++ b/Engine/Plugins/Runtime/NetworkPrediction/Source/NetworkPrediction/Public/NetworkedSimulationModelReplicators.h
@@ -447,6 +447,9 @@ struct TRepController_Simulated : public TBase
 
 		if (GetSimulatedUpdateMode() == ESimulatedUpdateMode::Extrapolate && NetworkSimulationModelCVars::EnableSimulatedReconcile())
 		{
+			// This is effectively a rollback: we went back in time during NetSerialize and are now going to catch up ReconcileSimulationTime ms.
+			Buffers.CueDispatcher.NotifyRollback(Ticker.GetTotalProcessedSimulationTime());
+
 			// Simulated Reconcile requires the input buffer to be kept up to date with the Sync buffer
 			// Generate a new, fake, command since we just added a new sync state to head
 			while (Buffers.Input.HeadFrame() < Buffers.Sync.HeadFrame())
@@ -490,7 +493,7 @@ struct TRepController_Simulated : public TBase
 			const FNetworkSimTime InterpolationSimTime = FNetworkSimTime::FromRealTimeSeconds(InterpolationRealTime);
 
 			Buffers.CueDispatcher.SetMaxDispatchTime(InterpolationSimTime);
-			Buffers.CueDispatcher.SetPruneTime(InterpolationSimTime);
+			Buffers.CueDispatcher.SetConfirmedTime(InterpolationSimTime);
 		}
 		else
 		{
@@ -501,7 +504,7 @@ struct TRepController_Simulated : public TBase
 			}
 
 			Buffers.CueDispatcher.ClearMaxDispatchTime();
-			Buffers.CueDispatcher.SetPruneTime(LastSerializedSimulationTime);
+			Buffers.CueDispatcher.SetConfirmedTime(LastSerializedSimulationTime);
 		}
 	}
 	
@@ -521,6 +524,8 @@ struct TRepController_Simulated : public TBase
 
 		Driver->FinalizeFrame(LastSerializedSyncState, LastSerializedAuxState);
 
+		Buffers.CueDispatcher.NotifyRollback(LastSerializedSimulationTime);
+
 		if (NETSIM_MODEL_DEBUG)
 		{
 			FVisualLoggingParameters VLogParams(EVisualLoggingContext::FirstMispredicted, ParentFrame, EVisualLoggingLifetime::Persistent, TEXT("DependentRollbackBegin"));
@@ -733,6 +738,9 @@ struct TRepController_Autonomous: public TBase
 		{
 			DependentSim->BeginRollback(RollbackDeltaTime, SerializedFrame);
 		}
+
+		// Tell cue dispatch that we've rolledback as well
+		Buffers.CueDispatcher.NotifyRollback(SerializedTime);
 		
 		// Resimulate all user commands 
 		const int32 LastFrameToProcess = Ticker.MaxAllowedFrame;
@@ -821,7 +829,7 @@ struct TRepController_Autonomous: public TBase
 		TBase::PostSimTick(Driver, Buffers, Ticker, TickParameters);
 
 		Buffers.CueDispatcher.ClearMaxDispatchTime();
-		Buffers.CueDispatcher.SetPruneTime(SerializedTime);
+		Buffers.CueDispatcher.SetConfirmedTime(SerializedTime);
 	}
 	
 
diff --git a/Engine/Plugins/Runtime/NetworkPrediction/readme.txt b/Engine/Plugins/Runtime/NetworkPrediction/readme.txt
--- a/Engine/Plugins/Runtime/NetworkPrediction/readme.txt
+++ b/Engine/Plugins/Runtime/NetworkPrediction/readme.txt
@@ -76,9 +76,7 @@ High level focus:
 -We intend to support RootMotion (Animation) / RootMotionSources (non anim) in the new movement system.
 
 Road Map:
--Event System / "Game API"
--Ability System integration
--New Movement System
+-New Movement System / Ability System integration
  
  
 // ----------------------------------------------------------------------------------------------------------
@@ -108,6 +106,10 @@ Smoothing: Taking the output of the simulation and applying an additional layer
 // Release notes
 // ----------------------------------------------------------------------------------------------------------
 
+Update (11-21-19)
+-Rewindable NetSimCue support added. This just about wraps up the initial implemention of NetSimCues and the "Game API / Event System".
+-Still some cleanup and debugging/logging stuff to do but will be transitioning into porting the legacy CMC movement system over soon.
+
 Update (11-15-19)
 -NetSimCues initial check in
 -These are the async/deferred events that are outputted during the network simulation tick
diff --git a/Engine/Plugins/Runtime/NetworkPredictionExtras/Source/NetworkPredictionExtras/Private/MockAbilitySimulation.cpp b/Engine/Plugins/Runtime/NetworkPredictionExtras/Source/NetworkPredictionExtras/Private/MockAbilitySimulation.cpp
--- a/Engine/Plugins/Runtime/NetworkPredictionExtras/Source/NetworkPredictionExtras/Private/MockAbilitySimulation.cpp
+++ b/Engine/Plugins/Runtime/NetworkPredictionExtras/Source/NetworkPredictionExtras/Private/MockAbilitySimulation.cpp
@@ -15,9 +15,17 @@ namespace MockAbilityCVars
 
 	NETSIM_DEVCVAR_SHIPCONST_FLOAT(DashMaxSpeed, 7500.f, "mockability.DashMaxSpeed", "Max Speed when dashing.");
 	NETSIM_DEVCVAR_SHIPCONST_FLOAT(DashAcceleration, 100000.f, "mockability.DashAcceleration", "Acceleration when dashing.");
+
+
+	NETSIM_DEVCVAR_SHIPCONST_INT(BlinkCueType, 0, "mockability.BlinkCueType", "0=default. 1=weak. 2=");
 }
 
 NETSIMCUE_REGISTER(FMockAbilityBlinkCue, TEXT("MockAbilityBlink"));
+NETSIMCUE_REGISTER(FMockAbilityBlinkCue_Weak, TEXT("MockAbilityBlink_Weak"));
+NETSIMCUE_REGISTER(FMockAbilityBlinkCue_ReplicatedNonPredicted, TEXT("MockAbilityBlink_RepNonPredicted"));
+NETSIMCUE_REGISTER(FMockAbilityBlinkCue_ReplicatedXOrPredicted, TEXT("MockAbilityBlink_RepXOrPredicted"));
+NETSIMCUE_REGISTER(FMockAbilityBlinkCue_Strong, TEXT("MockAbilityBlink_Strong"));
+
 NETSIMCUESET_REGISTER(UMockFlyingAbilityComponent, FMockAbilityCueSet);
 
 // -------------------------------------------------------------------------------------------------------------
@@ -119,7 +127,29 @@ void FMockAbilitySimulation::SimulationTick(const FNetSimTimeStep& TimeStep, con
 
 				// Invoke a NetCue for the blink. This is essentially capturing the start/end location so that all receivers of the event
 				// get the exact coordinates (maybe overkill in practice but key is that we have data that we want to pass out via an event)
-				Output.CueDispatch.Invoke<FMockAbilityBlinkCue>({Input.Sync.Location, DestLocation});
+
+				UE_LOG(LogTemp, Warning, TEXT("INVOKING CUE EMPLACE"));
+				switch(MockAbilityCVars::BlinkCueType()) // Only for dev/testing. Not a normal setup.
+				{
+				case 0:
+					Output.CueDispatch.Invoke<FMockAbilityBlinkCue>( Input.Sync.Location, DestLocation );
+					break;
+				case 1:
+					Output.CueDispatch.Invoke<FMockAbilityBlinkCue_Weak>( Input.Sync.Location, DestLocation );
+					break;
+				case 2:
+					Output.CueDispatch.Invoke<FMockAbilityBlinkCue_ReplicatedNonPredicted>( Input.Sync.Location, DestLocation );
+					break;
+
+				case 3:
+					Output.CueDispatch.Invoke<FMockAbilityBlinkCue_ReplicatedXOrPredicted>( Input.Sync.Location, DestLocation );
+					break;
+
+				case 4:
+					Output.CueDispatch.Invoke<FMockAbilityBlinkCue_Strong>( Input.Sync.Location, DestLocation );
+					break;
+
+				};
 			}
 		}
 	}
@@ -294,16 +324,33 @@ static float BlinkCueDuration = 1.f;
 static FAutoConsoleVariableRef CVarBindAutomatically(TEXT("NetworkPredictionExtras.FlyingPawn.BlinkCueDuration"),
 	BlinkCueDuration, TEXT("Duration of BlinkCue"), ECVF_Default);
 
-void UMockFlyingAbilityComponent::HandleCue(FMockAbilityBlinkCue& BlinkCue, const FNetworkSimTime& ElapsedTime)
+void UMockFlyingAbilityComponent::HandleCue(FMockAbilityBlinkCue& BlinkCue, const FNetSimCueSystemParamemters& SystemParameters)
 {
 	FString RoleStr = GetOwnerRole() == ROLE_Authority ? TEXT("Server") : TEXT("Client");
 
 	FVector Delta = GetOwner()->GetActorLocation() - BlinkCue.StopLocation;
 
 	UE_LOG(LogTemp, Display, TEXT("[%s] BlinkCue! : <%f, %f, %f> - <%f, %f, %f>. ElapsedTime: %s. Delta: %.3f"), *RoleStr, BlinkCue.StartLocation.X, BlinkCue.StartLocation.Y, BlinkCue.StartLocation.Z,
-		BlinkCue.StopLocation.X, BlinkCue.StopLocation.Y, BlinkCue.StopLocation.Z, *ElapsedTime.ToString(), Delta.Size()); //*BlinkCue.StartLocation.ToString(), *BlinkCue.StopLocation.ToString());
+		BlinkCue.StopLocation.X, BlinkCue.StopLocation.Y, BlinkCue.StopLocation.Z, *SystemParameters.TimeSinceInvocation.ToString(), Delta.Size()); //*BlinkCue.StartLocation.ToString(), *BlinkCue.StopLocation.ToString());
 
 	// Crude compensation for cue firing in the past (note this is not necessary! Some cues not care and need to see the "full" effect regardless of when it happened)
-	float Duration = BlinkCueDuration - ElapsedTime.ToRealTimeSeconds();
+	float Duration = FMath::Max<float>(0.1f, BlinkCueDuration - SystemParameters.TimeSinceInvocation.ToRealTimeSeconds());
 	DrawDebugLine(GetWorld(), BlinkCue.StartLocation, BlinkCue.StopLocation, FColor::Red, false, Duration);
+
+
+	if (SystemParameters.Callbacks)
+	{
+		UE_LOG(LogTemp, Display, TEXT("  System Callbacks available!"));
+
+		SystemParameters.Callbacks->OnRollback.AddLambda([RoleStr]()
+		{
+			UE_LOG(LogTemp, Display, TEXT("  %s BlinkCue Rollback!"), *RoleStr);
+		});
+
+		SystemParameters.Callbacks->OnConfirmed.AddLambda([RoleStr]()
+		{
+			UE_LOG(LogTemp, Display, TEXT("  %s BlinkCue Confirmed!"), *RoleStr);
+		});
+
+	}
 }
\ No newline at end of file
diff --git a/Engine/Plugins/Runtime/NetworkPredictionExtras/Source/NetworkPredictionExtras/Public/MockAbilitySimulation.h b/Engine/Plugins/Runtime/NetworkPredictionExtras/Source/NetworkPredictionExtras/Public/MockAbilitySimulation.h
--- a/Engine/Plugins/Runtime/NetworkPredictionExtras/Source/NetworkPredictionExtras/Public/MockAbilitySimulation.h
+++ b/Engine/Plugins/Runtime/NetworkPredictionExtras/Source/NetworkPredictionExtras/Public/MockAbilitySimulation.h
@@ -125,30 +125,79 @@ struct FMockAbilityAuxState : public FFlyingMovementAuxState
 // MockAbility NetSimCues - events emitted by the sim
 // -------------------------------------------------------
 
+#define LOG_BLINK_CUE 1 // During development, its useful to sanity check that we aren't doing more construction or moves than we expect
+
 // Cue for blink activation.
 struct FMockAbilityBlinkCue
 {
+	FMockAbilityBlinkCue()
+	{ 
+		UE_CLOG(LOG_BLINK_CUE, LogTemp, Warning, TEXT("  Default Constructor 0x%X"), this);
+	}
+
+	FMockAbilityBlinkCue(FVector Start, FVector Stop) : StartLocation(Start), StopLocation(Stop)
+	{ 
+		UE_CLOG(LOG_BLINK_CUE, LogTemp, Warning, TEXT("  Custom Constructor 0x%X"), this);
+	}
+
+	~FMockAbilityBlinkCue()
+	{ 
+		UE_CLOG(LOG_BLINK_CUE, LogTemp, Warning, TEXT("  Destructor 0x%X"), this);
+	}
+	
+	FMockAbilityBlinkCue(FMockAbilityBlinkCue&& Other)
+		: StartLocation(MoveTemp(Other.StartLocation)), StopLocation(Other.StopLocation)
+	{
+		UE_CLOG(LOG_BLINK_CUE, LogTemp, Warning, TEXT("  Move Constructor 0x%X (Other: 0x%X)"), this, &Other);
+	}
+	
+	FMockAbilityBlinkCue& operator=(FMockAbilityBlinkCue&& Other)
+	{
+		UE_CLOG(LOG_BLINK_CUE, LogTemp, Warning, TEXT("  Move assignment 0x%X (Other: 0x%X)"), this, &Other);
+		StartLocation = MoveTemp(Other.StartLocation);
+		StopLocation = MoveTemp(Other.StopLocation);
+		return *this;
+	}
+
+	FMockAbilityBlinkCue(const FMockAbilityBlinkCue& Other) = delete;
+	FMockAbilityBlinkCue& operator=(const FMockAbilityBlinkCue& Other) = delete;
+
 	NETSIMCUE_BODY();
 
 	FVector_NetQuantize10 StartLocation;
 	FVector_NetQuantize10 StopLocation;
-
+	
 	void NetSerialize(FArchive& Ar)
 	{
 		bool b = false;
 		StartLocation.NetSerialize(Ar, nullptr, b);
 		StopLocation.NetSerialize(Ar, nullptr, b);
 	}
-
-	static bool Unique(const FMockAbilityBlinkCue& A, const FMockAbilityBlinkCue& B)
+	
+	bool NetUnique(const FMockAbilityBlinkCue& Other) const
 	{
 		const float ErrorTolerance = 1.f;
-		return !A.StartLocation.Equals(B.StartLocation, ErrorTolerance) || !A.StopLocation.Equals(B.StopLocation, ErrorTolerance);
+		return !StartLocation.Equals(Other.StartLocation, ErrorTolerance) || !StopLocation.Equals(Other.StopLocation, ErrorTolerance);
 	}
 };
 
 template<>
-struct TCueHandlerTraits<FMockAbilityBlinkCue> : public TNetSimCueTraits_ReplicatedNonPredicted { };
+struct TCueHandlerTraits<FMockAbilityBlinkCue> : public TNetSimCueTraits_Strong { };
+
+// -----------------------------------------------------------------------------------------------------
+// Subtypes of the BlinkCue - this is not an expected setup! This is done for testing/debugging so we can 
+// see the differences between the cue type traits in a controlled setup. See FMockAbilitySimulation::SimulationTick
+// -----------------------------------------------------------------------------------------------------
+#define DECLARE_BLINKCUE_SUBTYPE(TYPE, TRAITS) \
+ struct TYPE : public FMockAbilityBlinkCue { \
+ template <typename... ArgsType> TYPE(ArgsType&&... Args) : FMockAbilityBlinkCue(Forward<ArgsType>(Args)...) { } \
+ NETSIMCUE_BODY(); }; \
+ template<> struct TCueHandlerTraits<TYPE> : public TRAITS { };
+
+DECLARE_BLINKCUE_SUBTYPE(FMockAbilityBlinkCue_Weak, TNetSimCueTraits_Weak);
+DECLARE_BLINKCUE_SUBTYPE(FMockAbilityBlinkCue_ReplicatedNonPredicted, TNetSimCueTraits_ReplicatedNonPredicted);
+DECLARE_BLINKCUE_SUBTYPE(FMockAbilityBlinkCue_ReplicatedXOrPredicted, TNetSimCueTraits_ReplicatedXOrPredicted);
+DECLARE_BLINKCUE_SUBTYPE(FMockAbilityBlinkCue_Strong, TNetSimCueTraits_Strong);
 
 // The set of Cues the MockAbility simulation will invoke
 struct FMockAbilityCueSet
@@ -157,6 +206,12 @@ struct FMockAbilityCueSet
 	static void RegisterNetSimCueTypes(TDispatchTable& DispatchTable)
 	{
 		DispatchTable.template RegisterType<FMockAbilityBlinkCue>();
+
+		// (Again, not a normal setup, just for debugging/testing purposes)
+		DispatchTable.template RegisterType<FMockAbilityBlinkCue_Weak>();
+		DispatchTable.template RegisterType<FMockAbilityBlinkCue_ReplicatedNonPredicted>();
+		DispatchTable.template RegisterType<FMockAbilityBlinkCue_ReplicatedXOrPredicted>();
+		DispatchTable.template RegisterType<FMockAbilityBlinkCue_Strong>();
 	}
 };
 
@@ -239,7 +294,7 @@ public:
 	using UFlyingMovementComponent::FinalizeFrame;
 
 	// NetSimCues
-	void HandleCue(FMockAbilityBlinkCue& BlinkCue, const FNetworkSimTime& ElapsedTime);
+	void HandleCue(FMockAbilityBlinkCue& BlinkCue, const FNetSimCueSystemParamemters& SystemParameters);
 
 	// -------------------------------------------------------------------------------------
 	//	Ability State and Notifications
