NetworkPrediction Plugin
-Renamed NEtworkedSimulationModelReplicators.h -> NetworkedSimulationModelRepControllers.h

#rnx

diff --git a/Engine/Plugins/Runtime/NetworkPrediction/Source/NetworkPrediction/Public/NetworkedSimulationModel.h b/Engine/Plugins/Runtime/NetworkPrediction/Source/NetworkPrediction/Public/NetworkedSimulationModel.h
--- a/Engine/Plugins/Runtime/NetworkPrediction/Source/NetworkPrediction/Public/NetworkedSimulationModel.h
+++ b/Engine/Plugins/Runtime/NetworkPrediction/Source/NetworkPrediction/Public/NetworkedSimulationModel.h
@@ -4,7 +4,7 @@
 #include "NetworkPredictionTypes.h"
 #include "NetworkedSimulationModelBuffer.h"
 #include "NetworkedSimulationModelTypes.h"
-#include "NetworkedSimulationModelReplicators.h"
+#include "NetworkedSimulationModelRepControllers.h"
 #include "NetworkedSimulationModelDebugger.h"
 #include "NetworkedSimulationModelTraits.h"
 
diff --git a/Engine/Plugins/Runtime/NetworkPrediction/Source/NetworkPrediction/Public/NetworkedSimulationModelRepControllers.h b/Engine/Plugins/Runtime/NetworkPrediction/Source/NetworkPrediction/Public/NetworkedSimulationModelRepControllers.h
new file mode 100644
--- /dev/null
+++ b/Engine/Plugins/Runtime/NetworkPrediction/Source/NetworkPrediction/Public/NetworkedSimulationModelRepControllers.h
@@ -0,0 +1,899 @@
+// Copyright 1998-2019 Epic Games, Inc. All Rights Reserved.
+
+#pragma once
+
+#include "Net/UnrealNetwork.h" // For MakeRelative
+#include "NetworkedSimulationModelCVars.h"
+#include "NetworkedSimulationModelInterpolator.h"
+
+// ----------------------------------------------------------------------------------------------------------------------------------------------------------------------
+//	CVars and compile time constants
+// ----------------------------------------------------------------------------------------------------------------------------------------------------------------------
+
+namespace NetworkSimulationModelCVars
+{
+	NETSIM_DEVCVAR_SHIPCONST_INT(EnableLocalPrediction, 1, "ns.EnableLocalPrediction", "Toggle local prediction.");
+	NETSIM_DEVCVAR_SHIPCONST_INT(EnableSimulatedReconcile, 1, "ns.EnableSimulatedReconcile", "Toggle simulated proxy reconciliation.");
+	NETSIM_DEVCVAR_SHIPCONST_INT(EnableSimulatedExtrapolation, 1, "ns.EnableSimulatedExtrapolation", "Toggle simulated proxy extrapolation.");
+	NETSIM_DEVCVAR_SHIPCONST_INT(ForceReconcile, 0, "ns.ForceReconcile", "Forces reconcile even if state does not differ. E.g, force resimulation after every netupdate.");
+	NETSIM_DEVCVAR_SHIPCONST_INT(ForceReconcileSingle, 0, "ns.ForceReconcileSingle", "Forces a since reconcile to happen on the next frame");
+}
+
+#ifndef NETSIM_NETCONSTANT_NUM_BITS_FRAME
+	#define NETSIM_NETCONSTANT_NUM_BITS_FRAME 8	// Allows you to override this setting via UBT, but access via cleaner FActorMotionNetworkingConstants::NUM_BITS_FRAME
+#endif
+
+struct FNetworkSimulationSerialization
+{
+	// How many bits we use to encode the key frame number for buffers.
+	// Client Frames are stored locally as 32 bit integers, but we use a smaller # of bits to NetSerialize.
+	// Frames are only relatively relevant: the absolute value doesn't really matter. We just need to detect newer/older.
+	enum { NUM_BITS_FRAME = NETSIM_NETCONSTANT_NUM_BITS_FRAME };		
+
+	// Abs max value we encode into the bit writer
+	enum { MAX_FRAME_WRITE = 1 << NUM_BITS_FRAME };
+
+	// This is the threshold at which we would wrap around and incorrectly assign a frame on the receiving side.
+	// E.g, If there are FRAME_ERROR_THRESHOLD frames that do not make it across from sender->receiver, the
+	// receiver will have incorrect local values. With 8 bits, this works out to be 128 frames or about 2 seconds at 60fps.
+	enum { FRAME_ERROR_THRESHOLD = MAX_FRAME_WRITE / 2};
+
+	// Helper to serialize the int32 HeadFrame. Returns the unpacked value (this will be same as input in the save path)
+	static int32 SerializeFrame(FArchive& Ar, int32 LocalHeadFrame)
+	{
+		if (Ar.IsSaving())
+		{
+			((FNetBitWriter&)Ar).WriteIntWrapped( LocalHeadFrame, MAX_FRAME_WRITE );
+			return LocalHeadFrame;
+		}
+		
+		return MakeRelative(((FNetBitReader&)Ar).ReadInt( MAX_FRAME_WRITE ), LocalHeadFrame, MAX_FRAME_WRITE );
+	}
+};
+
+// -------------------------------------------------------------------------------------------------------
+//	Helpers
+// -------------------------------------------------------------------------------------------------------
+
+/** Helper that writes a new input cmd to the input buffer, at given Frame (usually the sim's PendingFrame). If frame doesn't exist, ProduceInput is called on the driver if bProduceInputViaDriver=true, otherwise the input cmd will be initialized from the previous input cmd. */
+template<typename Model>
+void GenerateLocalInputCmdAtFrame(typename TNetSimModelTraits<Model>::Driver* Driver, TNetworkSimBufferContainer<Model>& Buffers, const FNetworkSimTime& DeltaSimTime, int32 Frame, bool bProduceInputViaDriver=true)
+{
+	using TInputCmd = typename TNetworkSimBufferContainer<Model>::TInputCmd;
+	if (TInputCmd* InputCmd = Buffers.Input[Frame])
+	{
+		InputCmd->SetFrameDeltaTime(DeltaSimTime);
+	}
+	else
+	{
+		InputCmd = Buffers.Input.WriteFrameInitializedFromHead(Frame);
+		InputCmd->SetFrameDeltaTime(DeltaSimTime);
+		if (bProduceInputViaDriver)
+		{
+			Driver->ProduceInput(DeltaSimTime, *InputCmd);
+		}
+	}
+}
+
+/** Helper to generate a local input cmd if we have simulation time to spend and advance the simulation's MaxAllowedFrame so that it can be processed. */
+template<typename Model>
+void TryGenerateLocalInput(typename TNetSimModelTraits<Model>::Driver* Driver, TNetworkSimBufferContainer<Model>& Buffers, TSimulationTicker<typename Model::TickSettings>& Ticker, bool bProduceInputViaDriver=true)
+{
+	using TBufferTypes = typename TNetSimModelTraits<Model>::InternalBufferTypes;
+	using TInputCmd = typename TBufferTypes::TInputCmd;
+
+	FNetworkSimTime DeltaSimTime = Ticker.GetRemainingAllowedSimulationTime();
+	if (DeltaSimTime.IsPositive())
+	{
+		GenerateLocalInputCmdAtFrame<Model>(Driver, Buffers, DeltaSimTime, Ticker.PendingFrame, bProduceInputViaDriver);
+		Ticker.MaxAllowedFrame = Ticker.PendingFrame;
+	}
+}
+
+template<typename T>
+struct TNetSimModelTraits;
+
+// ----------------------------------------------------------------------------------------------------------------------------------------------------------------------
+//	Replication Controllers
+//	
+//	The RepControllers are the pieces of the TNetworkedSimulationModel that make up the role-specific functionality (Server, Autonomous Client, and Simulated Client).
+//	Mainly they NetSerialize, Reconcile, and have Pre/Post sim tick functions. The TNetworkedSimulationModel is still the core piece that ticks the sim, but the Replicators
+//	do everything else in a role-specific way.
+// ----------------------------------------------------------------------------------------------------------------------------------------------------------------------
+
+// Empty RepController: zero functionality, just stubbed out functions and using definitions for convenience
+template<typename Model>
+struct TRepController_Empty
+{ 
+	using TSimulation = typename Model::Simulation;
+	using TTickSettings = typename Model::TickSettings;
+	using TDriver = typename TNetSimModelTraits<Model>::Driver;
+	using TBufferTypes = typename TNetSimModelTraits<Model>::InternalBufferTypes;
+	using TInputCmd = typename TBufferTypes::TInputCmd;
+	using TSyncState = typename TBufferTypes::TSyncState;
+	using TAuxState = typename TBufferTypes::TAuxState;
+
+	// Used for client shadowstate compares. Should just call GetDirtyCount() on the buffer you are replicating
+	int32 GetProxyDirtyCount(TNetworkSimBufferContainer<Model>& Buffers) const { return 0; }
+
+	// NetSerialize: just serialize the network data. Don't run simulation steps. Every replicator will be NetSerialized before moving on to Reconcile phase.
+	void NetSerialize(const FNetSerializeParams& P, TNetworkSimBufferContainer<Model>& Buffers, const TSimulationTicker<TTickSettings>& Ticker) { }
+
+	// Reconcile: called after everyone has NetSerialized. "Get right with the server": this function is about reconciling what the server told you vs what you have extrapolated or forward predicted locally
+	void Reconcile(TSimulation* Sim, TDriver* Driver, TNetworkSimBufferContainer<Model>& Buffers, TSimulationTicker<TTickSettings>& Ticker) { }
+
+	// Called prior to input processing. This function must update Ticker to allow simulation time (from TickParameters) and to possibly get new input.
+	void PreSimTick(TDriver* Driver, TNetworkSimBufferContainer<Model>& Buffers, TSimulationTicker<TTickSettings>& Ticker, const FNetSimTickParameters& TickParameters) { }
+
+	// Called after input processing. Should finalize the frame and do any smoothing/interpolation. This function is not allowed to modify the buffers or tick state, or even call the simulation/Update function.
+	void PostSimTick(TDriver* Driver, const TNetworkSimBufferContainer<Model>& Buffers, const TSimulationTicker<TTickSettings>& Ticker, const FNetSimTickParameters& TickParameters) { }
+};
+
+// Base implementation with minimal functionality: will grant simulation time and call finalize frame
+template<typename Model, typename TBase=TRepController_Empty<Model>>
+struct TRepController_Base : public TBase
+{
+	using typename TBase::TSimulation;
+	using typename TBase::TTickSettings;
+	using typename TBase::TDriver;
+	using typename TBase::TBufferTypes;
+	using typename TBase::TInputCmd;
+
+	// Called prior to input processing. This function must updated Ticker to allow simulation time (from TickParameters) and to possibly get new input.
+	void PreSimTick(TDriver* Driver, TNetworkSimBufferContainer<Model>& Buffers, TSimulationTicker<TTickSettings>& Ticker, const FNetSimTickParameters& TickParameters)
+	{
+		// Accumulate local delta time into Ticker
+		Ticker.GiveSimulationTime(TickParameters.LocalDeltaTimeSeconds);
+	}
+
+	// Called after input processing. Should finalize the frame and do any smoothing/interpolation. This function is not allowed to modify the buffers or tick state, or even call the simulation/Update function.
+	void PostSimTick(TDriver* Driver, const TNetworkSimBufferContainer<Model>& Buffers, const TSimulationTicker<TTickSettings>& Ticker, const FNetSimTickParameters& TickParameters)
+	{
+		// Sync to latest frame if there is any
+		if (Buffers.Sync.Num() > 0)
+		{
+			Driver->FinalizeFrame(*Buffers.Sync.HeadElement(), *Buffers.Aux.HeadElement());
+		}
+	}
+};
+
+// -------------------------------------------------------------------------------------------------------
+//	Simulation Time replicator
+//	Just serializes simulation time. Can be disabled by Enabled templated parameter. Used as base class in other classes below.
+// -------------------------------------------------------------------------------------------------------
+template<typename Model, typename TBase=TRepController_Base<Model>, bool Enabled=true>
+struct TRepController_SimTime : public TBase
+{
+	using typename TBase::TSimulation;
+	using typename TBase::TTickSettings;
+	using typename TBase::TDriver;
+	using typename TBase::TBufferTypes;
+	using typename TBase::TInputCmd;
+
+	void NetSerialize(const FNetSerializeParams& P, TNetworkSimBufferContainer<Model>& Buffers, TSimulationTicker<TTickSettings>& Ticker)
+	{
+		TBase::NetSerialize(P, Buffers, Ticker);
+		SerializedTime = Ticker.GetTotalProcessedSimulationTime();
+		SerializedTime.NetSerialize(P.Ar);
+	}
+
+	FNetworkSimTime SerializedTime;
+};
+
+// Enabled=false specialization: do nothing
+template<typename Model, typename TBase>
+struct TRepController_SimTime<Model, TBase, false> : public TBase { };
+
+// -------------------------------------------------------------------------------------------------------
+//	Sequence replicator
+//
+//	Replicates a sequence of elements. i.e., "the last MaxNumElements".
+//	On the receiving side, we merge the sequence into whatever we have locally
+//	Frames are synchronized. SimTime is also serialized by default (change by changing TBase)
+// -------------------------------------------------------------------------------------------------------
+
+template<typename Model, ENetworkSimBufferTypeId BufferId, int32 MaxNumElements=3, typename TBase=TRepController_SimTime<Model>>
+struct TRepController_Sequence : public TBase
+{
+	using typename TBase::TSimulation;
+	using typename TBase::TTickSettings;
+	using typename TBase::TDriver;
+	using typename TBase::TBufferTypes;
+	using typename TBase::TInputCmd;
+
+	int32 GetProxyDirtyCount(TNetworkSimBufferContainer<Model>& Buffers) const 
+	{
+		return Buffers.template Get<BufferId>().GetDirtyCount() ^ (TBase::GetProxyDirtyCount(Buffers) << 2); 
+	}
+
+	void NetSerialize(const FNetSerializeParams& P, TNetworkSimBufferContainer<Model>& Buffers, TSimulationTicker<TTickSettings>& Ticker)
+	{
+		TBase::NetSerialize(P, Buffers, Ticker);
+
+		auto& Buffer = Buffers.template Get<BufferId>();
+		FArchive& Ar = P.Ar;
+		uint8 SerializedNumElements = FMath::Min<uint8>(MaxNumElements, Buffer.Num());
+		Ar << SerializedNumElements;
+
+		const int32 HeadFrame = FNetworkSimulationSerialization::SerializeFrame(Ar, Buffer.HeadFrame());
+		const int32 StartingFrame = FMath::Max(0, HeadFrame - SerializedNumElements + 1);
+
+		if (Ar.IsLoading())
+		{
+			const int32 PrevHead = Buffer.HeadFrame();
+			if (PrevHead < StartingFrame && PrevHead >= 0)
+			{
+				// There is a gap in the stream. In some cases, we want this to be a "fault" and bubble up. We may want to synthesize state or maybe we just skip ahead.
+				UE_LOG(LogNetworkSim, Warning, TEXT("Fault: gap in received %s buffer. PrevHead: %d. Received: %d-%d. Reseting previous buffer contents"), *LexToString(BufferId), PrevHead, StartingFrame, HeadFrame);
+			}
+		}
+
+		for (int32 Frame = StartingFrame; Frame <= HeadFrame; ++Frame)
+		{
+			// This, as is, is bad. The intention is that these functions serialize multiple items in some delta compressed fashion.
+			// As is, we are just serializing the elements individually.
+			auto* Cmd = Ar.IsLoading() ? Buffer.WriteFrame(Frame) : Buffer[Frame];
+			Cmd->NetSerialize(P);
+		}
+
+		LastSerializedFrame = HeadFrame;
+	}
+
+	int32 GetLastSerializedFrame() const { return LastSerializedFrame; }
+
+protected:
+
+	int32 LastSerializedFrame = 0;
+};
+
+// -------------------------------------------------------------------------------------------------------
+//	Server Replication Controller (FIXME: should this be called something like "client autonomous proxy -> Server RPC" ?
+//
+// -------------------------------------------------------------------------------------------------------
+template<typename Model, typename TBase=TRepController_Sequence<Model, ENetworkSimBufferTypeId::Input, 3>>
+struct TRepController_Server : public TBase
+{
+	using typename TBase::TSimulation;
+	using typename TBase::TTickSettings;
+	using typename TBase::TDriver;
+	using typename TBase::TBufferTypes;
+	using typename TBase::TInputCmd;
+	
+	void Reconcile(TSimulation* Simulation, TDriver* Driver, TNetworkSimBufferContainer<Model>& Buffers, TSimulationTicker<TTickSettings>& Ticker)
+	{
+		// After receiving input, server may process up to the latest received frames.
+		// (If we needed to buffer input server side for whatever reason, we would do it here)
+		// (also note that we will implicitly guard against speed hacks in the core update loop by not processing cmds past what we have been "allowed")
+		Ticker.MaxAllowedFrame = Buffers.Input.HeadFrame();
+
+		// Check for gaps in commands
+		if ( Ticker.PendingFrame < Buffers.Input.TailFrame() )
+		{
+			UE_LOG(LogNetworkSim, Warning, TEXT("TReplicator_Server::Reconcile missing inputcmds. PendingFrame: %d. %s. This can happen via packet loss"), Ticker.PendingFrame, *Buffers.Input.GetBasicDebugStr());
+			Ticker.PendingFrame = Buffers.Input.TailFrame();
+		}
+	}
+	
+	void PreSimTick(TDriver* Driver, TNetworkSimBufferContainer<Model>& Buffers, TSimulationTicker<TTickSettings>& Ticker, const FNetSimTickParameters& TickParameters)
+	{
+		Ticker.GiveSimulationTime(TickParameters.LocalDeltaTimeSeconds);
+		if (TickParameters.bGenerateLocalInputCmds)
+		{
+			TryGenerateLocalInput(Driver, Buffers, Ticker);
+		}
+	}
+};
+
+// Simulated: "non locally controlled" simulations. We support "Simulation Extrapolation" here (using the sim to fake inputs to advance the sim)
+//	-TODO: this is replicating Input/Sync/Aux which is required to do accurate Simulation Extrapolation. For interpolated simulated proxies, we can skip the input and possibly the aux.
+//	More settings/config options would be nice here
+template<typename Model, typename TBase=TRepController_Empty<Model>>
+struct TRepController_Simulated : public TBase
+{
+	using typename TBase::TSimulation;
+	using typename TBase::TTickSettings;
+	using typename TBase::TDriver;
+	using typename TBase::TBufferTypes;
+	using typename TBase::TInputCmd;
+	using typename TBase::TSyncState;
+	using typename TBase::TAuxState;
+
+	// Parent Simulation. If this is set, this simulation will forward predict in sync with this parent sim. The parent sim should be an autonomous proxy driven simulation
+	INetworkedSimulationModel* ParentSimulation = nullptr;
+
+	// Instance flag for enabling simulated extrapolation
+	bool bAllowSimulatedExtrapolation = true;
+
+	// Interpolated that will be used if bAllowSimulatedExtrapolation == false && ParentSimulation == nullptr
+	TNetSimInterpolator<Model> Interpolator;
+
+	// Last Serialized time and state
+	FNetworkSimTime GetLastSerializedSimulationTime() const { return LastSerializedSimulationTime; }
+	const TSyncState& GetLastSerializedSyncState() const { return LastSerializedSyncState; }
+	
+	ESimulatedUpdateMode GetSimulatedUpdateMode() const
+	{
+		if (ParentSimulation)
+		{
+			return ESimulatedUpdateMode::ForwardPredict;
+		}
+		if (bAllowSimulatedExtrapolation && NetworkSimulationModelCVars::EnableSimulatedExtrapolation())
+		{
+			return ESimulatedUpdateMode::Extrapolate;
+		}
+
+		return ESimulatedUpdateMode::Interpolate;
+	}
+	
+	int32 GetProxyDirtyCount(TNetworkSimBufferContainer<Model>& Buffers) const
+	{
+		return Buffers.Sync.GetDirtyCount() ^ (TBase::GetProxyDirtyCount(Buffers) << 2); 
+	}
+
+	void NetSerialize(const FNetSerializeParams& P, TNetworkSimBufferContainer<Model>& Buffers, TSimulationTicker<TTickSettings>& Ticker)
+	{
+		FArchive& Ar = P.Ar;
+
+		FNetworkSimTime PrevLastSerializedSimulationTime = LastSerializedSimulationTime;
+
+		// Serialize latest simulation time
+		LastSerializedSimulationTime = Ticker.GetTotalProcessedSimulationTime();
+		LastSerializedSimulationTime.NetSerialize(P.Ar);
+
+		// Serialize latest element
+		TInputCmd* InputCmd = nullptr;
+		TSyncState* SyncState = nullptr;
+		TAuxState* AuxState = nullptr;
+		
+		if (Ar.IsSaving())
+		{
+			const int32 Frame = Buffers.Sync.HeadFrame();
+
+			InputCmd = Buffers.Input.HeadElement();
+			if (!InputCmd)
+			{
+				InputCmd = &LastSerializedInputCmd;
+			}
+			SyncState = Buffers.Sync.HeadElement();
+			AuxState = Buffers.Aux[Frame];
+			check(InputCmd && SyncState && AuxState	); // We should not be here if the buffer is empty. Want to avoid serializing an "empty" flag at the top here.
+		}
+		else
+		{
+			check(Ticker.SimulationTimeBuffer.HeadFrame() == Buffers.Sync.HeadFrame());
+			check(Ticker.GetTotalProcessedSimulationTime() <= Ticker.GetTotalAllowedSimulationTime());
+
+			// Cache off our "starting" time before possibly overwriting it. We will use this in Reconcile to catch back up in some cases.
+			if (Ticker.GetTotalProcessedSimulationTime() > ReconcileSimulationTime)
+			{
+				ReconcileSimulationTime = Ticker.GetTotalProcessedSimulationTime();
+			}
+
+			// Find out where this should go in the local buffer based on the serialized time
+			int32 DestinationFrame = INDEX_NONE;
+			if (LastSerializedSimulationTime > Ticker.GetTotalProcessedSimulationTime())
+			{
+				// We are getting new state that is ahead of what we have locally, so it can safety go right to head
+				DestinationFrame = Ticker.SimulationTimeBuffer.HeadFrame()+1;
+			}
+			else
+			{
+				// We are getting state that is behind what we have locally
+				for (int32 Frame = Ticker.SimulationTimeBuffer.HeadFrame(); Frame >= Ticker.SimulationTimeBuffer.TailFrame(); --Frame)
+				{
+					if (LastSerializedSimulationTime > *Ticker.SimulationTimeBuffer[Frame])
+					{
+						DestinationFrame = Frame+1;
+						break;
+					}
+				}
+
+				if (DestinationFrame == INDEX_NONE)
+				{
+					FNetworkSimTime TotalTimeAhead = Ticker.GetTotalProcessedSimulationTime() - LastSerializedSimulationTime;
+					FNetworkSimTime SerializeDelta = LastSerializedSimulationTime - PrevLastSerializedSimulationTime;
+
+					// We are way far ahead of the server... we will need to clear out sync buffers, take what they gave us, and catch up in reconcile
+					//UE_LOG(LogNetworkSim, Warning, TEXT("!!! TReplicator_Simulated. Large gap detected. SerializedTime: %s. Buffer time: [%s-%s]. %d Elements. DeltaFromHead: %s. DeltaSerialize: %s"), *LastSerializedSimulationTime.ToString(), 
+					//	*Ticker.SimulationTimeBuffer.GetElementFromTail(0)->ToString(), *Ticker.SimulationTimeBuffer.GetElementFromHead(0)->ToString(), Ticker.SimulationTimeBuffer.GetNumValidElements(), *TotalTimeAhead.ToString(), *SerializeDelta.ToString());
+					DestinationFrame = Ticker.SimulationTimeBuffer.HeadFrame()+2; // (Skip ahead 2 to force a break in continuity)
+				}
+			}
+
+			check(DestinationFrame != INDEX_NONE);
+
+			// "Finalize" our buffers and time keeping such that we serialize the latest state from the server in the right spot
+			InputCmd = Buffers.Input.WriteFrame(DestinationFrame);
+			SyncState = Buffers.Sync.WriteFrame(DestinationFrame);
+			AuxState = Buffers.Aux.WriteFrame(DestinationFrame);
+			Buffers.Input.WriteFrame(DestinationFrame);
+
+			// Update tick info
+			Ticker.SetTotalProcessedSimulationTime(LastSerializedSimulationTime, DestinationFrame);
+			if (Ticker.GetTotalAllowedSimulationTime() < LastSerializedSimulationTime)
+			{
+				Ticker.SetTotalAllowedSimulationTime(LastSerializedSimulationTime);
+			}
+
+			check(Ticker.GetTotalProcessedSimulationTime() <= Ticker.GetTotalAllowedSimulationTime());
+
+			Ticker.PendingFrame = DestinationFrame;	// We are about to serialize state to DestinationFrame which will be "unprocessed" (has not been used to generate a new frame)
+			Ticker.MaxAllowedFrame = DestinationFrame-1; // Do not process PendingFrame on our account. ::PreSimTick will advance this based on our interpolation/extrapolation settings
+		}
+
+		check(SyncState && AuxState);
+		InputCmd->NetSerialize(Ar);
+		SyncState->NetSerialize(Ar);
+		AuxState->NetSerialize(Ar);
+
+		Buffers.CueDispatcher.NetSerializeSavedCues(Ar, GetSimulatedUpdateMode() == ESimulatedUpdateMode::Interpolate);
+
+		if (Ar.IsLoading())
+		{
+			LastSerializedInputCmd = *InputCmd;
+			LastSerializedSyncState = *SyncState;
+			LastSerializedAuxState = *AuxState;
+		}
+	}
+
+	void Reconcile(TSimulation* Simulation, TDriver* Driver, TNetworkSimBufferContainer<Model>& Buffers, TSimulationTicker<TTickSettings>& Ticker)
+	{
+		if (ReconcileSimulationTime.IsPositive() == false)
+		{
+			return;
+		}
+
+		check(Ticker.GetTotalProcessedSimulationTime() <= Ticker.GetTotalAllowedSimulationTime());
+
+		if (GetSimulatedUpdateMode() == ESimulatedUpdateMode::Extrapolate && NetworkSimulationModelCVars::EnableSimulatedReconcile())
+		{
+			// This is effectively a rollback: we went back in time during NetSerialize and are now going to catch up ReconcileSimulationTime ms.
+			Buffers.CueDispatcher.NotifyRollback(Ticker.GetTotalProcessedSimulationTime());
+
+			// Simulated Reconcile requires the input buffer to be kept up to date with the Sync buffer
+			// Generate a new, fake, command since we just added a new sync state to head
+			while (Buffers.Input.HeadFrame() < Buffers.Sync.HeadFrame())
+			{
+				TInputCmd* Next = Buffers.Input.WriteFrameInitializedFromHead(Buffers.Input.HeadFrame()+1);
+			}
+
+			// Do we have time to make up? We may have extrapolated ahead of the server (totally fine - can happen with small amount of latency variance)
+			FNetworkSimTime DeltaSimTime = ReconcileSimulationTime - Ticker.GetTotalProcessedSimulationTime();
+			if (DeltaSimTime.IsPositive() && NetworkSimulationModelCVars::EnableSimulatedReconcile())
+			{
+				SimulationExtrapolation(Simulation, Driver, Buffers, Ticker, DeltaSimTime);
+			}
+		}
+		
+		check(Ticker.GetTotalProcessedSimulationTime() <= Ticker.GetTotalAllowedSimulationTime());
+		ReconcileSimulationTime.Reset();
+	}
+	
+	void PreSimTick(TDriver* Driver, TNetworkSimBufferContainer<Model>& Buffers, TSimulationTicker<TTickSettings>& Ticker, const FNetSimTickParameters& TickParameters)
+	{
+		// Tick if we are dependent simulation or extrapolation is enabled
+		if (ParentSimulation || (GetSimulatedUpdateMode() == ESimulatedUpdateMode::Extrapolate))
+		{
+			// Don't start this simulation until you've gotten at least one update from the server
+			if (Ticker.GetTotalProcessedSimulationTime().IsPositive())
+			{
+				Ticker.GiveSimulationTime(TickParameters.LocalDeltaTimeSeconds);
+			}
+			
+			// Generate local input if we are ready to tick. Note that we pass false so we won't call into the Driver to produce the input, we will use the last serialized InputCmd's values
+			TryGenerateLocalInput<Model>(Driver, Buffers, Ticker, false);
+		}
+	}
+
+	void PostSimTick(TDriver* Driver, const TNetworkSimBufferContainer<Model>& Buffers, const TSimulationTicker<TTickSettings>& Ticker, const FNetSimTickParameters& TickParameters)
+	{
+		if (GetSimulatedUpdateMode() == ESimulatedUpdateMode::Interpolate)
+		{
+			const FNetworkSimTime::FRealTime InterpolationRealTime = Interpolator.template PostSimTick<TDriver>(Driver, Buffers, Ticker, TickParameters);
+			const FNetworkSimTime InterpolationSimTime = FNetworkSimTime::FromRealTimeSeconds(InterpolationRealTime);
+
+			Buffers.CueDispatcher.SetMaxDispatchTime(InterpolationSimTime);
+			Buffers.CueDispatcher.SetConfirmedTime(InterpolationSimTime);
+		}
+		else
+		{
+			// Sync to latest frame if there is any
+			if (Buffers.Sync.Num() > 0)
+			{
+				Driver->FinalizeFrame(*Buffers.Sync.HeadElement(), *Buffers.Aux.HeadElement());
+			}
+
+			Buffers.CueDispatcher.ClearMaxDispatchTime();
+			Buffers.CueDispatcher.SetConfirmedTime(LastSerializedSimulationTime);
+		}
+	}
+	
+	void DependentRollbackBegin(TSimulation* Simulation, TDriver* Driver, TNetworkSimBufferContainer<Model>& Buffers, TSimulationTicker<TTickSettings>& Ticker, const FNetworkSimTime& RollbackDeltaTime, const int32 ParentFrame)
+	{
+		// For now, we make the assumption that our last serialized state and time match the parent simulation.
+		// This would not always be the case with low frequency simulated proxies. But could be handled by replicating the simulations together (at the replication level)
+		const int32 NewHeadFrame = Buffers.Sync.HeadFrame()+1;
+		
+		Ticker.SetTotalProcessedSimulationTime(LastSerializedSimulationTime, NewHeadFrame);
+		Ticker.SetTotalAllowedSimulationTime(LastSerializedSimulationTime);
+		Ticker.PendingFrame = NewHeadFrame;
+		Ticker.MaxAllowedFrame = NewHeadFrame;
+
+		*Buffers.Sync.WriteFrame(NewHeadFrame) = LastSerializedSyncState;
+		*Buffers.Input.WriteFrame(NewHeadFrame) = TInputCmd();
+
+		Driver->FinalizeFrame(LastSerializedSyncState, LastSerializedAuxState);
+
+		Buffers.CueDispatcher.NotifyRollback(LastSerializedSimulationTime);
+
+		if (NETSIM_MODEL_DEBUG)
+		{
+			FVisualLoggingParameters VLogParams(EVisualLoggingContext::FirstMispredicted, ParentFrame, EVisualLoggingLifetime::Persistent, TEXT("DependentRollbackBegin"));
+			Driver->VisualLog(nullptr, &LastSerializedSyncState, &LastSerializedAuxState, VLogParams);
+		}
+	}
+		
+	void DependentRollbackStep(TSimulation* Simulation, TDriver* Driver, TNetworkSimBufferContainer<Model>& Buffers, TSimulationTicker<TTickSettings>& Ticker, const FNetworkSimTime& StepTime, const int32 ParentFrame, const bool bFinalStep)
+	{
+		Ticker.SetTotalAllowedSimulationTime( Ticker.GetTotalAllowedSimulationTime() + StepTime );
+
+		SimulationExtrapolation(Simulation, Driver, Buffers, Ticker, StepTime);
+
+		TSyncState* SyncState = Buffers.Sync.HeadElement();
+		check(SyncState);
+		
+		if (NETSIM_MODEL_DEBUG)
+		{
+			FVisualLoggingParameters VLogParams(bFinalStep ? EVisualLoggingContext::LastMispredicted : EVisualLoggingContext::OtherMispredicted, ParentFrame, EVisualLoggingLifetime::Persistent, TEXT("DependentRollbackStep"));
+			Driver->VisualLog(nullptr, SyncState, Buffers.Aux.HeadElement(), VLogParams);
+		}
+	}
+
+private:
+	
+	void SimulationExtrapolation(TSimulation* Simulation, TDriver* Driver, TNetworkSimBufferContainer<Model>& Buffers, TSimulationTicker<TTickSettings>& Ticker, const FNetworkSimTime DeltaSimTime)
+	{
+		// We have extrapolated ahead of the server. The latest network update is now "in the past" from what we rendered last frame.
+		// We will insert a new frame to make up the difference from the last known state to where we want to be in the now.
+
+		ensure(Buffers.Input.HeadFrame() == Buffers.Sync.HeadFrame());
+
+		const int32 InputFrame = Buffers.Input.HeadFrame();
+		const int32 OutputFrame = InputFrame + 1;
+
+		// Create fake cmd				
+		TInputCmd* NewCmd = Buffers.Input.WriteFrameInitializedFromHead(OutputFrame);
+		NewCmd->SetFrameDeltaTime(DeltaSimTime);	
+		
+		// Create new sync state to write to
+		TSyncState* PrevSyncState = Buffers.Sync[InputFrame];
+		TSyncState* NextSyncState = Buffers.Sync.WriteFrame(OutputFrame);
+		TAuxState* AuxState = Buffers.Aux[InputFrame];
+
+		// Do the actual update
+		{
+			TScopedSimulationTick<Model> UpdateScope(Ticker, Buffers.CueDispatcher, ESimulationTickContext::Resimulate, OutputFrame, NewCmd->GetFrameDeltaTime());
+			Simulation->SimulationTick( 
+				{ NewCmd->GetFrameDeltaTime(), Ticker },
+				{ *NewCmd, *PrevSyncState, *AuxState },
+				{ *NextSyncState, Buffers.Aux.LazyWriter(OutputFrame), Buffers.CueDispatcher } );
+		}
+
+		Ticker.MaxAllowedFrame = OutputFrame;
+	}
+	
+	FNetworkSimTime ReconcileSimulationTime;
+	FNetworkSimTime LastSerializedSimulationTime;
+	
+	TInputCmd LastSerializedInputCmd;
+	TSyncState LastSerializedSyncState;
+	TAuxState LastSerializedAuxState;	// Temp? This should be conditional or optional. We want to support not replicating the aux state to simulated proxies
+};
+
+/** Replicates TSyncState and does basic reconciliation. */
+template<typename Model, typename TBase=TRepController_Base<Model>>
+struct TRepController_Autonomous: public TBase
+{
+	using typename TBase::TSimulation;
+	using typename TBase::TTickSettings;
+	using typename TBase::TDriver;
+	using typename TBase::TBufferTypes;
+	using typename TBase::TInputCmd;
+	using typename TBase::TSyncState;
+	using typename TBase::TAuxState;
+
+	int32 GetLastSerializedFrame() const { return SerializedFrame; }
+	bool IsReconcileFaultDetected() const { return bReconcileFaultDetected; }
+	const FNetworkSimTime& GetLastSerializedSimTime() const { return SerializedTime; }
+
+	TArray<INetworkedSimulationModel*> DependentSimulations;
+	bool bDependentSimulationNeedsReconcile = false;
+
+	int32 GetProxyDirtyCount(TNetworkSimBufferContainer<Model>& Buffers) const
+	{
+		return (Buffers.Sync.GetDirtyCount()) ^ (Buffers.Aux.GetDirtyCount() << 1) ^ (TBase::GetProxyDirtyCount(Buffers) << 2); 
+	}
+
+	// --------------------------------------------------------------------
+	//	NetSerialize
+	// --------------------------------------------------------------------
+	void NetSerialize(const FNetSerializeParams& P, TNetworkSimBufferContainer<Model>& Buffers, TSimulationTicker<TTickSettings>& Ticker)
+	{
+		FArchive& Ar = P.Ar;
+		
+		SerializedFrame = FNetworkSimulationSerialization::SerializeFrame(Ar, Buffers.Sync.HeadFrame());
+		
+		SerializedTime = Ticker.GetTotalProcessedSimulationTime();
+		SerializedTime.NetSerialize(P.Ar);
+
+		if (Ar.IsSaving())
+		{
+			// Server serialize the latest state
+			Buffers.Sync.HeadElement()->NetSerialize(Ar);
+			Buffers.Aux.HeadElement()->NetSerialize(Ar);
+		}
+		else
+		{
+			SerializedSyncState.NetSerialize(Ar);
+			SerializedAuxState.NetSerialize(Ar);
+		}
+
+		Buffers.CueDispatcher.NetSerializeSavedCues(Ar, false);
+
+		if (Ar.IsLoading())
+		{
+			bReconcileFaultDetected = false;
+			bPendingReconciliation = false;
+
+			TSyncState* ClientSync = Buffers.Sync[SerializedFrame];
+			TAuxState* ClientAux = Buffers.Aux[SerializedFrame];
+
+			// The state the client predicted that corresponds to the state the server just serialized to us
+			if (ClientSync && ClientAux)
+			{
+				const bool bForceReconcile = (NetworkSimulationModelCVars::ForceReconcile() > 0) || (NetworkSimulationModelCVars::ForceReconcileSingle() > 0);
+				
+				if (Model::ShouldReconcile(SerializedSyncState, SerializedAuxState, *ClientSync, *ClientAux))
+				{
+					NetworkSimulationModelCVars::SetForceReconcileSingle(0);
+					UE_CLOG(!Buffers.Input.IsValidFrame(SerializedFrame-1), LogNetworkSim, Error, TEXT("::NetSerialize: Client InputBuffer does not contain data for frame %d. {%s} {%s}"), SerializedFrame, *Buffers.Input.GetBasicDebugStr(), *Buffers.Sync.GetBasicDebugStr());
+					bPendingReconciliation =  true;
+				}
+			}
+			else
+			{
+				if (SerializedFrame < Buffers.Sync.TailFrame())
+				{
+					// Case 1: the serialized state is older than what we've kept in our buffer. A bigger buffer would solve this! (at the price of more resimulated frames to recover when this happens)
+					// This is a reconcile fault and we just need to chill. We'll stop sending user commands until the cmds in flight flush through the system and we catch back up.
+					bReconcileFaultDetected = true;
+				}
+				else
+				{
+					// Case 2: We've received a newer frame than what we've processed locally. This could happen if we are buffering our inputs locally (but still sending to the server) or just not predicting
+					bPendingReconciliation =  true;
+				}
+			}
+		}
+	}
+
+	// --------------------------------------------------------------------
+	//	Reconcile
+	// --------------------------------------------------------------------
+	void Reconcile(TSimulation* Simulation, TDriver* Driver, TNetworkSimBufferContainer<Model>& Buffers, TSimulationTicker<TTickSettings>& Ticker)
+	{
+		if (bPendingReconciliation == false && (bDependentSimulationNeedsReconcile == false || SerializedFrame == INDEX_NONE))
+		{
+			return;
+		}
+		bPendingReconciliation = false;
+		bDependentSimulationNeedsReconcile = false;
+
+		TSyncState* ClientSyncState = Buffers.Sync[SerializedFrame];
+		const bool bDoVisualLog = NetworkSimulationModelCVars::EnableLocalPrediction() > 0 && NETSIM_MODEL_DEBUG; // don't visual log if we have prediction disabled
+		
+		if (bDoVisualLog)
+		{
+			FVisualLoggingParameters VLogParameters(EVisualLoggingContext::LastConfirmed, SerializedFrame, EVisualLoggingLifetime::Persistent, TEXT("Serialized State"));
+			Driver->VisualLog(Buffers.Input[SerializedFrame], &SerializedSyncState, &SerializedAuxState, VLogParameters);
+		}
+
+		if (ClientSyncState)
+		{
+			// Existing ClientSyncState, log it before overwriting it
+			if (bDoVisualLog)
+			{
+				FVisualLoggingParameters VLogParameters(EVisualLoggingContext::FirstMispredicted, SerializedFrame, EVisualLoggingLifetime::Persistent, TEXT("Mispredicted State"));
+				Driver->VisualLog(Buffers.Input[SerializedFrame], ClientSyncState, Buffers.Aux[SerializedFrame], VLogParameters);
+			}
+		}
+		else
+		{
+			// No existing state, so create add it explicitly
+			ClientSyncState = Buffers.Sync.WriteFrame( SerializedFrame );
+		}
+
+		// Set client's sync state to the server version
+		check(ClientSyncState);
+		*ClientSyncState = SerializedSyncState;
+
+		// Set Client's aux state to the server version
+		*Buffers.Aux.WriteFrame(SerializedFrame) = SerializedAuxState;
+
+		const FNetworkSimTime RollbackDeltaTime = SerializedTime - Ticker.GetTotalProcessedSimulationTime();
+
+		// Set the canonical simulation time to what we received (we will advance it as we resimulate)
+		Ticker.SetTotalProcessedSimulationTime(SerializedTime, SerializedFrame);
+		Ticker.PendingFrame = SerializedFrame;
+		//Ticker.MaxAllowedFrame = FMath::Max(Ticker.MaxAllowedFrame, Ticker.PendingFrame); // Make sure this doesn't lag behind. This is the only place we should need to do this.
+
+		if (NetworkSimulationModelCVars::EnableLocalPrediction() == 0)
+		{
+			// If we aren't predicting at all, then we advanced the allowed sim time here, (since we aren't doing it in PreSimTick). This just keeps us constantly falling behind and not being able to toggle prediction on/off for debugging.
+			Ticker.SetTotalAllowedSimulationTime(SerializedTime);
+		}
+
+		// Tell dependent simulations to rollback
+		for (INetworkedSimulationModel* DependentSim : DependentSimulations)
+		{
+			DependentSim->BeginRollback(RollbackDeltaTime, SerializedFrame);
+		}
+
+		// Tell cue dispatch that we've rolledback as well
+		Buffers.CueDispatcher.NotifyRollback(SerializedTime);
+		
+		// Resimulate all user commands 
+		const int32 LastFrameToProcess = Ticker.MaxAllowedFrame;
+		for (int32 Frame = SerializedFrame; Frame <= LastFrameToProcess; ++Frame)
+		{
+			const int32 OutputFrame = Frame+1;
+
+			// Frame is the frame we are resimulating right now.
+			TInputCmd* ResimulateCmd  = Buffers.Input[Frame];
+			TAuxState* AuxState = Buffers.Aux[Frame];
+			TSyncState* PrevSyncState = Buffers.Sync[Frame];
+			TSyncState* NextSyncState = Buffers.Sync.WriteFrame(OutputFrame);
+			
+			check(ResimulateCmd);
+			check(PrevSyncState);
+			check(NextSyncState);
+			check(AuxState);
+
+			// Log out the Mispredicted state that we are about to overwrite.
+			if (NETSIM_MODEL_DEBUG)
+			{
+				FVisualLoggingParameters VLogParameters(Frame == LastFrameToProcess ? EVisualLoggingContext::LastMispredicted : EVisualLoggingContext::OtherMispredicted, Frame, EVisualLoggingLifetime::Persistent, TEXT("Resimulate Step: mispredicted"));
+				Driver->VisualLog(Buffers.Input[OutputFrame], Buffers.Sync[OutputFrame], Buffers.Aux[OutputFrame], VLogParameters);
+			}
+
+			// Do the actual update
+			{
+				TScopedSimulationTick<Model> UpdateScope(Ticker, Buffers.CueDispatcher, ESimulationTickContext::Resimulate, OutputFrame, ResimulateCmd->GetFrameDeltaTime());
+				Simulation->SimulationTick( 
+					{ ResimulateCmd->GetFrameDeltaTime(), Ticker },
+					{ *ResimulateCmd, *PrevSyncState, *AuxState },
+					{ *NextSyncState, Buffers.Aux.LazyWriter(OutputFrame), Buffers.CueDispatcher } );
+			}
+
+			// Log out the newly predicted state that we got.
+			if (NETSIM_MODEL_DEBUG)
+			{			
+				FVisualLoggingParameters VLogParameters(Frame == LastFrameToProcess ? EVisualLoggingContext::LastPredicted : EVisualLoggingContext::OtherPredicted, Frame, EVisualLoggingLifetime::Persistent, TEXT("Resimulate Step: repredicted"));
+				Driver->VisualLog(Buffers.Input[OutputFrame], Buffers.Sync[OutputFrame], Buffers.Aux[OutputFrame], VLogParameters);
+			}
+
+			// Tell dependent simulations to advance
+			for (INetworkedSimulationModel* DependentSim : DependentSimulations)
+			{
+				DependentSim->StepRollback(ResimulateCmd->GetFrameDeltaTime(), Frame, (Frame == LastFrameToProcess));
+			}
+		}
+	}
+
+	// --------------------------------------------------------------------
+	//	PreSimTick
+	// --------------------------------------------------------------------
+	void PreSimTick(TDriver* Driver, TNetworkSimBufferContainer<Model>& Buffers, TSimulationTicker<TTickSettings>& Ticker, const FNetSimTickParameters& TickParameters)
+	{
+		// If we have a reconcile fault, we cannot continue on with the simulation until it clears itself out. This effectively drops the input time and does not sample new inputs
+		if (bReconcileFaultDetected)
+		{
+			return;
+		}
+
+		if (TickParameters.bGenerateLocalInputCmds)
+		{
+			if (NetworkSimulationModelCVars::EnableLocalPrediction() > 0)
+			{
+				// Prediction: add simulation time and generate new commands
+				Ticker.GiveSimulationTime(TickParameters.LocalDeltaTimeSeconds);
+				TryGenerateLocalInput<Model>(Driver, Buffers, Ticker);
+			}
+			else
+			{
+				// Local prediction disabled: we must use a separate time accumulator to figure out when we should add more input cmds.
+				// Since we aren't processing the simulation locally, our core simulation time will only advance from network updates.
+				// (still, we need *something* to tell us when to generate a new command and what the delta time should be)
+				FNetworkSimTime NonPredictedInputTime;
+				NonPredictedInputTimeAccumulator.Accumulate(NonPredictedInputTime, TickParameters.LocalDeltaTimeSeconds);
+				if (NonPredictedInputTime.IsPositive())
+				{
+					GenerateLocalInputCmdAtFrame<Model>(Driver, Buffers, NonPredictedInputTime, Buffers.Input.HeadFrame() + 1);
+				}
+			}
+		}
+	}
+
+	void PostSimTick(TDriver* Driver, const TNetworkSimBufferContainer<Model>& Buffers, const TSimulationTicker<TTickSettings>& Ticker, const FNetSimTickParameters& TickParameters)
+	{
+		TBase::PostSimTick(Driver, Buffers, Ticker, TickParameters);
+
+		Buffers.CueDispatcher.ClearMaxDispatchTime();
+		Buffers.CueDispatcher.SetConfirmedTime(SerializedTime);
+	}
+	
+
+private:
+	
+	TSyncState SerializedSyncState;
+	TAuxState SerializedAuxState;
+	FNetworkSimTime SerializedTime;
+	int32 SerializedFrame = -1;
+
+	bool bPendingReconciliation = false;	// Reconciliation is pending: we need to reconcile state from the server that differs from the locally predicted state
+	bool bReconcileFaultDetected = false;	// A fault was detected: we received state from the server that we are unable to reconcile with locally predicted state
+
+	TRealTimeAccumulator<TTickSettings> NonPredictedInputTimeAccumulator; // for tracking input time in the non predictive case
+};
+
+/** Special replicator for debug buffer, this preserves the local buffer and receives into a replicator-owned buffer (we want these buffers to be distinct/not merged) */
+template<typename Model, int32 MaxNumElements=5, typename TBase=TRepController_Empty<Model>>
+struct TRepController_Debug : public TBase
+{
+	using typename TBase::TSimulation;
+	using typename TBase::TTickSettings;
+	using typename TBase::TDriver;
+	using typename TBase::TBufferTypes;
+
+	using TDebugState = typename TBase::TBufferTypes::TDebugState;
+	using TDebugBuffer = typename TNetworkSimBufferContainer<Model>::TDebugBuffer;
+
+	int32 GetProxyDirtyCount(TNetworkSimBufferContainer<Model>& Buffers) const 
+	{
+		return Buffers.Debug.GetDirtyCount() ^ (TBase::GetProxyDirtyCount(Buffers) << 2); 
+	}
+
+	void NetSerialize(const FNetSerializeParams& P, TNetworkSimBufferContainer<Model>& Buffers, TSimulationTicker<TTickSettings>& Ticker)
+	{
+		TBase::NetSerialize(P, Buffers, Ticker);
+		FArchive& Ar = P.Ar;
+
+		TDebugBuffer& Buffer = Ar.IsSaving() ? Buffers.Debug : ReceivedBuffer;
+		
+		uint8 SerializedNumElements = FMath::Min<uint8>(MaxNumElements, Buffer.Num());
+		Ar << SerializedNumElements;
+
+		const int32 HeadFrame = FNetworkSimulationSerialization::SerializeFrame(Ar, Buffer.HeadFrame());
+		const int32 StartingFrame = FMath::Max(0, HeadFrame - SerializedNumElements + 1);
+
+		if (Ar.IsLoading())
+		{
+			const int32 PrevHead = Buffer.HeadFrame();
+			if (PrevHead < StartingFrame && PrevHead >= 0)
+			{
+				// There is a gap in the stream. In some cases, we want this to be a "fault" and bubble up. We may want to synthesize state or maybe we just skip ahead.
+				UE_LOG(LogNetworkSim, Warning, TEXT("Fault: gap in received Debug buffer. PrevHead: %d. Received: %d-%d. Reseting previous buffer contents"), PrevHead, StartingFrame, HeadFrame);
+			}
+		}
+
+		for (int32 Frame = StartingFrame; Frame <= HeadFrame; ++Frame)
+		{
+			// This, as is, is bad. The intention is that these functions serialize multiple items in some delta compressed fashion.
+			// As is, we are just serializing the elements individually.
+			auto* Cmd = Ar.IsLoading() ? Buffer.WriteFrame(Frame) : Buffer[Frame];
+			Cmd->NetSerialize(P);
+		}
+	}
+
+	TDebugBuffer ReceivedBuffer;
+};
\ No newline at end of file
diff --git a/Engine/Plugins/Runtime/NetworkPrediction/Source/NetworkPrediction/Public/NetworkedSimulationModelReplicators.h b/Engine/Plugins/Runtime/NetworkPrediction/Source/NetworkPrediction/Public/NetworkedSimulationModelReplicators.h
deleted file mode 100644
