NetworkPrediction Plugin
-More iteration on rewindable cues after creating more testing and use cases and seeing some issues.
-Simplified user system callbacks for cues: just give a rollback ("undo") callback that is invoked in the right circumstances (rather than explicit confirmed/rollback/resimulate callbacks which require context to use properly)
-Proper support for rolling back "unconfirmed" cues (predicted cue was never received from server, independent of simulation rollback)
-Proper support for "resimulate rollback" of cues (cue was not repredicted during a resimulate)

#rnx

diff --git a/Engine/Plugins/Runtime/NetworkPrediction/Source/NetworkPrediction/Private/NetworkedSimulationModelCues.cpp b/Engine/Plugins/Runtime/NetworkPrediction/Source/NetworkPrediction/Private/NetworkedSimulationModelCues.cpp
--- a/Engine/Plugins/Runtime/NetworkPrediction/Source/NetworkPrediction/Private/NetworkedSimulationModelCues.cpp
+++ b/Engine/Plugins/Runtime/NetworkPrediction/Source/NetworkPrediction/Private/NetworkedSimulationModelCues.cpp
@@ -111,7 +111,7 @@ public:
 		DispatchTable.template RegisterType<FMockImpactCue>();
 	}
 
-	void HandleCue(FMockImpactCue& ImpactCue, const FNetSimCueSystemParamemters& SystemParameters)
+	void HandleCue(const FMockImpactCue& ImpactCue, const FNetSimCueSystemParamemters& SystemParameters)
 	{
 		UE_LOG(LogTemp, Warning, TEXT("Impact!"));
 	}
@@ -131,12 +131,12 @@ public:
 
 	using TestHandlerBase::HandleCue; // Required for us to "use" the HandleCue methods in our parent class
 
-	void HandleCue(FMockDamageCue& DamageCue, const FNetSimCueSystemParamemters& SystemParameters)
+	void HandleCue(const FMockDamageCue& DamageCue, const FNetSimCueSystemParamemters& SystemParameters)
 	{
 		UE_LOG(LogTemp, Warning, TEXT("HandleCue: Damage. SourceID: %d DamageType: %d HitLocation: %s"), DamageCue.SourceID, DamageCue.DamageType, *DamageCue.HitLocation.ToString());
 	}
 
-	void HandleCue(FMockHealingCue& HealingCue, const FNetSimCueSystemParamemters& SystemParameters)
+	void HandleCue(const FMockHealingCue& HealingCue, const FNetSimCueSystemParamemters& SystemParameters)
 	{
 		UE_LOG(LogTemp, Warning, TEXT("HandleCue: Healing. SourceID: %d Healing Amount: %d"), HealingCue.SourceID, HealingCue.HealingAmount);
 	}
diff --git a/Engine/Plugins/Runtime/NetworkPrediction/Source/NetworkPrediction/Public/NetworkedSimulationModelCues.h b/Engine/Plugins/Runtime/NetworkPrediction/Source/NetworkPrediction/Public/NetworkedSimulationModelCues.h
--- a/Engine/Plugins/Runtime/NetworkPrediction/Source/NetworkPrediction/Public/NetworkedSimulationModelCues.h
+++ b/Engine/Plugins/Runtime/NetworkPrediction/Source/NetworkPrediction/Public/NetworkedSimulationModelCues.h
@@ -25,7 +25,7 @@ enum class ESimulationTickContext : uint8
 enum class ENetSimCueReplicationTarget : uint8
 {
 	None,			// Do not replicate cue to anyone
-	Interpolators,	// Only replicate to /accept on clients that are interpolating (not running the simulation themselves)
+	Interpolators,	// Only replicate to / accept on clients that are interpolating (not running the simulation themselves)
 	All,			// Replicate to everyone
 };
 
@@ -87,6 +87,16 @@ struct TCueHandlerTraits : public TNetSimCueTraitsBase
 {
 };
 
+
+// Traits for TNetSimCueDispatcher
+struct NETWORKPREDICTION_API FCueDispatcherTraitsBase
+{
+	// Window for replicating a NetSimCue. That is, after a cue is invoked, it has ReplicationWindow time before it will be pruned.
+	// If a client does not get a net update for the sim in this window, they will miss the event.
+	static constexpr FNetworkSimTime ReplicationWindow = FNetworkSimTime::FromRealTimeMS(200);
+};
+template<typename T> struct TCueDispatcherTraits : public FCueDispatcherTraitsBase { };
+
 // ----------------------------------------------------------------------------------------------------------------------------------------------
 //	Wrapper
 // ----------------------------------------------------------------------------------------------------------------------------------------------
@@ -95,7 +105,7 @@ struct FNetSimCueWrapperBase
 {
 	virtual ~FNetSimCueWrapperBase() { }
 	virtual void NetSerialize(FArchive& Ar) = 0;
-	virtual bool NetUnique(const FNetSimCueWrapperBase* Other) const = 0;
+	virtual bool NetUnique(const void* OtherCueData) const = 0;
 	virtual void* CueData() const = 0;
 	virtual ENetSimCueReplicationTarget GetReplicationTarget() const = 0;
 	virtual bool Rollbackable() const = 0;
@@ -116,10 +126,10 @@ struct TNetSimCueWrapper : FNetSimCueWrapperBase
 		Instance.NetSerialize(Ar);
 	}
 
-	bool NetUnique(const FNetSimCueWrapperBase* Other) const override final
+	bool NetUnique(const void* OtherCueData) const override final
 	{
 		// Cue types must implement bool NetUnique(const TMyCueType& Other) const
-		return Instance.NetUnique(*((TCue*)Other->CueData()));
+		return Instance.NetUnique(*((const TCue*)OtherCueData));
 	}
 
 	void* CueData() const override final
@@ -146,16 +156,9 @@ struct TNetSimCueWrapper : FNetSimCueWrapperBase
 
 struct FNetSimCueCallbacks
 {
-	/** Rollback: the keyframe this cue was invoked on was rolled back (and will be resimulated) */
+	/** Callback to rollback any side effects of the cue. */
 	DECLARE_MULTICAST_DELEGATE(FOnRollback)
 	FOnRollback	OnRollback;
-	
-	/** Confirmed: the keyframe this was was predictively invoked on has been confirmed. The frame will not be rolled back now. */
-	DECLARE_MULTICAST_DELEGATE(FOnConfirmed)
-	FOnConfirmed OnConfirmed;
-
-	void ClearAll() { OnRollback.Clear(); OnConfirmed.Clear(); }
-	bool IsBound() const { return OnRollback.IsBound() || OnConfirmed.IsBound(); }
 };
 
 /** System parameters for NetSimCue events */
@@ -252,7 +255,7 @@ private:
 
 struct FSavedCue
 {
-	FSavedCue() = default;
+	FSavedCue(bool bInNetConfirmed) : bNetConfirmed(bInNetConfirmed) {}
 
 	FSavedCue(const FSavedCue&) = delete;
 	FSavedCue& operator=(const FSavedCue&) = delete;
@@ -260,9 +263,10 @@ struct FSavedCue
 	FSavedCue(FSavedCue&&) = default;
 	FSavedCue& operator=(FSavedCue&&) = default;
 	
-	FSavedCue(const FNetSimCueTypeId& InId, const FNetworkSimTime& InTime, FNetSimCueWrapperBase* Cue, const bool& bInAllowRollback)
-		: ID(InId), Time(InTime), CueInstance(Cue), bAllowRollback(bInAllowRollback)
+	FSavedCue(const FNetSimCueTypeId& InId, const FNetworkSimTime& InTime, const bool& bInAllowRollback, const bool& bInNetConfirmed, const bool& bInResimulates, FNetSimCueWrapperBase* Cue)
+		: ID(InId), Time(InTime), CueInstance(Cue), bAllowRollback(bInAllowRollback), bNetConfirmed(bInNetConfirmed), bResimulates(bInResimulates)
 	{
+		
 	}
 	
 	void NetSerialize(FArchive& Ar)
@@ -285,6 +289,7 @@ struct FSavedCue
 		}
 	}
 
+	// Test NetUniqueness against another saved cue
 	bool NetUnique(FSavedCue& OtherCue) const
 	{
 		if (ID != OtherCue.ID)
@@ -292,9 +297,21 @@ struct FSavedCue
 			return false;
 		}
 
-		return CueInstance->NetUnique(OtherCue.CueInstance.Get());
+		return CueInstance->NetUnique(OtherCue.CueInstance->CueData());
 	}
 
+	// Test NetUniqueness against an actual cue instance
+	template<typename TCue>
+	bool NetUnique(TCue& OtherCue) const
+	{
+		if (ID != TCue::ID)
+		{
+			return false;
+		}
+		return CueInstance->NetUnique(&OtherCue);
+	}
+
+
 	FString GetTypeName() const
 	{
 		return FGlobalCueTypeTable::Get().GetTypeName(ID);
@@ -303,8 +320,14 @@ struct FSavedCue
 	FNetSimCueTypeId ID = 0;
 	FNetworkSimTime Time;
 	TUniquePtr<FNetSimCueWrapperBase> CueInstance;
-	bool bDispatched = false;
-	bool bAllowRollback = false;
+	FNetSimCueCallbacks Callbacks;
+
+	bool bDispatched = false;		// Cue has been dispatched to the local handler. Never dispatch twice.
+	bool bAllowRollback = false;	// Cue supports rolling back. ie., we should pass the user valid FNetSimCueCallbacks rollback callback (note this is "this saved cue" specifically, not "this cue type". E.g, on authority, bAllowRollback is always false)
+	bool bNetConfirmed = false;		// This cue has been net confirmed, meaning we received it directly via replication or we received a replicated cue that matched this one that was locally predicted.
+	
+	bool bResimulates = false;				 // Whether this cue supports invocation during resimulation. Needed to set bResimulatePendingRollback
+	bool bPendingResimulateRollback = false; // Rollback is pending due to resimulation (unless CUe is matched with an Invocation during the resimulate)
 };
 
 
@@ -333,7 +356,8 @@ public:
 			// The actual Dispatch func that gets invoked
 			CueTypeInfo.Dispatch = [](FNetSimCueWrapperBase* Cue, TCueHandler& Handler, const FNetSimCueSystemParamemters& SystemParameters)
 			{
-				Handler.HandleCue( *static_cast<TCue*>(Cue->CueData()), SystemParameters );
+				// If you are finding compile errors here, you may be missing a ::HandleCue implementation for a specific cue type that your handler has registered with
+				Handler.HandleCue( *static_cast<const TCue*>(Cue->CueData()), SystemParameters );
 			};
 		};
 		
@@ -347,12 +371,12 @@ public:
 		}
 	}
 
-	void Dispatch(FSavedCue& SavedCue, TCueHandler& Handler, const FNetSimCueSystemParamemters& SystemParameters)
+	void Dispatch(FSavedCue& SavedCue, TCueHandler& Handler, const FNetworkSimTime& DispatchTime)
 	{
 		if (FCueTypeInfo* TypeInfo = CueTypeInfoMap.Find(SavedCue.ID))
 		{
 			check(TypeInfo->Dispatch);
-			TypeInfo->Dispatch(SavedCue.CueInstance.Get(), Handler, SystemParameters);
+			TypeInfo->Dispatch(SavedCue.CueInstance.Get(), Handler, {DispatchTime - SavedCue.Time, SavedCue.bAllowRollback ? &SavedCue.Callbacks : nullptr });
 		}
 	}
 
@@ -394,26 +418,91 @@ struct FNetSimCueDispatcher
 	{
 		if (EnsureValidContext())
 		{
-			if ((TCueHandlerTraits<T>::InvokeMask & (FNetSimCueTypeId)Context.TickContext) > 0)
+			if ((TCueHandlerTraits<T>::InvokeMask & (uint8)Context.TickContext) > 0)
 			{
-				// There is an implicit contract that we will invoke events going forward OR will receive explicit rollback notification (TODO)
-				ensure(SavedCues.Num() == 0 || SavedCues.Last().Time <= Context.CurrentSimTime);
+				/*
+				{
+					// Whether we go in the transient list. Transient cues are dumped after dispatching (not saved over multiple frames for uniqueness comparisons during Invoke or NetSerialize)
+					const bool bTransient = (Context.TickContext == ESimulationTickContext::Authority && TCueHandlerTraits<T>::ReplicationTarget == ENetSimCueReplicationTarget::None)	  // Transient on authority if it doesn't replicate
+						|| ((TCueHandlerTraits<T>::InvokeMask & (uint8)ESimulationTickContext::Resimulate) == 0 && TCueHandlerTraits<T>::ReplicationTarget != ENetSimCueReplicationTarget::All); // Transient elsewhere if it won't replicate to us and it won't be invoked during resimulates
 
-				// Whether this cue should be dispatched with rollback callbacks. This is a trait of the cue type + non authority (authority will never rollback)
-				const bool bAllowRollback = TCueHandlerTraits<T>::Rollbackable && Context.TickContext != ESimulationTickContext::Authority; 
+					// Whether this cue should be dispatched with rollback callbacks. This is a trait of the cue type + context (e.g, authority will never rollback)
+					const bool bAllowRollback = TCueHandlerTraits<T>::Rollbackable && !bTransient && Context.TickContext != ESimulationTickContext::Authority;
 
-				// Whether we go in the transient list. Transient cues are dumped after dispatching (not saved over multiple frames for uniqueness comparisons)
-				const bool bTransient = TCueHandlerTraits<T>::ReplicationTarget == ENetSimCueReplicationTarget::None;
+					// Is this already confirmed? (it is if we can't roll it back, or if it won't be replicated to us)
+					const bool bNetConfirmed = !bAllowRollback || TCueHandlerTraits<T>::ReplicationTarget != ENetSimCueReplicationTarget::All;
+				}
 
-				UE_LOG(LogNetSimCues, Log, TEXT("Invoking Cue %s. Transient: %d. Mask: %d. ReplicationTarget: %d"), *FGlobalCueTypeTable::Get().GetTypeName(T::ID), bTransient, TCueHandlerTraits<T>::InvokeMask, TCueHandlerTraits<T>::ReplicationTarget);
-				if (bTransient)
+
+				{
+					// Whether this cue should be dispatched with rollback callbacks. Authority never rolls back, clients rollback if they plan to invoke during resims or if the cue is replicated to everyone
+					const bool bRequiresRollbackForReplication = Context.TickContext != ESimulationTickContext::Authority && (TCueHandlerTraits<T>::ReplicationTarget == ENetSimCueReplicationTarget::All);
+					const bool bRequiresRollbackForResimulate = Context.TickContext != ESimulationTickContext::Authority && (TCueHandlerTraits<T>::InvokeMask & (uint8)ESimulationTickContext::Resimulate);
+					const bool bAllowRollback = bRequiresRollbackForReplication || bRequiresRollbackForResimulate;
+
+					// Whether we go in the transient list. Transient cues are dumped after dispatching (not saved over multiple frames for uniqueness comparisons during Invoke or NetSerialize)
+					const bool bTransient = !bAllowRollback &&
+						(Context.TickContext == ESimulationTickContext::Authority && (TCueHandlerTraits<T>::ReplicationTarget == ENetSimCueReplicationTarget::None)) ||
+						(Context.TickContext != ESimulationTickContext::Authority && (TCueHandlerTraits<T>::ReplicationTarget != ENetSimCueReplicationTarget::All));
+
+					const bool bNetConfirmed = !bAllowRollback || TCueHandlerTraits<T>::ReplicationTarget != ENetSimCueReplicationTarget::All;
+				}
+				*/
+
+				const bool bSupportsResimulate = (TCueHandlerTraits<T>::InvokeMask & (uint8)ESimulationTickContext::Resimulate) > 0;
+
+				bool bAllowRollback = false;	// Whether this cue should be dispatched with rollback callbacks.
+				bool bTransient = false;		// Whether we go in the transient list. Transient cues are dumped after dispatching (not saved over multiple frames for uniqueness comparisons during Invoke or NetSerialize)
+				bool bNetConfirmed = false;		// Is this already confirmed? (we should not look to undo it if we don't get it confirmed from the server)
+				if (Context.TickContext == ESimulationTickContext::Authority)
 				{
-					TransientCues.Emplace(T::ID, Context.CurrentSimTime, new TNetSimCueWrapper<T>(Forward<ArgsType>(Args)...), bAllowRollback);
+					// Authority: never rolls back, is already confirmed, and can treat cue as transient if it doesn't have to replicate it
+					bAllowRollback = false;
+					bTransient = TCueHandlerTraits<T>::ReplicationTarget == ENetSimCueReplicationTarget::None;
+					bNetConfirmed = true;
 				}
 				else
 				{
-					SavedCues.Emplace(T::ID, Context.CurrentSimTime, new TNetSimCueWrapper<T>(Forward<ArgsType>(Args)...), bAllowRollback);
+					// Everyone else that is running the simulation (since this in ::Invoke which is called from within the simulation)
+					// Rollback if it will replicate to us or if we plan to invoke this cue during resimulates. Transient and confirmed follow directly from this in the non authority case.
+					bAllowRollback = (TCueHandlerTraits<T>::ReplicationTarget == ENetSimCueReplicationTarget::All) || bSupportsResimulate;
+					bTransient = !bAllowRollback;
+					bNetConfirmed = !bAllowRollback;
+				}
+
+				ensure(!(bTransient && bAllowRollback)); // this combination cannot happen: we can't be transient and support rollback (but we can be transient without supporting rollback)
+				ensure(!(bNetConfirmed && bAllowRollback)); // a confirmed cue shouldn't be rolled back.
+
+				// In resimulate case, we have to see if we already predicted it
+				if (Context.TickContext == ESimulationTickContext::Resimulate)
+				{
+					ensure(RollbackTime.IsPositive() && RollbackTime <= Context.CurrentSimTime);
+					
+					// Since we haven't constructed the cue yet, we can't test for uniqueness!
+					// So, create one on the stack. If we let it through we can move it to the appropriate buffer
+					// (not as nice as the non resimulate path, but better than allocating a new FSavedCue+TNetSimCueWrapper and then removing them)
+					T NewCue(MoveTempIfPossible(Forward<ArgsType>(Args))...);
+
+					for (FSavedCue& ExistingCue : SavedCues)
+					{
+						if (RollbackTime <= ExistingCue.Time && ExistingCue.NetUnique(NewCue) == false)
+						{
+							// We've matched with an already predicted cue, so suppress this invocation and don't undo the predicted cue
+							ExistingCue.bPendingResimulateRollback = false;
+							UE_LOG(LogNetSimCues, Log, TEXT("Resimulated Cue %s matched existing cue. Suppressing Invocation."), *FGlobalCueTypeTable::Get().GetTypeName(T::ID));
+							return;
+						}
+					}
+
+					UE_LOG(LogNetSimCues, Log, TEXT("Invoking Cue %s during resimulate. Context: %d. Transient: %d. bAllowRollback: %d. bNetConfirmed: %d."), *FGlobalCueTypeTable::Get().GetTypeName(T::ID), Context.TickContext, bTransient, bAllowRollback, bNetConfirmed);
+					GetBuffer(bTransient).Emplace(T::ID, Context.CurrentSimTime, bAllowRollback, bNetConfirmed, bSupportsResimulate, new TNetSimCueWrapper<T>(MoveTempIfPossible(NewCue)));
 				}
+				else
+				{
+					// Not resimulate case is simple: construct the new cue emplace in the appropriate list
+					UE_LOG(LogNetSimCues, Log, TEXT("Invoking Cue %s. Context: %d. Transient: %d. bAllowRollback: %d. bNetConfirmed: %d."), *FGlobalCueTypeTable::Get().GetTypeName(T::ID), Context.TickContext, bTransient, bAllowRollback, bNetConfirmed);
+					GetBuffer(bTransient).Emplace(T::ID, Context.CurrentSimTime, bAllowRollback, bNetConfirmed, bSupportsResimulate, new TNetSimCueWrapper<T>(Forward<ArgsType>(Args)...));	
+				}				
 			}
 			else
 			{
@@ -438,19 +527,12 @@ protected:
 	
 	TArray<FSavedCue> SavedCues;		// Cues that must be saved for some period of time, either for replication or for uniqueness testing
 	TArray<FSavedCue> TransientCues;	// Cues that are dispatched on this frame and then forgotten about
+	TArray<FSavedCue>& GetBuffer(const bool& bTransient) { return bTransient ? TransientCues : SavedCues; }
 
 	FContext Context;
+	FNetworkSimTime RollbackTime;		// Time of last rollback, reset after dispatching
 };
 
-// Traits for TNetSimCueDispatcher
-struct NETWORKPREDICTION_API FCueDispatcherTraitsBase
-{
-	// Window for replicating a NetSimCue. That is, after a cue is invoked, it has ReplicationWindow time before it will be pruned.
-	// If a client does not get a net update for the sim in this window, they will miss the event.
-	static constexpr FNetworkSimTime ReplicationWindow = FNetworkSimTime::FromRealTimeMS(200);
-};
-template<typename T> struct TCueDispatcherTraits : public FCueDispatcherTraitsBase { };
-
 // Templated cue dispatcher that can be specialized per networking model definition. This is what the system actually uses internally, but is not exposed to user code.
 template<typename Model=void>
 struct TNetSimCueDispatcher : public FNetSimCueDispatcher
@@ -480,7 +562,7 @@ struct TNetSimCueDispatcher : public FNetSimCueDispatcher
 
 			for (int32 CueIdx=0; CueIdx < NumCues; ++CueIdx)
 			{
-				FSavedCue SerializedCue;
+				FSavedCue SerializedCue(true);
 				SerializedCue.NetSerialize(Ar);
 
 				// Decide if we should accept the cue:
@@ -491,22 +573,27 @@ struct TNetSimCueDispatcher : public FNetSimCueDispatcher
 					UE_LOG(LogNetSimCues, Log, TEXT("Discarding replicated NetSimCue %s intended for interpolators."), *SerializedCue.GetTypeName());
 					continue;
 				}
-
+				
 				// Uniqueness: have we already received/predicted it?
+				// Note: we are basically ignoring invocation time when matching right now. This could potentially be a trait of the cue if needed.
+				// This could create issues if a cue is invoked several times in quick succession, but that can be worked around with arbitrary counter parameters on the cue itself (to force NetUniqueness)
 				bool bUniqueCue = true;
 				for (int32 ExistingIdx=0; ExistingIdx < StartingNum; ++ExistingIdx)
 				{
-					if (SerializedCue.NetUnique(SavedCues[ExistingIdx]) == false)
+					FSavedCue& ExistingCue = SavedCues[ExistingIdx];
+					if (SerializedCue.NetUnique(ExistingCue) == false)
 					{
 						// These cues are not unique ("close enough") so we are skipping receiving this one
 						UE_LOG(LogNetSimCues, Log, TEXT("Discarding replicated NetSimCue %s because we've already processed it."), *SerializedCue.GetTypeName());
 						bUniqueCue = false;
+						ExistingCue.bNetConfirmed = true;
 						break;
 					}
 				}
 
 				if (bUniqueCue)
 				{
+					UE_LOG(LogNetSimCues, Log, TEXT("Received NetUnique Cue: %s. (Num replicated cue sent this bunch: %d)"), *SerializedCue.GetTypeName(), NumCues);
 					SavedCues.Emplace(MoveTemp(SerializedCue));
 				}
 			}
@@ -525,26 +612,41 @@ struct TNetSimCueDispatcher : public FNetSimCueDispatcher
 		int32 SavedCuePruneIdx = -1;
 
 		// ------------------------------------------------------------------------
-		// User callbacks
+		// Rollback events if necessary
+		//	Fixme - this code was written for clarity, it could be sped up considerably by taking advantage of sorting by time, or keeping acceleration lists for this type of pruning
 		// ------------------------------------------------------------------------
-		
-		ensure(UserConfirmedTime.IsPositive() || InvocationCallbacks.Num() == 0); // Should not have callbacks when PruneTime is zero
 
-		for (FInvocationCallbackContainer& Callback : InvocationCallbacks)
+		if (UserConfirmedTime.IsPositive())
 		{
-			if (RollbackTime.IsPositive() && Callback.InvocationTime > RollbackTime)
+			// Look for cues that should have been matched by now, but were not
+			for (auto It = SavedCues.CreateIterator(); It; ++It)
 			{
-				Callback.UserCallbacks.OnRollback.Broadcast();
-				Callback.UserCallbacks.OnRollback.Clear();
+				FSavedCue& SavedCue = *It;
+				if (!SavedCue.bNetConfirmed && SavedCue.Time <= ConfirmedTime)
+				{
+					UE_LOG(LogNetSimCues, Log, TEXT("Calling OnRollback for SavedCue NetSimCue %s. Cue has not been matched but it <= ConfirmedTime."), *SavedCue.GetTypeName());
+					SavedCue.Callbacks.OnRollback.Broadcast();
+					SavedCues.RemoveAt(It.GetIndex(), 1, false);
+				}
 			}
+		}
 
-			if (Callback.UserCallbacks.OnConfirmed.IsBound() && Callback.InvocationTime <= ConfirmedTime)
+		if (RollbackTime.IsPositive())
+		{
+			for (auto It = SavedCues.CreateIterator(); It; ++It)
 			{
-				Callback.UserCallbacks.OnConfirmed.Broadcast();
-				Callback.UserCallbacks.ClearAll(); // Confirmed is the end of the road. Clear all and removal will happen at the bottom of this function
+				FSavedCue& SavedCue = *It;
+				if (SavedCue.bPendingResimulateRollback)
+				{
+					// Unmatched cue whose time has passed, time to rollback
+					UE_LOG(LogNetSimCues, Log, TEXT("Calling OnRollback for SavedCue NetSimCue %s. Cue was not matched during a resimulate. "), *SavedCue.GetTypeName());
+					SavedCue.Callbacks.OnRollback.Broadcast();
+					SavedCues.RemoveAt(It.GetIndex(), 1, false);
+				}
 			}
+
+			RollbackTime.Reset();
 		}
-		RollbackTime.Reset();
 
 		// ------------------------------------------------------------------------
 		// Dispatch (call ::HandleCue)
@@ -556,23 +658,27 @@ struct TNetSimCueDispatcher : public FNetSimCueDispatcher
 			if ( SavedCue.Time <= SavedCuePruneTime)
 			{
 				SavedCuePruneIdx = SavedCueIdx;
-			}
-
-			const bool bWithold = SavedCue.Time > DispatchTime;
-			UE_CLOG(bWithold, LogNetSimCues, Log, TEXT("Withholding Cue %s. %s > %s"), *SavedCue.GetTypeName(), *SavedCue.Time.ToString(), *DispatchTime.ToString());
+			}			
 
-			if (SavedCue.bDispatched == false && !bWithold)
+			if (!SavedCue.bDispatched)
 			{
-				UE_LOG(LogNetSimCues, Log, TEXT("Dispatching NetSimCue %s."), *SavedCue.GetTypeName());
-				SavedCue.bDispatched = true;
-				TCueDispatchTable<T>::Get().Dispatch(SavedCue, Handler, {DispatchTime - SavedCue.Time, GetUserCallbackPtr(SavedCue)} );
+				if (SavedCue.Time <= DispatchTime)
+				{
+					UE_LOG(LogNetSimCues, Log, TEXT("Dispatching NetSimCue %s."), *SavedCue.GetTypeName());
+					SavedCue.bDispatched = true;
+					TCueDispatchTable<T>::Get().Dispatch(SavedCue, Handler,DispatchTime);
+				}
+				else
+				{
+					UE_LOG(LogNetSimCues, Log, TEXT("Withholding Cue %s. %s > %s"), *SavedCue.GetTypeName(), *SavedCue.Time.ToString(), *DispatchTime.ToString());
+				}
 			}
 		}
 
 		for (FSavedCue& TransientCue : TransientCues)
 		{
 			UE_LOG(LogNetSimCues, Log, TEXT("Dispatching transient NetSimCue %s."), *TransientCue.GetTypeName());
-			TCueDispatchTable<T>::Get().Dispatch(TransientCue, Handler, {DispatchTime - TransientCue.Time, GetUserCallbackPtr(TransientCue)} );
+			TCueDispatchTable<T>::Get().Dispatch(TransientCue, Handler, DispatchTime);
 		}
 		TransientCues.Reset();
 
@@ -593,17 +699,6 @@ struct TNetSimCueDispatcher : public FNetSimCueDispatcher
 
 			SavedCues.RemoveAt(0, SavedCuePruneIdx+1, false);
 		}
-
-		// Remove any unbounded callbacks
-		for (int32 idx=InvocationCallbacks.Num()-1; idx >= 0; --idx)
-		{
-			FInvocationCallbackContainer& Callback = InvocationCallbacks[idx];
-			if (!Callback.UserCallbacks.IsBound())
-			{
-				UE_LOG(LogNetSimCues, Log, TEXT("Pruned Callbacks at time=%s"), *Callback.InvocationTime.ToString());
-				InvocationCallbacks.RemoveAt(idx, 1, false);
-			}
-		}
 	}
 
 	// Tell dispatcher that we've rolled back to a new simulation time (resimulate steps to follow, most likely)
@@ -616,8 +711,18 @@ struct TNetSimCueDispatcher : public FNetSimCueDispatcher
 		}
 		else
 		{
+			// Two rollbacks could happen in between DispatchCueRecord calls. That is ok as long as the subsequent rollbacks are further ahead in simulation time
 			ensure(RollbackTime < InRollbackTime);
 		}
+
+		// Mark all cues that support invocation during resimulation as pending rollback (unless they match in an Invoke)
+		for (FSavedCue& SavedCue : SavedCues)
+		{
+			if (SavedCue.bResimulates && SavedCue.Time >= InRollbackTime)
+			{
+				SavedCue.bPendingResimulateRollback = true;
+			}
+		}
 	}
 
 	// Push/pop simulation context.
@@ -635,49 +740,6 @@ private:
 
 	mutable FNetworkSimTime UserMaxDispatchTime;	// (If set) Max time of a saved cue that can be dispatched. Needed for delaying interpolation / client side buffering.
 	mutable FNetworkSimTime UserConfirmedTime;		// (If set) latest confirmed simulation time. If not set, we assume we are authority and will never rollback.
-	
-	struct FInvocationCallbackContainer
-	{
-		FInvocationCallbackContainer(const FNetworkSimTime& InTime) : InvocationTime(InTime) { }
-
-		FNetworkSimTime InvocationTime;
-		FNetSimCueCallbacks UserCallbacks;
-	};
-
-	TArray<FInvocationCallbackContainer> InvocationCallbacks;
-	FNetworkSimTime RollbackTime;
-
-	FNetSimCueCallbacks* GetUserCallbackPtr(const FSavedCue& SavedCue)
-	{
-		FNetSimCueCallbacks* CallbackPtr = nullptr;
-		if (SavedCue.bAllowRollback)
-		{
-			for (int32 CallbacksIdx = InvocationCallbacks.Num()-1; CallbacksIdx >=0; --CallbacksIdx)
-			{
-				FInvocationCallbackContainer& CallbackContainer = InvocationCallbacks[CallbacksIdx];
-
-				if (CallbackContainer.InvocationTime == SavedCue.Time)
-				{
-					CallbackPtr = &CallbackContainer.UserCallbacks;
-					break;
-				}
-				if (CallbackContainer.InvocationTime < SavedCue.Time)
-				{
-					UE_LOG(LogNetSimCues, Log, TEXT("Added Callbacks for time=%s at idx=%d"), *SavedCue.Time.ToString(), CallbacksIdx);
-					CallbackPtr = &InvocationCallbacks.EmplaceAt_GetRef( CallbacksIdx+1, SavedCue.Time ).UserCallbacks;
-					break;
-				}
-			}
-
-			if (!CallbackPtr)
-			{
-				UE_LOG(LogNetSimCues, Log, TEXT("Added Callbacks for time=%s"));
-				CallbackPtr = &InvocationCallbacks.Emplace_GetRef( SavedCue.Time ).UserCallbacks;
-			}
-
-		}
-		return CallbackPtr;
-	}
 };
 
 // ------------------------------------------------------------------------------------------------------------------------------------------------------------------------------------------------
diff --git a/Engine/Plugins/Runtime/NetworkPredictionExtras/Source/NetworkPredictionExtras/Private/MockAbilitySimulation.cpp b/Engine/Plugins/Runtime/NetworkPredictionExtras/Source/NetworkPredictionExtras/Private/MockAbilitySimulation.cpp
--- a/Engine/Plugins/Runtime/NetworkPredictionExtras/Source/NetworkPredictionExtras/Private/MockAbilitySimulation.cpp
+++ b/Engine/Plugins/Runtime/NetworkPredictionExtras/Source/NetworkPredictionExtras/Private/MockAbilitySimulation.cpp
@@ -17,10 +17,13 @@ namespace MockAbilityCVars
 	NETSIM_DEVCVAR_SHIPCONST_FLOAT(DashAcceleration, 100000.f, "mockability.DashAcceleration", "Acceleration when dashing.");
 
 
-	NETSIM_DEVCVAR_SHIPCONST_INT(BlinkCueType, 0, "mockability.BlinkCueType", "0=default. 1=weak. 2=");
+	NETSIM_DEVCVAR_SHIPCONST_INT(BlinkCueType, 4, "mockability.BlinkCueType", "0=Skip. 1=weak. 2=ReplicatedNonPredicted, 3=ReplicatedXOrPredicted, 4=Strong");
+	NETSIM_DEVCVAR_SHIPCONST_INT(BlinkWarmupMS, 750, "mockability.BlinkWarmupMS", "Duration in MS of blink warmup period");
 }
 
+NETSIMCUE_REGISTER(FMockAbilityBlinkActivateCue, TEXT("MockAbilityBlinkActivate"));
 NETSIMCUE_REGISTER(FMockAbilityBlinkCue, TEXT("MockAbilityBlink"));
+
 NETSIMCUE_REGISTER(FMockAbilityBlinkCue_Weak, TEXT("MockAbilityBlink_Weak"));
 NETSIMCUE_REGISTER(FMockAbilityBlinkCue_ReplicatedNonPredicted, TEXT("MockAbilityBlink_RepNonPredicted"));
 NETSIMCUE_REGISTER(FMockAbilityBlinkCue_ReplicatedXOrPredicted, TEXT("MockAbilityBlink_RepXOrPredicted"));
@@ -84,12 +87,26 @@ void FMockAbilitySimulation::SimulationTick(const FNetSimTimeStep& TimeStep, con
 	// -------------------------------------------------------------------------
 
 	static float BlinkCost = 25.f;
-	static int16 BlinkWarmupMS = 750; 
+	const int16 BlinkWarmupMS = MockAbilityCVars::BlinkWarmupMS();
+
+	auto GetBlinkDestination = [&LocalSync]()
+	{
+		static float BlinkDist = 1000.f;
+		const FVector DestLocation = LocalSync.Location + LocalSync.Rotation.RotateVector( FVector(BlinkDist, 0.f, 0.f) );
+		return DestLocation;
+	};
 
 	const bool bBlinkActivate = (Input.Cmd.bBlinkPressed && Input.Sync.Stamina > BlinkCost && bAllowNewActivations);
 	if (bBlinkActivate)
 	{
+
+
 		LocalAux.BlinkWarmupLeft = BlinkWarmupMS;
+
+		// Invoke a cue to telegraph where the blink will land. This is making the assumption the handler wouldn't either want to or be able to derive the blink destination from the current state alone
+		// The randomValue being calculated here is purposefully to cause mis prediction of the cue, so that we can demonstrate rollback -> resimulate can do the correction seemlessly
+		UE_LOG(LogTemp, Warning, TEXT("Invoking FMockAbilityBlinkActivateCue from sim"));
+		Output.CueDispatch.Invoke<FMockAbilityBlinkActivateCue>( GetBlinkDestination(), FMath::Rand() % 255 );
 	}
 
 	if (LocalAux.BlinkWarmupLeft > 0)
@@ -106,8 +123,7 @@ void FMockAbilitySimulation::SimulationTick(const FNetSimTimeStep& TimeStep, con
 
 		if (NewBlinkWarmupLeft <= 0)
 		{
-			static float BlinkDist = 1000.f;
-			FVector DestLocation = LocalSync.Location + LocalSync.Rotation.RotateVector( FVector(BlinkDist, 0.f, 0.f) );
+			const FVector DestLocation = GetBlinkDestination();
 			AActor* OwningActor = UpdatedComponent->GetOwner();
 			check(OwningActor);
 
@@ -127,12 +143,11 @@ void FMockAbilitySimulation::SimulationTick(const FNetSimTimeStep& TimeStep, con
 
 				// Invoke a NetCue for the blink. This is essentially capturing the start/end location so that all receivers of the event
 				// get the exact coordinates (maybe overkill in practice but key is that we have data that we want to pass out via an event)
-
-				UE_LOG(LogTemp, Warning, TEXT("INVOKING CUE EMPLACE"));
+				
 				switch(MockAbilityCVars::BlinkCueType()) // Only for dev/testing. Not a normal setup.
 				{
 				case 0:
-					Output.CueDispatch.Invoke<FMockAbilityBlinkCue>( Input.Sync.Location, DestLocation );
+					// Skip o purpose
 					break;
 				case 1:
 					Output.CueDispatch.Invoke<FMockAbilityBlinkCue_Weak>( Input.Sync.Location, DestLocation );
@@ -318,13 +333,37 @@ void UMockFlyingAbilityComponent::VisualLog(const FMockAbilityInputCmd* Input, c
 	Super::VisualLog(Input, Sync, Aux, SystemParameters);
 }
 
+float UMockFlyingAbilityComponent::GetBlinkWarmupTimeSeconds() const
+{
+	return FNetworkSimTime::FromRealTimeMS( MockAbilityCVars::BlinkWarmupMS() ).ToRealTimeSeconds();
+}
+
 // ---------------------------------------------------------------------------------
 
+void UMockFlyingAbilityComponent::HandleCue(const FMockAbilityBlinkActivateCue& BlinkCue, const FNetSimCueSystemParamemters& SystemParameters)
+{
+	FString RoleStr = GetOwnerRole() == ROLE_Authority ? TEXT("Server") : TEXT("Client");
+	UE_LOG(LogTemp, Display, TEXT("[%s] BlinkActivatedCue!"), *RoleStr);
+
+	this->OnBlinkActivateEvent.Broadcast(BlinkCue.Destination, BlinkCue.RandomType, SystemParameters.TimeSinceInvocation.ToRealTimeSeconds());
+	
+	if (SystemParameters.Callbacks)
+	{
+		UE_LOG(LogTemp, Display, TEXT("  System Callbacks available!"));
+
+		SystemParameters.Callbacks->OnRollback.AddLambda([RoleStr, this]()
+		{
+			UE_LOG(LogTemp, Display, TEXT("  %s BlinkActivatedCue Rollback!"), *RoleStr);
+			this->OnBlinkActivateEventRollback.Broadcast();
+		});
+	}
+}
+
 static float BlinkCueDuration = 1.f;
 static FAutoConsoleVariableRef CVarBindAutomatically(TEXT("NetworkPredictionExtras.FlyingPawn.BlinkCueDuration"),
 	BlinkCueDuration, TEXT("Duration of BlinkCue"), ECVF_Default);
 
-void UMockFlyingAbilityComponent::HandleCue(FMockAbilityBlinkCue& BlinkCue, const FNetSimCueSystemParamemters& SystemParameters)
+void UMockFlyingAbilityComponent::HandleCue(const FMockAbilityBlinkCue& BlinkCue, const FNetSimCueSystemParamemters& SystemParameters)
 {
 	FString RoleStr = GetOwnerRole() == ROLE_Authority ? TEXT("Server") : TEXT("Client");
 
@@ -337,20 +376,14 @@ void UMockFlyingAbilityComponent::HandleCue(FMockAbilityBlinkCue& BlinkCue, cons
 	float Duration = FMath::Max<float>(0.1f, BlinkCueDuration - SystemParameters.TimeSinceInvocation.ToRealTimeSeconds());
 	DrawDebugLine(GetWorld(), BlinkCue.StartLocation, BlinkCue.StopLocation, FColor::Red, false, Duration);
 
-
 	if (SystemParameters.Callbacks)
 	{
 		UE_LOG(LogTemp, Display, TEXT("  System Callbacks available!"));
 
-		SystemParameters.Callbacks->OnRollback.AddLambda([RoleStr]()
+		SystemParameters.Callbacks->OnRollback.AddLambda([RoleStr, this]()
 		{
 			UE_LOG(LogTemp, Display, TEXT("  %s BlinkCue Rollback!"), *RoleStr);
 		});
-
-		SystemParameters.Callbacks->OnConfirmed.AddLambda([RoleStr]()
-		{
-			UE_LOG(LogTemp, Display, TEXT("  %s BlinkCue Confirmed!"), *RoleStr);
-		});
-
 	}
-}
\ No newline at end of file
+}
+
diff --git a/Engine/Plugins/Runtime/NetworkPredictionExtras/Source/NetworkPredictionExtras/Public/MockAbilitySimulation.h b/Engine/Plugins/Runtime/NetworkPredictionExtras/Source/NetworkPredictionExtras/Public/MockAbilitySimulation.h
--- a/Engine/Plugins/Runtime/NetworkPredictionExtras/Source/NetworkPredictionExtras/Public/MockAbilitySimulation.h
+++ b/Engine/Plugins/Runtime/NetworkPredictionExtras/Source/NetworkPredictionExtras/Public/MockAbilitySimulation.h
@@ -125,9 +125,41 @@ struct FMockAbilityAuxState : public FFlyingMovementAuxState
 // MockAbility NetSimCues - events emitted by the sim
 // -------------------------------------------------------
 
+// Cue for blink activation (the moment the ability starts)
+struct FMockAbilityBlinkActivateCue
+{
+	FMockAbilityBlinkActivateCue() = default;
+	FMockAbilityBlinkActivateCue(const FVector& InDestination, uint8 InRandomType)
+		: Destination(InDestination), RandomType(InRandomType) { }
+
+	NETSIMCUE_BODY();
+
+	FVector_NetQuantize10 Destination;
+	uint8 RandomType; // Random value used to color code the effect. This is the test/prove out mispredictions
+
+	void NetSerialize(FArchive& Ar)
+	{
+		bool b = false;
+		Destination.NetSerialize(Ar, nullptr, b);
+		Ar << RandomType;
+	}
+	
+	bool NetUnique(const FMockAbilityBlinkActivateCue& Other) const
+	{
+		const float ErrorTolerance = 1.f;
+		return !Destination.Equals(Other.Destination, ErrorTolerance) || RandomType != Other.RandomType;
+	}
+};
+
+template<>
+struct TCueHandlerTraits<FMockAbilityBlinkActivateCue> : public TNetSimCueTraits_Strong { };
+
+// ----------------------------------------------------------------------------------------------
+
+
 #define LOG_BLINK_CUE 1 // During development, its useful to sanity check that we aren't doing more construction or moves than we expect
 
-// Cue for blink activation.
+// Cue for blink (the moment the teleport happens)
 struct FMockAbilityBlinkCue
 {
 	FMockAbilityBlinkCue()
@@ -135,7 +167,7 @@ struct FMockAbilityBlinkCue
 		UE_CLOG(LOG_BLINK_CUE, LogTemp, Warning, TEXT("  Default Constructor 0x%X"), this);
 	}
 
-	FMockAbilityBlinkCue(FVector Start, FVector Stop) : StartLocation(Start), StopLocation(Stop)
+	FMockAbilityBlinkCue(const FVector& Start, const FVector& Stop) : StartLocation(Start), StopLocation(Stop)
 	{ 
 		UE_CLOG(LOG_BLINK_CUE, LogTemp, Warning, TEXT("  Custom Constructor 0x%X"), this);
 	}
@@ -205,6 +237,7 @@ struct FMockAbilityCueSet
 	template<typename TDispatchTable>
 	static void RegisterNetSimCueTypes(TDispatchTable& DispatchTable)
 	{
+		DispatchTable.template RegisterType<FMockAbilityBlinkActivateCue>();
 		DispatchTable.template RegisterType<FMockAbilityBlinkCue>();
 
 		// (Again, not a normal setup, just for debugging/testing purposes)
@@ -216,6 +249,8 @@ struct FMockAbilityCueSet
 };
 
 
+
+
 // -------------------------------------------------------
 // MockAbilitySimulation definition
 // -------------------------------------------------------
@@ -294,7 +329,8 @@ public:
 	using UFlyingMovementComponent::FinalizeFrame;
 
 	// NetSimCues
-	void HandleCue(FMockAbilityBlinkCue& BlinkCue, const FNetSimCueSystemParamemters& SystemParameters);
+	void HandleCue(const FMockAbilityBlinkActivateCue& BlinkCue, const FNetSimCueSystemParamemters& SystemParameters);
+	void HandleCue(const FMockAbilityBlinkCue& BlinkCue, const FNetSimCueSystemParamemters& SystemParameters);
 
 	// -------------------------------------------------------------------------------------
 	//	Ability State and Notifications
@@ -331,6 +367,18 @@ public:
 	// Are we currently in the blinking (startup) state
 	UFUNCTION(BlueprintCallable, Category="Mock AbilitySystem")
 	bool IsBlinking() const { return bIsBlinking; }
+	
+	// Blueprint assignable events for blinking. THis allows the user/blueprint to implement rollback-able events
+	DECLARE_DYNAMIC_MULTICAST_DELEGATE_ThreeParams(FMockAbilityBlinkCueEvent, FVector, DestinationLocation, int32, RandomValue, float, ElapsedTimeSeconds);
+	UPROPERTY(BlueprintAssignable, Category="Mock AbilitySystem")
+	FMockAbilityBlinkCueEvent OnBlinkActivateEvent;
+
+	DECLARE_DYNAMIC_MULTICAST_DELEGATE(FMockAbilityBlinkCueRollback);
+	UPROPERTY(BlueprintAssignable, Category="Mock AbilitySystem")
+	FMockAbilityBlinkCueRollback OnBlinkActivateEventRollback;
+
+	UFUNCTION(BlueprintCallable, Category="Mock AbilitySystem")
+	float GetBlinkWarmupTimeSeconds() const;
 
 private:
 	
