NetworkPrediction Plugin
-More test cases around predicted/replicated NetSimCues
-minor cleanup

#rnx

diff --git a/Engine/Plugins/Runtime/NetworkPrediction/Source/NetworkPrediction/Public/NetworkedSimulationModelCues.h b/Engine/Plugins/Runtime/NetworkPrediction/Source/NetworkPrediction/Public/NetworkedSimulationModelCues.h
--- a/Engine/Plugins/Runtime/NetworkPrediction/Source/NetworkPrediction/Public/NetworkedSimulationModelCues.h
+++ b/Engine/Plugins/Runtime/NetworkPrediction/Source/NetworkPrediction/Public/NetworkedSimulationModelCues.h
@@ -19,6 +19,9 @@ enum class ESimulationTickContext : uint8
 	Predict			= 1 << 1,
 	Resimulate		= 1 << 2,
 
+	AuthorityPredict	= Authority | Predict,
+	PredictResimulate	= Predict | Resimulate,
+
 	All				= (Authority | Predict | Resimulate),
 };
 
@@ -40,9 +43,6 @@ struct TNetSimCueTraitsBase
 
 	// Does the cue replicate? (from authority). This will also determine if the cue needs to be saved locally for NetUnique tests (to avoid double playing)
 	static constexpr ENetSimCueReplicationTarget ReplicationTarget { ENetSimCueReplicationTarget::All };
-
-	// Does the cue support rollback callbacks (in contexts where applicable. e.g, when we might get rolled back (!Authority))
-	static constexpr bool Rollbackable { true };
 };
 
 // Preset: non replicated cue that only plays during "latest" simulate. Will not be played during rewind/resimulate.
@@ -51,7 +51,6 @@ struct TNetSimCueTraits_Weak
 {
 	static constexpr uint8 InvokeMask { (uint8)ESimulationTickContext::Authority | (uint8)ESimulationTickContext::Predict };
 	static constexpr ENetSimCueReplicationTarget ReplicationTarget { ENetSimCueReplicationTarget::None };
-	static constexpr bool Rollbackable { false };
 };
 
 // Preset: Replicated, non predicted. Only invoked on authority and will replicate to everyone else.
@@ -60,7 +59,6 @@ struct TNetSimCueTraits_ReplicatedNonPredicted
 {
 	static constexpr uint8 InvokeMask { (uint8)ESimulationTickContext::Authority };
 	static constexpr ENetSimCueReplicationTarget ReplicationTarget { ENetSimCueReplicationTarget::All };
-	static constexpr bool Rollbackable { false };
 };
 
 // Preset: Replicated to interpolating proxies, predicted by autonomous proxy
@@ -69,7 +67,6 @@ struct TNetSimCueTraits_ReplicatedXOrPredicted
 {
 	static constexpr uint8 InvokeMask { (uint8)ESimulationTickContext::Authority | (uint8)ESimulationTickContext::Predict };
 	static constexpr ENetSimCueReplicationTarget ReplicationTarget { ENetSimCueReplicationTarget::Interpolators };
-	static constexpr bool Rollbackable { false };
 };
 
 // Preset: Invoked and replicated to all. Uniqueness testing to avoid double playing, rollbackable so that it can (re)play during resimulates
@@ -78,7 +75,6 @@ struct TNetSimCueTraits_Strong
 {
 	static constexpr uint8 InvokeMask { (uint8)ESimulationTickContext::All };
 	static constexpr ENetSimCueReplicationTarget ReplicationTarget { ENetSimCueReplicationTarget::All };
-	static constexpr bool Rollbackable { true };
 };
 
 // Actual trait struct that we use to look up traits. User cues must specialize this
@@ -108,7 +104,6 @@ struct FNetSimCueWrapperBase
 	virtual bool NetUnique(const void* OtherCueData) const = 0;
 	virtual void* CueData() const = 0;
 	virtual ENetSimCueReplicationTarget GetReplicationTarget() const = 0;
-	virtual bool Rollbackable() const = 0;
 };
 
 template<typename TCue>
@@ -142,11 +137,6 @@ struct TNetSimCueWrapper : FNetSimCueWrapperBase
 		return TCueHandlerTraits<TCue>::ReplicationTarget;
 	}
 
-	bool Rollbackable() const override final
-	{
-		return TCueHandlerTraits<TCue>::Rollbackable;
-	}
-
 	TCue Instance;
 };
 
diff --git a/Engine/Plugins/Runtime/NetworkPredictionExtras/Source/NetworkPredictionExtras/Private/MockAbilitySimulation.cpp b/Engine/Plugins/Runtime/NetworkPredictionExtras/Source/NetworkPredictionExtras/Private/MockAbilitySimulation.cpp
--- a/Engine/Plugins/Runtime/NetworkPredictionExtras/Source/NetworkPredictionExtras/Private/MockAbilitySimulation.cpp
+++ b/Engine/Plugins/Runtime/NetworkPredictionExtras/Source/NetworkPredictionExtras/Private/MockAbilitySimulation.cpp
@@ -105,6 +105,7 @@ void FMockAbilitySimulation::SimulationTick(const FNetSimTimeStep& TimeStep, con
 
 		// Invoke a cue to telegraph where the blink will land. This is making the assumption the handler wouldn't either want to or be able to derive the blink destination from the current state alone
 		// The randomValue being calculated here is purposefully to cause mis prediction of the cue, so that we can demonstrate rollback -> resimulate can do the correction seemlessly
+		
 		UE_LOG(LogTemp, Warning, TEXT("Invoking FMockAbilityBlinkActivateCue from sim"));
 		Output.CueDispatch.Invoke<FMockAbilityBlinkActivateCue>( GetBlinkDestination(), FMath::Rand() % 255 );
 	}
@@ -118,7 +119,8 @@ void FMockAbilitySimulation::SimulationTick(const FNetSimTimeStep& TimeStep, con
 		LocalCmd.MovementInput.Set(0.f, 0.f, 0.f);
 		LocalCmd.bDashPressed = false;
 		LocalCmd.bSprintPressed = false;
-
+		LocalCmd.RotationInput = FRotator::ZeroRotator;
+		//LocalSync.Rotation = Input.Sync.Rotation;
 		LocalSync.Velocity.Set(0.f, 0.f, 0.f);
 
 		if (NewBlinkWarmupLeft <= 0)
@@ -374,7 +376,7 @@ void UMockFlyingAbilityComponent::HandleCue(const FMockAbilityBlinkCue& BlinkCue
 
 	// Crude compensation for cue firing in the past (note this is not necessary! Some cues not care and need to see the "full" effect regardless of when it happened)
 	float Duration = FMath::Max<float>(0.1f, BlinkCueDuration - SystemParameters.TimeSinceInvocation.ToRealTimeSeconds());
-	DrawDebugLine(GetWorld(), BlinkCue.StartLocation, BlinkCue.StopLocation, FColor::Red, false, Duration);
+	DrawDebugLine(GetWorld(), BlinkCue.StartLocation, BlinkCue.StopLocation, (FMath::Rand() % 2) == 0 ? FColor::Red : FColor::Blue, false, Duration);
 
 	if (SystemParameters.Callbacks)
 	{
diff --git a/Engine/Plugins/Runtime/NetworkPredictionExtras/Source/NetworkPredictionExtras/Public/NetworkPredictionExtrasFlyingPawn.h b/Engine/Plugins/Runtime/NetworkPredictionExtras/Source/NetworkPredictionExtras/Public/NetworkPredictionExtrasFlyingPawn.h
--- a/Engine/Plugins/Runtime/NetworkPredictionExtras/Source/NetworkPredictionExtras/Public/NetworkPredictionExtrasFlyingPawn.h
+++ b/Engine/Plugins/Runtime/NetworkPredictionExtras/Source/NetworkPredictionExtras/Public/NetworkPredictionExtrasFlyingPawn.h
@@ -131,7 +131,9 @@ public:
 	virtual void SetupPlayerInputComponent(UInputComponent* PlayerInputComponent) override;
 
 	
+	UFUNCTION(BlueprintPure, Category="Ability")
 	UMockFlyingAbilityComponent* GetMockFlyingAbilityComponent();
+
 	const UMockFlyingAbilityComponent* GetMockFlyingAbilityComponent() const;
 
 	UPROPERTY(EditAnywhere, BlueprintReadWrite, Category="Automation")
