NetworkPrediction Plugin
-Fixed issue with simulation extrapolation getting its aux state out of sync with sync/inputs
-NetUnique renamed NetIdentical. Bug fixed in FSavedCue.
-Fixed visual logging not including state logs
-Added visual logging to simulation extrapolation

#rnx

diff --git a/Engine/Plugins/Runtime/NetworkPrediction/Source/NetworkPrediction/Private/NetworkPredictionTypes.cpp b/Engine/Plugins/Runtime/NetworkPrediction/Source/NetworkPrediction/Private/NetworkPredictionTypes.cpp
--- a/Engine/Plugins/Runtime/NetworkPrediction/Source/NetworkPrediction/Private/NetworkPredictionTypes.cpp
+++ b/Engine/Plugins/Runtime/NetworkPrediction/Source/NetworkPrediction/Private/NetworkPredictionTypes.cpp
@@ -119,6 +119,7 @@ void FVisualLoggingHelpers::VisualLogActor(AActor* Owner, FTransform& Transform,
 	if (NetSimVLogCVars::UseVLogger)
 	{
 		UE_VLOG(Owner, LogNetworkSim, Log, TEXT("%s"), *Params.DebugString);
+		UE_VLOG(Owner, LogNetworkSim, Log, TEXT("%s"), *Params.StateString);
 	}
 }
 
diff --git a/Engine/Plugins/Runtime/NetworkPrediction/Source/NetworkPrediction/Private/NetworkedSimulationModelCues.cpp b/Engine/Plugins/Runtime/NetworkPrediction/Source/NetworkPrediction/Private/NetworkedSimulationModelCues.cpp
--- a/Engine/Plugins/Runtime/NetworkPrediction/Source/NetworkPrediction/Private/NetworkedSimulationModelCues.cpp
+++ b/Engine/Plugins/Runtime/NetworkPrediction/Source/NetworkPrediction/Private/NetworkedSimulationModelCues.cpp
@@ -34,10 +34,10 @@ struct FMockImpactCue
 		Ar << ImpactLocation;
 	}
 
-	bool NetUnique(const FMockImpactCue& Other) const
+	bool NetIdentical(const FMockImpactCue& Other) const
 	{
 		const float ErrorTolerance = 1.f;
-		return !ImpactLocation.Equals(Other.ImpactLocation, ErrorTolerance);
+		return ImpactLocation.Equals(Other.ImpactLocation, ErrorTolerance);
 	}
 
 };
@@ -64,10 +64,10 @@ struct FMockDamageCue
 		Ar << HitLocation;
 	}
 
-	bool NetUnique(const FMockDamageCue& Other) const
+	bool NetIdentical(const FMockDamageCue& Other) const
 	{
 		const float ErrorTolerance = 1.f;
-		return SourceID != Other.SourceID || DamageType != Other.DamageType || !HitLocation.Equals(Other.HitLocation, ErrorTolerance);
+		return SourceID == Other.SourceID && DamageType && Other.DamageType && HitLocation.Equals(Other.HitLocation, ErrorTolerance);
 	}
 };
 NETSIMCUE_REGISTER(FMockDamageCue, TEXT("Damage"));
@@ -89,10 +89,10 @@ struct FMockHealingCue
 		Ar << HealingAmount;
 	}
 
-	bool NetUnique(const FMockHealingCue& Other) const
+	bool NetIdentical(const FMockHealingCue& Other) const
 	{
 		const float ErrorTolerance = 1.f;
-		return SourceID != Other.SourceID || HealingAmount != Other.HealingAmount;
+		return SourceID == Other.SourceID && HealingAmount == Other.HealingAmount;
 	}
 };
 NETSIMCUE_REGISTER(FMockHealingCue, TEXT("Healing"));
diff --git a/Engine/Plugins/Runtime/NetworkPrediction/Source/NetworkPrediction/Public/NetworkPredictionTypes.h b/Engine/Plugins/Runtime/NetworkPrediction/Source/NetworkPrediction/Public/NetworkPredictionTypes.h
--- a/Engine/Plugins/Runtime/NetworkPrediction/Source/NetworkPrediction/Public/NetworkPredictionTypes.h
+++ b/Engine/Plugins/Runtime/NetworkPrediction/Source/NetworkPrediction/Public/NetworkPredictionTypes.h
@@ -159,7 +159,9 @@ struct FVisualLoggingParameters
 	EVisualLoggingContext Context;
 	int32 Frame;
 	EVisualLoggingLifetime Lifetime;
-	FString DebugString;
+	
+	FString DebugString; // Debug String set by the invoker of the VisualLog call
+	mutable FString StateString; // String representation of the available state associated with this log call.
 
 	FColor GetDebugColor() const { return DebugColors[(int32)Context]; }
 	static NETWORKPREDICTION_API FColor DebugColors[(int32)EVisualLoggingContext::MAX];
diff --git a/Engine/Plugins/Runtime/NetworkPrediction/Source/NetworkPrediction/Public/NetworkedSimulationModel.h b/Engine/Plugins/Runtime/NetworkPrediction/Source/NetworkPrediction/Public/NetworkedSimulationModel.h
--- a/Engine/Plugins/Runtime/NetworkPrediction/Source/NetworkPrediction/Public/NetworkedSimulationModel.h
+++ b/Engine/Plugins/Runtime/NetworkPrediction/Source/NetworkPrediction/Public/NetworkedSimulationModel.h
@@ -142,6 +142,8 @@ public:
 		{ 
 			Buffers.CueDispatcher.template DispatchCueRecord<TInDriver>(*InDriver, Ticker.GetTotalProcessedSimulationTime()); 
 		};
+
+		this->Buffers.CueDispatcher.GetDebugName = [this]() { return Driver->GetDebugName(); };
 	}
 
 	virtual ~TNetworkedSimulationModel()
diff --git a/Engine/Plugins/Runtime/NetworkPrediction/Source/NetworkPrediction/Public/NetworkedSimulationModelBuffer.h b/Engine/Plugins/Runtime/NetworkPrediction/Source/NetworkPrediction/Public/NetworkedSimulationModelBuffer.h
--- a/Engine/Plugins/Runtime/NetworkPrediction/Source/NetworkPrediction/Public/NetworkedSimulationModelBuffer.h
+++ b/Engine/Plugins/Runtime/NetworkPrediction/Source/NetworkPrediction/Public/NetworkedSimulationModelBuffer.h
@@ -136,6 +136,8 @@ struct TNetworkSimContiguousBuffer : public TNetworkSimBufferBase<TNetworkSimCon
 	int32 Max() const { return Data.Max(); }
 	int32 GetDirtyCount() const { return DirtyCount; }
 
+	// Returns the element @ frame for writing. The contents of this element are unknown (could be stale content!). Note the element returned is immediately considered "valid" by Num(), Iterators, etc!
+	// The written frame becomes the Head element. Essentially invalidates any existing frames >= Frame
 	T* WriteFrame(int32 Frame)
 	{
 		check(Frame >= 0);
@@ -254,6 +256,7 @@ struct TNetworkSimSparseBuffer : public TNetworkSimBufferBase<TNetworkSimSparseB
 	int32 GetDirtyCount() const { return DirtyCount; }
 
 	// Returns the element @ frame for writing. The contents of this element are unknown (could be stale content!). Note the element returned is immediately considered "valid" by Num(), Iterators, etc!
+	// The written frame becomes the Head element. Essentially invalidates any existing frames >= Frame
 	T* WriteFrame(int32 Frame)
 	{
 		check(Frame >= 0);
@@ -269,6 +272,7 @@ struct TNetworkSimSparseBuffer : public TNetworkSimBufferBase<TNetworkSimSparseB
 				TInternal& InternalData = Data[Pos % Data.Num()];
 				if (InternalData.Frame == Frame)
 				{
+					HeadPos = Pos;
 					return &InternalData.Element;
 				}
 				if (InternalData.Frame <= Frame)
diff --git a/Engine/Plugins/Runtime/NetworkPrediction/Source/NetworkPrediction/Public/NetworkedSimulationModelCues.h b/Engine/Plugins/Runtime/NetworkPrediction/Source/NetworkPrediction/Public/NetworkedSimulationModelCues.h
--- a/Engine/Plugins/Runtime/NetworkPrediction/Source/NetworkPrediction/Public/NetworkedSimulationModelCues.h
+++ b/Engine/Plugins/Runtime/NetworkPrediction/Source/NetworkPrediction/Public/NetworkedSimulationModelCues.h
@@ -41,7 +41,7 @@ struct TNetSimCueTraitsBase
 	// Who can Invoke this Cue in their simulation (if this test fails, the Invoke call is supressed locally)
 	static constexpr uint8 InvokeMask { (uint8)ESimulationTickContext::All };
 
-	// Does the cue replicate? (from authority). This will also determine if the cue needs to be saved locally for NetUnique tests (to avoid double playing)
+	// Does the cue replicate? (from authority). This will also determine if the cue needs to be saved locally for NetIdentical tests (to avoid double playing)
 	static constexpr ENetSimCueReplicationTarget ReplicationTarget { ENetSimCueReplicationTarget::All };
 };
 
@@ -101,7 +101,7 @@ struct FNetSimCueWrapperBase
 {
 	virtual ~FNetSimCueWrapperBase() { }
 	virtual void NetSerialize(FArchive& Ar) = 0;
-	virtual bool NetUnique(const void* OtherCueData) const = 0;
+	virtual bool NetIdentical(const void* OtherCueData) const = 0;
 	virtual void* CueData() const = 0;
 	virtual ENetSimCueReplicationTarget GetReplicationTarget() const = 0;
 };
@@ -121,10 +121,10 @@ struct TNetSimCueWrapper : FNetSimCueWrapperBase
 		Instance.NetSerialize(Ar);
 	}
 
-	bool NetUnique(const void* OtherCueData) const override final
+	bool NetIdentical(const void* OtherCueData) const override final
 	{
-		// Cue types must implement bool NetUnique(const TMyCueType& Other) const
-		return Instance.NetUnique(*((const TCue*)OtherCueData));
+		// Cue types must implement bool NetIdentical(const TMyCueType& Other) const
+		return Instance.NetIdentical(*((const TCue*)OtherCueData));
 	}
 
 	void* CueData() const override final
@@ -280,31 +280,22 @@ struct FSavedCue
 	}
 
 	// Test NetUniqueness against another saved cue
-	bool NetUnique(FSavedCue& OtherCue) const
+	bool NetIdentical(FSavedCue& OtherCue) const
 	{
-		if (ID != OtherCue.ID)
-		{
-			return false;
-		}
-
-		return CueInstance->NetUnique(OtherCue.CueInstance->CueData());
+		return (ID == OtherCue.ID) && CueInstance->NetIdentical(OtherCue.CueInstance->CueData());
 	}
 
 	// Test NetUniqueness against an actual cue instance
 	template<typename TCue>
-	bool NetUnique(TCue& OtherCue) const
+	bool NetIdentical(TCue& OtherCue) const
 	{
-		if (ID != TCue::ID)
-		{
-			return false;
-		}
-		return CueInstance->NetUnique(&OtherCue);
+		return (ID == TCue::ID) && CueInstance->NetIdentical(&OtherCue);
 	}
 
 
-	FString GetTypeName() const
+	FString GetDebugName() const
 	{
-		return FGlobalCueTypeTable::Get().GetTypeName(ID);
+		return FString::Printf(TEXT("[%s 0x%X]"), *FGlobalCueTypeTable::Get().GetTypeName(ID), (int64)this);
 	}
 
 	FNetSimCueTypeId ID = 0;
@@ -475,32 +466,35 @@ struct FNetSimCueDispatcher
 
 					for (FSavedCue& ExistingCue : SavedCues)
 					{
-						if (RollbackTime <= ExistingCue.Time && ExistingCue.NetUnique(NewCue) == false)
+						if (RollbackTime <= ExistingCue.Time && ExistingCue.NetIdentical(NewCue))
 						{
 							// We've matched with an already predicted cue, so suppress this invocation and don't undo the predicted cue
 							ExistingCue.bPendingResimulateRollback = false;
-							UE_LOG(LogNetSimCues, Log, TEXT("Resimulated Cue %s matched existing cue. Suppressing Invocation."), *FGlobalCueTypeTable::Get().GetTypeName(T::ID));
+							UE_LOG(LogNetSimCues, Log, TEXT("%s. Resimulated Cue %s matched existing cue. Suppressing Invocation."), *GetDebugName(), *FGlobalCueTypeTable::Get().GetTypeName(T::ID));
 							return;
 						}
 					}
 
-					UE_LOG(LogNetSimCues, Log, TEXT("Invoking Cue %s during resimulate. Context: %d. Transient: %d. bAllowRollback: %d. bNetConfirmed: %d."), *FGlobalCueTypeTable::Get().GetTypeName(T::ID), Context.TickContext, bTransient, bAllowRollback, bNetConfirmed);
-					GetBuffer(bTransient).Emplace(T::ID, Context.CurrentSimTime, bAllowRollback, bNetConfirmed, bSupportsResimulate, new TNetSimCueWrapper<T>(MoveTempIfPossible(NewCue)));
+					
+					auto& SavedCue = GetBuffer(bTransient).Emplace_GetRef(T::ID, Context.CurrentSimTime, bAllowRollback, bNetConfirmed, bSupportsResimulate, new TNetSimCueWrapper<T>(MoveTempIfPossible(NewCue)));
+					UE_LOG(LogNetSimCues, Log, TEXT("%s. Invoking Cue %s during resimulate. Context: %d. Transient: %d. bAllowRollback: %d. bNetConfirmed: %d."), *GetDebugName(), *SavedCue.GetDebugName(), Context.TickContext, bTransient, bAllowRollback, bNetConfirmed);
 				}
 				else
 				{
-					// Not resimulate case is simple: construct the new cue emplace in the appropriate list
-					UE_LOG(LogNetSimCues, Log, TEXT("Invoking Cue %s. Context: %d. Transient: %d. bAllowRollback: %d. bNetConfirmed: %d."), *FGlobalCueTypeTable::Get().GetTypeName(T::ID), Context.TickContext, bTransient, bAllowRollback, bNetConfirmed);
-					GetBuffer(bTransient).Emplace(T::ID, Context.CurrentSimTime, bAllowRollback, bNetConfirmed, bSupportsResimulate, new TNetSimCueWrapper<T>(Forward<ArgsType>(Args)...));	
+					// Not resimulate case is simple: construct the new cue emplace in the appropriate list	
+					auto& SavedCue = GetBuffer(bTransient).Emplace_GetRef(T::ID, Context.CurrentSimTime, bAllowRollback, bNetConfirmed, bSupportsResimulate, new TNetSimCueWrapper<T>(Forward<ArgsType>(Args)...));	
+					UE_LOG(LogNetSimCues, Log, TEXT("%s. Invoking Cue %s. Context: %d. Transient: %d. bAllowRollback: %d. bNetConfirmed: %d."), *GetDebugName(), *SavedCue.GetDebugName(), Context.TickContext, bTransient, bAllowRollback, bNetConfirmed);
 				}				
 			}
 			else
 			{
-				UE_LOG(LogNetSimCues, Log, TEXT("Suppressing Cue Invocation %s. Mask: %d. TickContext: %d"), *FGlobalCueTypeTable::Get().GetTypeName(T::ID), TCueHandlerTraits<T>::InvokeMask, (int32)Context.TickContext);
+				UE_LOG(LogNetSimCues, Log, TEXT("%s .Suppressing Cue Invocation %s. Mask: %d. TickContext: %d"), *GetDebugName(), *FGlobalCueTypeTable::Get().GetTypeName(T::ID), TCueHandlerTraits<T>::InvokeMask, (int32)Context.TickContext);
 			}
 		}
 	}
 
+	TFunction<FString()> GetDebugName = []() { return TEXT("Unset"); };
+
 protected:
 
 	bool EnsureValidContext()
@@ -560,7 +554,7 @@ struct TNetSimCueDispatcher : public FNetSimCueDispatcher
 				// ReplicationTarget: Cues can be set to only replicate to interpolators
 				if (SerializedCue.CueInstance->GetReplicationTarget() == ENetSimCueReplicationTarget::Interpolators && !bIsInterpolatingSim)
 				{
-					UE_LOG(LogNetSimCues, Log, TEXT("Discarding replicated NetSimCue %s intended for interpolators."), *SerializedCue.GetTypeName());
+					UE_LOG(LogNetSimCues, Log, TEXT("Discarding replicated NetSimCue %s intended for interpolators."), *SerializedCue.GetDebugName());
 					continue;
 				}
 				
@@ -571,10 +565,10 @@ struct TNetSimCueDispatcher : public FNetSimCueDispatcher
 				for (int32 ExistingIdx=0; ExistingIdx < StartingNum; ++ExistingIdx)
 				{
 					FSavedCue& ExistingCue = SavedCues[ExistingIdx];
-					if (SerializedCue.NetUnique(ExistingCue) == false)
+					if (SerializedCue.NetIdentical(ExistingCue))
 					{
 						// These cues are not unique ("close enough") so we are skipping receiving this one
-						UE_LOG(LogNetSimCues, Log, TEXT("Discarding replicated NetSimCue %s because we've already processed it."), *SerializedCue.GetTypeName());
+						UE_LOG(LogNetSimCues, Log, TEXT("Discarding replicated NetSimCue %s because we've already processed it. (Matched %s)"), *SerializedCue.GetDebugName(), *ExistingCue.GetDebugName());
 						bUniqueCue = false;
 						ExistingCue.bNetConfirmed = true;
 						break;
@@ -583,8 +577,9 @@ struct TNetSimCueDispatcher : public FNetSimCueDispatcher
 
 				if (bUniqueCue)
 				{
-					UE_LOG(LogNetSimCues, Log, TEXT("Received NetUnique Cue: %s. (Num replicated cue sent this bunch: %d)"), *SerializedCue.GetTypeName(), NumCues);
-					SavedCues.Emplace(MoveTemp(SerializedCue));
+					
+					auto& SavedCue = SavedCues.Emplace_GetRef(MoveTemp(SerializedCue));
+					UE_LOG(LogNetSimCues, Log, TEXT("%s. Received !NetIdentical Cue: %s. (Num replicated cue sent this bunch: %d)."), *GetDebugName(), *SerializedCue.GetDebugName(), NumCues);
 				}
 			}
 		}
@@ -614,7 +609,7 @@ struct TNetSimCueDispatcher : public FNetSimCueDispatcher
 				FSavedCue& SavedCue = *It;
 				if (!SavedCue.bNetConfirmed && SavedCue.Time <= ConfirmedTime)
 				{
-					UE_LOG(LogNetSimCues, Log, TEXT("Calling OnRollback for SavedCue NetSimCue %s. Cue has not been matched but it <= ConfirmedTime."), *SavedCue.GetTypeName());
+					UE_LOG(LogNetSimCues, Log, TEXT("%s. Calling OnRollback for SavedCue NetSimCue %s. Cue has not been matched but it <= ConfirmedTime."), *GetDebugName(), *SavedCue.GetDebugName());
 					SavedCue.Callbacks.OnRollback.Broadcast();
 					SavedCues.RemoveAt(It.GetIndex(), 1, false);
 				}
@@ -629,7 +624,7 @@ struct TNetSimCueDispatcher : public FNetSimCueDispatcher
 				if (SavedCue.bPendingResimulateRollback)
 				{
 					// Unmatched cue whose time has passed, time to rollback
-					UE_LOG(LogNetSimCues, Log, TEXT("Calling OnRollback for SavedCue NetSimCue %s. Cue was not matched during a resimulate. "), *SavedCue.GetTypeName());
+					UE_LOG(LogNetSimCues, Log, TEXT("%s. Calling OnRollback for SavedCue NetSimCue %s. Cue was not matched during a resimulate."), *GetDebugName(), *SavedCue.GetDebugName());
 					SavedCue.Callbacks.OnRollback.Broadcast();
 					SavedCues.RemoveAt(It.GetIndex(), 1, false);
 				}
@@ -654,20 +649,20 @@ struct TNetSimCueDispatcher : public FNetSimCueDispatcher
 			{
 				if (SavedCue.Time <= DispatchTime)
 				{
-					UE_LOG(LogNetSimCues, Log, TEXT("Dispatching NetSimCue %s."), *SavedCue.GetTypeName());
+					UE_LOG(LogNetSimCues, Log, TEXT("%s Dispatching NetSimCue %s"), *GetDebugName(), *SavedCue.GetDebugName());
 					SavedCue.bDispatched = true;
 					TCueDispatchTable<T>::Get().Dispatch(SavedCue, Handler,DispatchTime);
 				}
 				else
 				{
-					UE_LOG(LogNetSimCues, Log, TEXT("Withholding Cue %s. %s > %s"), *SavedCue.GetTypeName(), *SavedCue.Time.ToString(), *DispatchTime.ToString());
+					UE_LOG(LogNetSimCues, Log, TEXT("Withholding Cue %s. %s > %s."), *SavedCue.GetDebugName(), *SavedCue.Time.ToString(), *DispatchTime.ToString());
 				}
 			}
 		}
 
 		for (FSavedCue& TransientCue : TransientCues)
 		{
-			UE_LOG(LogNetSimCues, Log, TEXT("Dispatching transient NetSimCue %s."), *TransientCue.GetTypeName());
+			UE_LOG(LogNetSimCues, Log, TEXT("Dispatching transient NetSimCue %s."), *TransientCue.GetDebugName());
 			TCueDispatchTable<T>::Get().Dispatch(TransientCue, Handler, DispatchTime);
 		}
 		TransientCues.Reset();
@@ -682,8 +677,8 @@ struct TNetSimCueDispatcher : public FNetSimCueDispatcher
 #if !(UE_BUILD_SHIPPING || UE_BUILD_TEST)
 			for (int32 i=0; i <= SavedCuePruneIdx; ++i)
 			{
-				UE_CLOG(!SavedCues[i].bDispatched, LogNetSimCues, Warning, TEXT("Non-Dispatched Cue is about to be pruned! %s. %s"), *SavedCues[i].GetTypeName());
-				UE_LOG(LogNetSimCues, Log, TEXT("Pruning Cue %s. Invoke Time: %s. Current Time: %s"), *SavedCues[i].GetTypeName(), *SavedCues[i].Time.ToString(), *DispatchTime.ToString());
+				UE_CLOG(!SavedCues[i].bDispatched, LogNetSimCues, Warning, TEXT("Non-Dispatched Cue is about to be pruned! %s. %s"), *SavedCues[i].GetDebugName());
+				UE_LOG(LogNetSimCues, Log, TEXT("%s. Pruning Cue %s. Invoke Time: %s. Current Time: %s."), *GetDebugName(), *SavedCues[i].GetDebugName(), *SavedCues[i].Time.ToString(), *DispatchTime.ToString());
 			}
 #endif
 
@@ -711,6 +706,7 @@ struct TNetSimCueDispatcher : public FNetSimCueDispatcher
 			if (SavedCue.bResimulates && SavedCue.Time >= InRollbackTime)
 			{
 				SavedCue.bPendingResimulateRollback = true;
+				UE_LOG(LogNetSimCues, Log, TEXT("%s. Marking %s bPendingResimulateRollback."), *GetDebugName(), *SavedCue.GetDebugName());
 			}
 		}
 	}
diff --git a/Engine/Plugins/Runtime/NetworkPrediction/Source/NetworkPrediction/Public/NetworkedSimulationModelDebugger.h b/Engine/Plugins/Runtime/NetworkPrediction/Source/NetworkPrediction/Public/NetworkedSimulationModelDebugger.h
--- a/Engine/Plugins/Runtime/NetworkPrediction/Source/NetworkPrediction/Public/NetworkedSimulationModelDebugger.h
+++ b/Engine/Plugins/Runtime/NetworkPrediction/Source/NetworkPrediction/Public/NetworkedSimulationModelDebugger.h
@@ -665,7 +665,7 @@ struct TNetworkSimulationModelDebugger : public INetworkSimulationModelDebugger
 		{
 			FVisualLoggingParameters VLogParams(EVisualLoggingContext::LastPredicted, NetworkSim->Buffers.Sync.HeadFrame(), EVisualLoggingLifetime::Transient);
 			int32 SyncFrame = NetworkSim->Buffers.Sync.HeadFrame();
-			NetworkSim->Driver->VisualLog(NetworkSim->Buffers.Input[SyncFrame], LatestSync, NetworkSim->Buffers.Aux[SyncFrame], VLogParams);
+			NetworkSim->Driver->InvokeVisualLog(NetworkSim->Buffers.Input[SyncFrame], LatestSync, NetworkSim->Buffers.Aux[SyncFrame], VLogParams);
 		}
 
 		FStuff ServerPIEStuff = GetServerPIEStuff();
@@ -675,7 +675,7 @@ struct TNetworkSimulationModelDebugger : public INetworkSimulationModelDebugger
 			{
 				FVisualLoggingParameters VLogParams(EVisualLoggingContext::CurrentServerPIE, ServerPIEStuff.NetworkSim->Buffers.Sync.HeadFrame(), EVisualLoggingLifetime::Transient);
 				int32 SyncFrame = ServerPIEStuff.NetworkSim->Buffers.Sync.HeadFrame();
-				NetworkSim->Driver->VisualLog(ServerPIEStuff.NetworkSim->Buffers.Input[SyncFrame], ServerLatestSync, ServerPIEStuff.NetworkSim->Buffers.Aux[SyncFrame], VLogParams);
+				NetworkSim->Driver->InvokeVisualLog(ServerPIEStuff.NetworkSim->Buffers.Input[SyncFrame], ServerLatestSync, ServerPIEStuff.NetworkSim->Buffers.Aux[SyncFrame], VLogParams);
 			}
 		}
 
@@ -687,7 +687,7 @@ struct TNetworkSimulationModelDebugger : public INetworkSimulationModelDebugger
 				{
 					const EVisualLoggingContext Context = (Frame == NetworkSim->RepProxy_Autonomous.GetLastSerializedFrame()) ? EVisualLoggingContext::LastConfirmed : EVisualLoggingContext::OtherPredicted;
 					FVisualLoggingParameters VLogParams(Context, NetworkSim->Buffers.Sync.HeadFrame(), EVisualLoggingLifetime::Transient);
-					NetworkSim->Driver->VisualLog(NetworkSim->Buffers.Input[Frame], SyncState, NetworkSim->Buffers.Aux[Frame], VLogParams);
+					NetworkSim->Driver->InvokeVisualLog(NetworkSim->Buffers.Input[Frame], SyncState, NetworkSim->Buffers.Aux[Frame], VLogParams);
 				}
 			}
 		}
@@ -695,7 +695,7 @@ struct TNetworkSimulationModelDebugger : public INetworkSimulationModelDebugger
 		{
 			FVisualLoggingParameters VLogParams(EVisualLoggingContext::LastConfirmed, NetworkSim->Buffers.Sync.HeadFrame(), EVisualLoggingLifetime::Transient);
 			int32 HeadFrame = NetworkSim->Buffers.Sync.HeadFrame();
-			NetworkSim->Driver->VisualLog(NetworkSim->Buffers.Input[HeadFrame], NetworkSim->Buffers.Sync[HeadFrame], NetworkSim->Buffers.Aux[HeadFrame], VLogParams);
+			NetworkSim->Driver->InvokeVisualLog(NetworkSim->Buffers.Input[HeadFrame], NetworkSim->Buffers.Sync[HeadFrame], NetworkSim->Buffers.Aux[HeadFrame], VLogParams);
 			
 			if (NetworkSim->GetSimulatedUpdateMode() != ESimulatedUpdateMode::Interpolate)
 			{
@@ -704,7 +704,6 @@ struct TNetworkSimulationModelDebugger : public INetworkSimulationModelDebugger
 
 				FVector2D LocalSimulationTimeData(Owner->GetWorld()->GetTimeSeconds(), NetworkSim->Ticker.GetTotalProcessedSimulationTime().ToRealTimeMS());
 				UE_VLOG_HISTOGRAM(Owner, LogNetworkSimDebug, Log, "Simulated Time Graph", "Local Simulation Time", LocalSimulationTimeData);
-			
 			}
 		}
 	}
diff --git a/Engine/Plugins/Runtime/NetworkPrediction/Source/NetworkPrediction/Public/NetworkedSimulationModelInterpolator.h b/Engine/Plugins/Runtime/NetworkPrediction/Source/NetworkPrediction/Public/NetworkedSimulationModelInterpolator.h
--- a/Engine/Plugins/Runtime/NetworkPrediction/Source/NetworkPrediction/Public/NetworkedSimulationModelInterpolator.h
+++ b/Engine/Plugins/Runtime/NetworkPrediction/Source/NetworkPrediction/Public/NetworkedSimulationModelInterpolator.h
@@ -237,7 +237,7 @@ struct TNetSimInterpolator
 					auto VLogHelper = [&](int32 Frame, EVisualLoggingContext Context)
 					{
 						FVisualLoggingParameters VLogParams(Context, Frame, EVisualLoggingLifetime::Transient);
-						Driver->VisualLog(Buffers.Input[Frame], Buffers.Sync[Frame], Buffers.Aux[Frame], VLogParams);
+						Driver->InvokeVisualLog(Buffers.Input[Frame], Buffers.Sync[Frame], Buffers.Aux[Frame], VLogParams);
 					};
 
 					VLogHelper(Buffers.Sync.TailFrame(), EVisualLoggingContext::InterpolationBufferTail);
@@ -245,17 +245,17 @@ struct TNetSimInterpolator
 
 					{
 						FVisualLoggingParameters VLogParams(EVisualLoggingContext::InterpolationFrom, InterpolationFrame-1, EVisualLoggingLifetime::Transient);
-						Driver->VisualLog(Buffers.Input[InterpolationFrame-1], &FromState.Sync, &FromState.Aux, VLogParams);
+						Driver->InvokeVisualLog(Buffers.Input[InterpolationFrame-1], &FromState.Sync, &FromState.Aux, VLogParams);
 					}
 
 					{
 						FVisualLoggingParameters VLogParams(EVisualLoggingContext::InterpolationTo, InterpolationFrame, EVisualLoggingLifetime::Transient);
-						Driver->VisualLog(Buffers.Input[InterpolationFrame], ToState, ToAuxState, VLogParams);
+						Driver->InvokeVisualLog(Buffers.Input[InterpolationFrame], ToState, ToAuxState, VLogParams);
 					}
 
 					{
 						FVisualLoggingParameters VLogParams(LoggingContext, InterpolationFrame, EVisualLoggingLifetime::Transient);
-						Driver->VisualLog(Buffers.Input[InterpolationFrame], &OutputState.Sync, &OutputState.Aux, VLogParams);
+						Driver->InvokeVisualLog(Buffers.Input[InterpolationFrame], &OutputState.Sync, &OutputState.Aux, VLogParams);
 					}
 				}
 				
diff --git a/Engine/Plugins/Runtime/NetworkPrediction/Source/NetworkPrediction/Public/NetworkedSimulationModelRepControllers.h b/Engine/Plugins/Runtime/NetworkPrediction/Source/NetworkPrediction/Public/NetworkedSimulationModelRepControllers.h
--- a/Engine/Plugins/Runtime/NetworkPrediction/Source/NetworkPrediction/Public/NetworkedSimulationModelRepControllers.h
+++ b/Engine/Plugins/Runtime/NetworkPrediction/Source/NetworkPrediction/Public/NetworkedSimulationModelRepControllers.h
@@ -406,7 +406,10 @@ struct TRepController_Simulated : public TBase
 			InputCmd = Buffers.Input.WriteFrame(DestinationFrame);
 			SyncState = Buffers.Sync.WriteFrame(DestinationFrame);
 			AuxState = Buffers.Aux.WriteFrame(DestinationFrame);
-			Buffers.Input.WriteFrame(DestinationFrame);
+
+			ensure(Buffers.Input.HeadFrame() <= DestinationFrame);
+			ensure(Buffers.Sync.HeadFrame() <= DestinationFrame);
+			ensure(Buffers.Aux.HeadFrame() <= DestinationFrame);
 
 			// Update tick info
 			Ticker.SetTotalProcessedSimulationTime(LastSerializedSimulationTime, DestinationFrame);
@@ -529,7 +532,7 @@ struct TRepController_Simulated : public TBase
 		if (NETSIM_MODEL_DEBUG)
 		{
 			FVisualLoggingParameters VLogParams(EVisualLoggingContext::FirstMispredicted, ParentFrame, EVisualLoggingLifetime::Persistent, TEXT("DependentRollbackBegin"));
-			Driver->VisualLog(nullptr, &LastSerializedSyncState, &LastSerializedAuxState, VLogParams);
+			Driver->InvokeVisualLog(nullptr, &LastSerializedSyncState, &LastSerializedAuxState, VLogParams);
 		}
 	}
 		
@@ -545,7 +548,7 @@ struct TRepController_Simulated : public TBase
 		if (NETSIM_MODEL_DEBUG)
 		{
 			FVisualLoggingParameters VLogParams(bFinalStep ? EVisualLoggingContext::LastMispredicted : EVisualLoggingContext::OtherMispredicted, ParentFrame, EVisualLoggingLifetime::Persistent, TEXT("DependentRollbackStep"));
-			Driver->VisualLog(nullptr, SyncState, Buffers.Aux.HeadElement(), VLogParams);
+			Driver->InvokeVisualLog(nullptr, SyncState, Buffers.Aux.HeadElement(), VLogParams);
 		}
 	}
 
@@ -570,6 +573,12 @@ private:
 		TSyncState* NextSyncState = Buffers.Sync.WriteFrame(OutputFrame);
 		TAuxState* AuxState = Buffers.Aux[InputFrame];
 
+		if (NETSIM_MODEL_DEBUG)
+		{
+			FVisualLoggingParameters VLogParameters( EVisualLoggingContext::OtherPredicted, InputFrame, EVisualLoggingLifetime::Persistent, TEXT("Pre SimulationExtrapolation"));
+			Driver->InvokeVisualLog(NewCmd, PrevSyncState, AuxState, VLogParameters);
+		}
+
 		// Do the actual update
 		{
 			TScopedSimulationTick<Model> UpdateScope(Ticker, Buffers.CueDispatcher, ESimulationTickContext::Resimulate, OutputFrame, NewCmd->GetFrameDeltaTime());
@@ -579,6 +588,12 @@ private:
 				{ *NextSyncState, Buffers.Aux.LazyWriter(OutputFrame), Buffers.CueDispatcher } );
 		}
 
+		if (NETSIM_MODEL_DEBUG)
+		{
+			FVisualLoggingParameters VLogParameters( EVisualLoggingContext::OtherPredicted, InputFrame, EVisualLoggingLifetime::Persistent, TEXT("Post SimulationExtrapolation"));
+			Driver->InvokeVisualLog(NewCmd, NextSyncState, Buffers.Aux[OutputFrame], VLogParameters);
+		}
+
 		Ticker.MaxAllowedFrame = OutputFrame;
 	}
 	
@@ -695,7 +710,7 @@ struct TRepController_Autonomous: public TBase
 		if (bDoVisualLog)
 		{
 			FVisualLoggingParameters VLogParameters(EVisualLoggingContext::LastConfirmed, SerializedFrame, EVisualLoggingLifetime::Persistent, TEXT("Serialized State"));
-			Driver->VisualLog(Buffers.Input[SerializedFrame], &SerializedSyncState, &SerializedAuxState, VLogParameters);
+			Driver->InvokeVisualLog(Buffers.Input[SerializedFrame], &SerializedSyncState, &SerializedAuxState, VLogParameters);
 		}
 
 		if (ClientSyncState)
@@ -704,7 +719,7 @@ struct TRepController_Autonomous: public TBase
 			if (bDoVisualLog)
 			{
 				FVisualLoggingParameters VLogParameters(EVisualLoggingContext::FirstMispredicted, SerializedFrame, EVisualLoggingLifetime::Persistent, TEXT("Mispredicted State"));
-				Driver->VisualLog(Buffers.Input[SerializedFrame], ClientSyncState, Buffers.Aux[SerializedFrame], VLogParameters);
+				Driver->InvokeVisualLog(Buffers.Input[SerializedFrame], ClientSyncState, Buffers.Aux[SerializedFrame], VLogParameters);
 			}
 		}
 		else
@@ -763,7 +778,7 @@ struct TRepController_Autonomous: public TBase
 			if (NETSIM_MODEL_DEBUG)
 			{
 				FVisualLoggingParameters VLogParameters(Frame == LastFrameToProcess ? EVisualLoggingContext::LastMispredicted : EVisualLoggingContext::OtherMispredicted, Frame, EVisualLoggingLifetime::Persistent, TEXT("Resimulate Step: mispredicted"));
-				Driver->VisualLog(Buffers.Input[OutputFrame], Buffers.Sync[OutputFrame], Buffers.Aux[OutputFrame], VLogParameters);
+				Driver->InvokeVisualLog(Buffers.Input[OutputFrame], Buffers.Sync[OutputFrame], Buffers.Aux[OutputFrame], VLogParameters);
 			}
 
 			// Do the actual update
@@ -779,7 +794,7 @@ struct TRepController_Autonomous: public TBase
 			if (NETSIM_MODEL_DEBUG)
 			{			
 				FVisualLoggingParameters VLogParameters(Frame == LastFrameToProcess ? EVisualLoggingContext::LastPredicted : EVisualLoggingContext::OtherPredicted, Frame, EVisualLoggingLifetime::Persistent, TEXT("Resimulate Step: repredicted"));
-				Driver->VisualLog(Buffers.Input[OutputFrame], Buffers.Sync[OutputFrame], Buffers.Aux[OutputFrame], VLogParameters);
+				Driver->InvokeVisualLog(Buffers.Input[OutputFrame], Buffers.Sync[OutputFrame], Buffers.Aux[OutputFrame], VLogParameters);
 			}
 
 			// Tell dependent simulations to advance
diff --git a/Engine/Plugins/Runtime/NetworkPrediction/Source/NetworkPrediction/Public/NetworkedSimulationModelTraits.h b/Engine/Plugins/Runtime/NetworkPrediction/Source/NetworkPrediction/Public/NetworkedSimulationModelTraits.h
--- a/Engine/Plugins/Runtime/NetworkPrediction/Source/NetworkPrediction/Public/NetworkedSimulationModelTraits.h
+++ b/Engine/Plugins/Runtime/NetworkPrediction/Source/NetworkPrediction/Public/NetworkedSimulationModelTraits.h
@@ -17,14 +17,47 @@ public:
 	// Owning object for Visual Logs so that the system can emit them internally
 	virtual const AActor* GetVLogOwner() const = 0;
 
-	// Call to visual log the given states. Note that not all 3 will always be present and you should check for nullptrs.
-	virtual void VisualLog(const TInputCmd* Input, const TSyncState* Sync, const TAuxState* Aux, const FVisualLoggingParameters& SystemParameters) const = 0;
+	// Call to visual log the given states. Note that not all 3 will always be present and you should check for nullptrs. Child classes most likely want to override ::VisualLog, not this
+	virtual void InvokeVisualLog(const TInputCmd* Input, const TSyncState* Sync, const TAuxState* Aux, const FVisualLoggingParameters& SystemParameters) const
+	{
+		// String-ify the passed in states
+		FStringOutputDevice StrOut;
+		StrOut.SetAutoEmitLineTerminator(true);
+
+		FStandardLoggingParameters LogParameters(&StrOut, EStandardLoggingContext::Full, SystemParameters.Frame);
+		if (Input)
+		{
+			StrOut.Log(TEXT("Input:"));
+			Input->Log(LogParameters);
+			StrOut.Log(TEXT(""));
+		}
+		if (Sync)
+		{
+			StrOut.Log(TEXT("Sync:"));
+			Sync->Log(LogParameters);
+			StrOut.Log(TEXT(""));
+		}
+		if (Aux)
+		{
+			StrOut.Log(TEXT("Aux:"));
+			Aux->Log(LogParameters);
+		}
+
+		SystemParameters.StateString = StrOut;
+		VisualLog(Input, Sync, Aux, SystemParameters);
+	}
 	
 	// Called whenever the sim is ready to process new local input.
 	virtual void ProduceInput(const FNetworkSimTime SimTime, TInputCmd&) = 0;
 	
 	// Called from the Network Sim at the end of the sim frame when there is new sync data.
 	virtual void FinalizeFrame(const typename TBufferTypes::TSyncState& SyncState, const typename TBufferTypes::TAuxState& AuxState) = 0;
+
+protected:
+
+	// Called to visual log the given states. Note that not all 3 will always be present and you should check for nullptrs.
+	virtual void VisualLog(const TInputCmd* Input, const TSyncState* Sync, const TAuxState* Aux, const FVisualLoggingParameters& SystemParameters) const = 0;
+
 };
 
 // Internal traits that are derived from the user settings (E.g, these can't go in FNetSimModelDefBase because they are dependent on the user types)
diff --git a/Engine/Plugins/Runtime/NetworkPredictionExtras/Source/NetworkPredictionExtras/Private/MockAbilitySimulation.cpp b/Engine/Plugins/Runtime/NetworkPredictionExtras/Source/NetworkPredictionExtras/Private/MockAbilitySimulation.cpp
--- a/Engine/Plugins/Runtime/NetworkPredictionExtras/Source/NetworkPredictionExtras/Private/MockAbilitySimulation.cpp
+++ b/Engine/Plugins/Runtime/NetworkPredictionExtras/Source/NetworkPredictionExtras/Private/MockAbilitySimulation.cpp
@@ -99,8 +99,6 @@ void FMockAbilitySimulation::SimulationTick(const FNetSimTimeStep& TimeStep, con
 	const bool bBlinkActivate = (Input.Cmd.bBlinkPressed && Input.Sync.Stamina > BlinkCost && bAllowNewActivations);
 	if (bBlinkActivate)
 	{
-
-
 		LocalAux.BlinkWarmupLeft = BlinkWarmupMS;
 
 		// Invoke a cue to telegraph where the blink will land. This is making the assumption the handler wouldn't either want to or be able to derive the blink destination from the current state alone
@@ -146,10 +144,13 @@ void FMockAbilitySimulation::SimulationTick(const FNetSimTimeStep& TimeStep, con
 				// Invoke a NetCue for the blink. This is essentially capturing the start/end location so that all receivers of the event
 				// get the exact coordinates (maybe overkill in practice but key is that we have data that we want to pass out via an event)
 				
+				//UE_LOG(LogTemp, Warning, TEXT("Invoking FMockAbilityBlinkCue from sim. %s - %s (LocalSync.Location: %s) "), *Input.Sync.Location.ToString(), *DestLocation.ToString(), *LocalSync.Location.ToString());
+				//UE_VLOG(OwningActor, LogNetworkSim, Log, TEXT("Invoking FMockAbilityBlinkCue from sim. %s - %s (LocalSync.Location: %s) "), *Input.Sync.Location.ToString(), *DestLocation.ToString(), *LocalSync.Location.ToString());
+
 				switch(MockAbilityCVars::BlinkCueType()) // Only for dev/testing. Not a normal setup.
 				{
 				case 0:
-					// Skip o purpose
+					// Skip on purpose
 					break;
 				case 1:
 					Output.CueDispatch.Invoke<FMockAbilityBlinkCue_Weak>( Input.Sync.Location, DestLocation );
@@ -367,8 +368,7 @@ static FAutoConsoleVariableRef CVarBindAutomatically(TEXT("NetworkPredictionExtr
 
 void UMockFlyingAbilityComponent::HandleCue(const FMockAbilityBlinkCue& BlinkCue, const FNetSimCueSystemParamemters& SystemParameters)
 {
-	FString RoleStr = GetOwnerRole() == ROLE_Authority ? TEXT("Server") : TEXT("Client");
-
+	FString RoleStr = *UEnum::GetValueAsString(TEXT("Engine.ENetRole"), GetOwnerRole());
 	FVector Delta = GetOwner()->GetActorLocation() - BlinkCue.StopLocation;
 
 	UE_LOG(LogTemp, Display, TEXT("[%s] BlinkCue! : <%f, %f, %f> - <%f, %f, %f>. ElapsedTime: %s. Delta: %.3f"), *RoleStr, BlinkCue.StartLocation.X, BlinkCue.StartLocation.Y, BlinkCue.StartLocation.Z,
diff --git a/Engine/Plugins/Runtime/NetworkPredictionExtras/Source/NetworkPredictionExtras/Public/MockAbilitySimulation.h b/Engine/Plugins/Runtime/NetworkPredictionExtras/Source/NetworkPredictionExtras/Public/MockAbilitySimulation.h
--- a/Engine/Plugins/Runtime/NetworkPredictionExtras/Source/NetworkPredictionExtras/Public/MockAbilitySimulation.h
+++ b/Engine/Plugins/Runtime/NetworkPredictionExtras/Source/NetworkPredictionExtras/Public/MockAbilitySimulation.h
@@ -144,10 +144,10 @@ struct FMockAbilityBlinkActivateCue
 		Ar << RandomType;
 	}
 	
-	bool NetUnique(const FMockAbilityBlinkActivateCue& Other) const
+	bool NetIdentical(const FMockAbilityBlinkActivateCue& Other) const
 	{
 		const float ErrorTolerance = 1.f;
-		return !Destination.Equals(Other.Destination, ErrorTolerance) || RandomType != Other.RandomType;
+		return Destination.Equals(Other.Destination, ErrorTolerance) && RandomType == Other.RandomType;
 	}
 };
 
@@ -206,10 +206,10 @@ struct FMockAbilityBlinkCue
 		StopLocation.NetSerialize(Ar, nullptr, b);
 	}
 	
-	bool NetUnique(const FMockAbilityBlinkCue& Other) const
+	bool NetIdentical(const FMockAbilityBlinkCue& Other) const
 	{
 		const float ErrorTolerance = 1.f;
-		return !StartLocation.Equals(Other.StartLocation, ErrorTolerance) || !StopLocation.Equals(Other.StopLocation, ErrorTolerance);
+		return StartLocation.Equals(Other.StartLocation, ErrorTolerance) && StopLocation.Equals(Other.StopLocation, ErrorTolerance);
 	}
 };
 
