NetworkPrediction Plugin
-NetSerialize and NetIdentical no longer need to be defined on types that do not need to call them.
-Added clearer static_assert if a cue type does need the functions but they are not there.

#rnx

diff --git a/Engine/Plugins/Runtime/NetworkPrediction/Source/NetworkPrediction/Public/NetworkedSimulationModelCues.h b/Engine/Plugins/Runtime/NetworkPrediction/Source/NetworkPrediction/Public/NetworkedSimulationModelCues.h
--- a/Engine/Plugins/Runtime/NetworkPrediction/Source/NetworkPrediction/Public/NetworkedSimulationModelCues.h
+++ b/Engine/Plugins/Runtime/NetworkPrediction/Source/NetworkPrediction/Public/NetworkedSimulationModelCues.h
@@ -77,6 +77,14 @@ struct TNetSimCueTraits_Strong
 	static constexpr ENetSimCueReplicationTarget ReplicationTarget { ENetSimCueReplicationTarget::All };
 };
 
+// Preset: Non replicated but if a resimulate happens, the cue is undone and replayed.
+// This is not common and doesn't really have a clear use case. But the system can support it.
+struct TNetSimCueTraits_NonReplicatedResimulated
+{
+	static constexpr uint8 InvokeMask { (uint8)ESimulationTickContext::All };
+	static constexpr ENetSimCueReplicationTarget ReplicationTarget { ENetSimCueReplicationTarget::None };
+};
+
 // Actual trait struct that we use to look up traits. User cues must specialize this
 template<typename TCue>
 struct TCueHandlerTraits : public TNetSimCueTraitsBase
@@ -84,6 +92,55 @@ struct TCueHandlerTraits : public TNetSimCueTraitsBase
 };
 
 
+// Type requirements: helper to determine if NetSerialize/NetIdentical functions need to be defined for user types based on traits
+template <typename TCue>
+struct TNetSimCueTypeRequirements
+{
+	enum {
+		// NetSerialize is required if we ever need to replicate
+		RequiresNetSerialize = (TCueHandlerTraits<TCue>::ReplicationTarget != ENetSimCueReplicationTarget::None),
+		// Likewise for NetIdentical, but also if we plan to invoke during Resimulate too (even if non replicated, we use NetIdentical for comparisons. though this is probably a non practical use case).
+		RequiresNetIdentical = (TCueHandlerTraits<TCue>::ReplicationTarget != ENetSimCueReplicationTarget::None) || (TCueHandlerTraits<TCue>::InvokeMask & (uint8)ESimulationTickContext::Resimulate)
+	};
+};
+
+// Helper to compile time check if NetSerialize exists
+GENERATE_MEMBER_FUNCTION_CHECK(NetSerialize, void,, FArchive&);
+
+// Helper to compile time check if NetIdentical exists (since argument is template type, must be wrapped in helper struct)
+template<typename TCue>
+struct THasNetIdenticalHelper
+{
+	GENERATE_MEMBER_FUNCTION_CHECK(NetIdentical, bool, const, const TCue&);
+	enum { Value = THasMemberFunction_NetIdentical<TCue>::Value };
+};
+
+// Helper to call NetIdentical if type defines it
+template<typename TCue, bool Enabled=THasNetIdenticalHelper<TCue>::Value>
+struct TNetCueNetIdenticalHelper
+{
+	static bool CallNetIdenticalOrNot(const TCue& Cue, const TCue& Other) { ensure(false); return false; } // This should never be hit by cue types that don't need to NetIdentical
+};
+
+template<typename TCue>
+struct TNetCueNetIdenticalHelper<TCue, true>
+{
+	static bool CallNetIdenticalOrNot(const TCue& Cue, const TCue& Other) { return Cue.NetIdentical(Other); }
+};
+
+// Helper to call NetSerialize if type defines it
+template<typename TCue, bool Enabled=THasMemberFunction_NetSerialize<TCue>::Value>
+struct TNetCueNetSerializeHelper
+{
+	static void CallNetSerializeOrNot(TCue& Cue, FArchive& Ar) { ensure(false); } // This should never be hit by cue types that don't need to NetSerialize
+};
+
+template<typename TCue>
+struct TNetCueNetSerializeHelper<TCue, true>
+{
+	static void CallNetSerializeOrNot(TCue& Cue, FArchive& Ar) { Cue.NetSerialize(Ar); }
+};
+
 // Traits for TNetSimCueDispatcher
 struct NETWORKPREDICTION_API FCueDispatcherTraitsBase
 {
@@ -117,14 +174,12 @@ struct TNetSimCueWrapper : FNetSimCueWrapperBase
 
 	void NetSerialize(FArchive& Ar) override final
 	{
-		// Cue types must implement NetSerialize(FArchive& Ar)
-		Instance.NetSerialize(Ar);
+		TNetCueNetSerializeHelper<TCue>::CallNetSerializeOrNot(Instance, Ar);
 	}
 
 	bool NetIdentical(const void* OtherCueData) const override final
 	{
-		// Cue types must implement bool NetIdentical(const TMyCueType& Other) const
-		return Instance.NetIdentical(*((const TCue*)OtherCueData));
+		return TNetCueNetIdenticalHelper<TCue>::CallNetIdenticalOrNot(Instance, *((const TCue*)OtherCueData));
 	}
 
 	void* CueData() const override final
@@ -164,6 +219,7 @@ struct FNetSimCueSystemParamemters
 // ----------------------------------------------------------------------------------------------------------------------------------------------
 // GlobalCueTypeTable: Cue types register with this to get a Type ID assigned. That ID is used in net serialization to talk about types.
 // ----------------------------------------------------------------------------------------------------------------------------------------------
+
 class FGlobalCueTypeTable
 {
 public:
@@ -185,6 +241,9 @@ public:
 			TypeName
 		};
 		PendingCueTypes.Emplace(TypeInfo);
+		
+		static_assert(THasMemberFunction_NetSerialize<TCue>::Value || !TNetSimCueTypeRequirements<TCue>::RequiresNetSerialize, "TCue must implement void NetSerialize(FArchive&)");
+		static_assert(THasNetIdenticalHelper<TCue>::Value || !TNetSimCueTypeRequirements<TCue>::RequiresNetIdentical, "TCue must implement bool NetIdentical(const TCue&) const");
 	}
 
 	void FinalizeTypes()
@@ -738,12 +797,12 @@ private:
 // ------------------------------------------------------------------------------------------------
 // Register a cue type with the global table.
 // ------------------------------------------------------------------------------------------------
-template<typename T>
+template<typename TCue>
 struct TNetSimCueTypeAutoRegisterHelper
 {
 	TNetSimCueTypeAutoRegisterHelper(const FString& Name)
 	{
-		FGlobalCueTypeTable::Get().RegisterType<T>(Name);
+		FGlobalCueTypeTable::Get().RegisterType<TCue>(Name);
 	}
 };
 
@@ -753,12 +812,12 @@ struct TNetSimCueTypeAutoRegisterHelper
 // ------------------------------------------------------------------------------------------------
 // Register a handler's cue types via a static "RegisterNetSimCueTypes" on the handler itself
 // ------------------------------------------------------------------------------------------------
-template<typename T>
+template<typename TCue>
 struct TNetSimCueHandlerAutoRegisterHelper
 {
 	TNetSimCueHandlerAutoRegisterHelper()
 	{
-		T::RegisterNetSimCueTypes( TCueDispatchTable<T>::Get() );
+		TCue::RegisterNetSimCueTypes( TCueDispatchTable<TCue>::Get() );
 	}
 };
 
diff --git a/Engine/Plugins/Runtime/NetworkPredictionExtras/Source/NetworkPredictionExtras/Public/MockAbilitySimulation.h b/Engine/Plugins/Runtime/NetworkPredictionExtras/Source/NetworkPredictionExtras/Public/MockAbilitySimulation.h
--- a/Engine/Plugins/Runtime/NetworkPredictionExtras/Source/NetworkPredictionExtras/Public/MockAbilitySimulation.h
+++ b/Engine/Plugins/Runtime/NetworkPredictionExtras/Source/NetworkPredictionExtras/Public/MockAbilitySimulation.h
@@ -136,7 +136,7 @@ struct FMockAbilityBlinkActivateCue
 
 	FVector_NetQuantize10 Destination;
 	uint8 RandomType; // Random value used to color code the effect. This is the test/prove out mispredictions
-
+	
 	void NetSerialize(FArchive& Ar)
 	{
 		bool b = false;
@@ -157,7 +157,7 @@ struct TCueHandlerTraits<FMockAbilityBlinkActivateCue> : public TNetSimCueTraits
 // ----------------------------------------------------------------------------------------------
 
 
-#define LOG_BLINK_CUE 1 // During development, its useful to sanity check that we aren't doing more construction or moves than we expect
+#define LOG_BLINK_CUE 0 // During development, its useful to sanity check that we aren't doing more construction or moves than we expect
 
 // Cue for blink (the moment the teleport happens)
 struct FMockAbilityBlinkCue
@@ -206,6 +206,7 @@ struct FMockAbilityBlinkCue
 		StopLocation.NetSerialize(Ar, nullptr, b);
 	}
 	
+	
 	bool NetIdentical(const FMockAbilityBlinkCue& Other) const
 	{
 		const float ErrorTolerance = 1.f;
@@ -223,6 +224,8 @@ struct TCueHandlerTraits<FMockAbilityBlinkCue> : public TNetSimCueTraits_Strong
 #define DECLARE_BLINKCUE_SUBTYPE(TYPE, TRAITS) \
  struct TYPE : public FMockAbilityBlinkCue { \
  template <typename... ArgsType> TYPE(ArgsType&&... Args) : FMockAbilityBlinkCue(Forward<ArgsType>(Args)...) { } \
+ void NetSerialize(FArchive& Ar) { FMockAbilityBlinkCue::NetSerialize(Ar); } \
+ bool NetIdentical(const TYPE& Other) const { return FMockAbilityBlinkCue::NetIdentical(Other); } \
  NETSIMCUE_BODY(); }; \
  template<> struct TCueHandlerTraits<TYPE> : public TRAITS { };
 
