NetworkPrediction Plugin
-NetSimCue cleanup, some changes to organization and syntax of cue traits. now supports setting trait presets inside of class
-More comments

#rnx

diff --git a/Engine/Plugins/Runtime/NetworkPrediction/Source/NetworkPrediction/Private/NetworkedSimulationModelCues.cpp b/Engine/Plugins/Runtime/NetworkPrediction/Source/NetworkPrediction/Private/NetworkedSimulationModelCues.cpp
--- a/Engine/Plugins/Runtime/NetworkPrediction/Source/NetworkPrediction/Private/NetworkedSimulationModelCues.cpp
+++ b/Engine/Plugins/Runtime/NetworkPrediction/Source/NetworkPrediction/Private/NetworkedSimulationModelCues.cpp
@@ -5,21 +5,25 @@
 DEFINE_LOG_CATEGORY(LogNetSimCues);
 
 FGlobalCueTypeTable FGlobalCueTypeTable::Singleton;
-
 constexpr FNetworkSimTime FCueDispatcherTraitsBase::ReplicationWindow;
 
 // ------------------------------------------------------------------------------------------------------------------------------------------
-//	Mock Cue example
+//	Mock Cue example: Minimal example of NetSimCue pipeline
+//		(1)	- How to define cue types and traits
+//		(2) - How to define cue sets
+//		(3) - How to define handlers
+//		(4) - How to invoke cues in simulation
+//
 // ------------------------------------------------------------------------------------------------------------------------------------------
 
-// Minimal example of NetSimCue pipeline
 #define ENABLE_MOCK_CUE 0
 #if ENABLE_MOCK_CUE
 
 // ---------------------------------------------------------------------
-// Mock Cue Types
+// (1) Mock Cue Types
 // ---------------------------------------------------------------------
 
+// Example "Weak" cue (will default to NetSimCueTraits::Default traits, which are "Weak" non replicated, non resimulated). No NetSerialize/NetIdentical is required.
 struct FMockImpactCue
 {
 	NETSIMCUE_BODY();
@@ -29,22 +33,13 @@ struct FMockImpactCue
 
 	FVector ImpactLocation;
 
-	void NetSerialize(FArchive& Ar)
-	{
-		Ar << ImpactLocation;
-	}
-
-	bool NetIdentical(const FMockImpactCue& Other) const
-	{
-		const float ErrorTolerance = 1.f;
-		return ImpactLocation.Equals(Other.ImpactLocation, ErrorTolerance);
-	}
-
 };
-NETSIMCUE_REGISTER(FMockImpactCue, TEXT("Impact"));
+NETSIMCUE_REGISTER(FMockImpactCue, TEXT("Impact")); // Must go in .cpp
+static_assert(TNetSimCueTraits<FMockImpactCue>::InvokeMask == NetSimCueTraits::Default::InvokeMask, "Traits Errors"); // Not required: just asserting what this classes traits are
 
 // ----------------
 
+// Example "Strong" cue. Explicitly sets NetSimCueTraits::Strong and implements NetSerialize/NetIdentical (will not compile without these functions) 
 struct FMockDamageCue
 {
 	NETSIMCUE_BODY();
@@ -56,6 +51,8 @@ struct FMockDamageCue
 	int32 SourceID;
 	int32 DamageType;
 	FVector HitLocation;
+
+	using Traits = NetSimCueTraits::Strong; // Sets our traits via preset.
 	
 	void NetSerialize(FArchive& Ar)
 	{
@@ -70,8 +67,10 @@ struct FMockDamageCue
 		return SourceID == Other.SourceID && DamageType && Other.DamageType && HitLocation.Equals(Other.HitLocation, ErrorTolerance);
 	}
 };
-NETSIMCUE_REGISTER(FMockDamageCue, TEXT("Damage"));
+NETSIMCUE_REGISTER(FMockDamageCue, TEXT("Damage")); // Must go in .cpp
+static_assert(TNetSimCueTraits<FMockDamageCue>::InvokeMask == NetSimCueTraits::Strong::InvokeMask, "Traits Errors"); // Not required: just asserting what this classes traits are
 
+// Alternative way of explicitly setting traits
 struct FMockHealingCue
 {
 	NETSIMCUE_BODY();
@@ -82,6 +81,8 @@ struct FMockHealingCue
 
 	int32 SourceID;
 	int32 HealingAmount;
+
+	using Traits = TNetSimCueTraitsExplicit<(uint8)ESimulationTickContext::Authority, ENetSimCueReplicationTarget::All>; // Sets our traits explicitly
 	
 	void NetSerialize(FArchive& Ar)
 	{
@@ -95,40 +96,63 @@ struct FMockHealingCue
 		return SourceID == Other.SourceID && HealingAmount == Other.HealingAmount;
 	}
 };
-NETSIMCUE_REGISTER(FMockHealingCue, TEXT("Healing"));
+NETSIMCUE_REGISTER(FMockHealingCue, TEXT("Healing")); // Must go in .cpp
+static_assert(TNetSimCueTraits<FMockHealingCue>::InvokeMask == (uint8)ESimulationTickContext::Authority, "Traits Errors"); // Not required: just asserting what this classes traits are
 
 // ---------------------------------------------------------------------
-// Mock Handlers (equevilent to actor component hierarchy)
+// (2) Cue Sets: just a collection of cue types that a simulation emits.
+//
+//	This isn't really a core type, it is just a pattern for grouping cue types
+//	together so that handlers can register with sets of cues rather than each cue individually
+//
+//	This example illustrates how a base set can be built off of by child simulations
 // ---------------------------------------------------------------------
 
-class TestHandlerBase
+struct FMockCueSet_Base
 {
-public:
-
 	template<typename TDispatchTable>
 	static void RegisterNetSimCueTypes(TDispatchTable& DispatchTable)
 	{
 		DispatchTable.template RegisterType<FMockImpactCue>();
 	}
+};
 
+struct FMockCueSet_Child
+{
+	template<typename TDispatchTable>
+	static void RegisterNetSimCueTypes(TDispatchTable& DispatchTable)
+	{
+		FMockCueSet_Base::RegisterNetSimCueTypes(DispatchTable); // Superset of the base cueset
+		DispatchTable.template RegisterType<FMockDamageCue>();
+		DispatchTable.template RegisterType<FMockHealingCue>();
+	}
+};
+
+
+// ---------------------------------------------------------------------
+// (3) Mock Handlers 
+//	-Implements ::HandeCue function to receive the cue (thats it!)
+//	-Meant to be implemented by owning actor/component that is driving the NetworkedSimulation
+//	-E.g, in practice, these aren't standalone classes. These are just functions/body that you put on your actor/component.
+//
+//	Again, we show how a base handler can be implemented and how a child class can extend off of that.
+// ---------------------------------------------------------------------
+
+class TestHandlerBase
+{
+public:
 	void HandleCue(const FMockImpactCue& ImpactCue, const FNetSimCueSystemParamemters& SystemParameters)
 	{
 		UE_LOG(LogTemp, Warning, TEXT("Impact!"));
 	}
 };
 
+NETSIMCUESET_REGISTER(TestHandlerBase, FMockCueSet_Base); // Must go in cpp. Says "TestHandlerBase" can accepts cues in "FMockCueSet_Base"
+
 class TestHandlerChild : public TestHandlerBase
 {
 public:
 
-	template<typename TDispatchTable>
-	static void RegisterNetSimCueTypes(TDispatchTable& DispatchTable)
-	{
-		TestHandlerBase::RegisterNetSimCueTypes(DispatchTable);
-		DispatchTable.template RegisterType<FMockDamageCue>();
-		DispatchTable.template RegisterType<FMockHealingCue>();
-	}
-
 	using TestHandlerBase::HandleCue; // Required for us to "use" the HandleCue methods in our parent class
 
 	void HandleCue(const FMockDamageCue& DamageCue, const FNetSimCueSystemParamemters& SystemParameters)
@@ -142,19 +166,19 @@ public:
 	}
 };
 
-NETSIMCUEHANDLER_REGISTER(TestHandlerChild);
+NETSIMCUESET_REGISTER(TestHandlerChild, FMockCueSet_Child); // Must go in cpp. Says "TestHandlerBase" can accepts cues in "FMockCueSet_Base" // Must go in cpp. Says "TestHandlerBase" can accepts cues in "FMockCueSet_Base"
 
 void TestCues()
 {
 	// -------------------------------------
-	// Simuation/Invoke
+	// (4) Simuation/Invoke
 	// -------------------------------------
 
 	struct FBaseSimulation
 	{
 		static void TickSimulation(FNetSimCueDispatcher& Dispatcher)
 		{
-			Dispatcher.Invoke<FMockImpactCue>(FVector(1.f, 2.f, 3.f));
+			Dispatcher.Invoke<FMockImpactCue>(FVector(1.f, 2.f, 3.f)); // Constructing emplace avoids copy and moves. This is the ideal way to invoke!
 		}
 	};
 
diff --git a/Engine/Plugins/Runtime/NetworkPrediction/Source/NetworkPrediction/Public/NetworkedSimulationModelCueTraits.h b/Engine/Plugins/Runtime/NetworkPrediction/Source/NetworkPrediction/Public/NetworkedSimulationModelCueTraits.h
new file mode 100644
--- /dev/null
+++ b/Engine/Plugins/Runtime/NetworkPrediction/Source/NetworkPrediction/Public/NetworkedSimulationModelCueTraits.h
@@ -0,0 +1,226 @@
+// Copyright 1998-2019 Epic Games, Inc. All Rights Reserved.
+#pragma  once
+
+// ----------------------------------------------------------------------------------------------------------------------------------------------
+//	NetSimCue Traits: compile time settings for NetSimeCue types that determine who can invoke the event and who it replicates to.
+//	See "Mock Cue Example" in NetworkedSimulationModelCues.cpp for minimal examples for using/setting traits.
+//
+//	There are two traits:
+//
+//	// Who can Invoke this Cue in their simulation (if this test fails, the Invoke call is supressed locally)
+//	static constexpr uint8 InvokeMask { InInvokeMask };
+//
+//	// Does the cue replicate? (from authority). This will also determine if the cue needs to be saved locally for NetIdentical tests (to avoid double playing)
+//	static constexpr ENetSimCueReplicationTarget ReplicationTarget { InReplicationTarget };
+//
+// ----------------------------------------------------------------------------------------------------------------------------------------------
+
+enum class ESimulationTickContext : uint8
+{
+	None			= 0,
+	Authority		= 1 << 0,
+	Predict			= 1 << 1,
+	Resimulate		= 1 << 2,
+
+	AuthorityPredict	= Authority | Predict,
+	PredictResimulate	= Predict | Resimulate,
+
+	All				= (Authority | Predict | Resimulate),
+};
+
+enum class ENetSimCueReplicationTarget : uint8
+{
+	None,			// Do not replicate cue to anyone
+	Interpolators,	// Only replicate to / accept on clients that are interpolating (not running the simulation themselves)
+	All,			// Replicate to everyone
+};
+
+
+// ----------------------------------------------------------------------------------------------------------------------------------------------
+//	NetSimCue Traits Presets. These are reasonable settings that may be appropriate in common cases. Presets are provided so that individual settings
+//	are not duplicated throughout the code base and to establish consistent vocabulary for the common types.
+//
+//	For quick reference, we expect these to be the most common:
+//		NetSimCueTraits::Weak - The default trait type and cheapest cue to use. No replication or NetIdentical testing. Never rolled back. Predicted but never resimulated.
+//		NetSimCueTraits::Strong - The most robust trait type. Will replicate to everyone and will rollback/resimulate. Must implement NetSerialize/NetIdentical. Most expensive.
+//
+//	The other presets fall somewhere in the middle and require thought/nuance to decide if they are right for your case.
+//
+// ----------------------------------------------------------------------------------------------------------------------------------------------
+
+namespace NetSimCueTraits
+{
+	// Default Traits: if you do not explicitly set Traits on your class, this is what it defaults to
+	using Default = struct Weak;
+
+	// Non replicated cue that only plays during "latest" simulate. Will not be played during rewind/resimulate.
+	// Lightest weight cue. Best used for cosmetic, non critical events. Footsteps, impact effects, etc.
+	struct Weak
+	{
+		static constexpr uint8 InvokeMask { (uint8)ESimulationTickContext::Authority | (uint8)ESimulationTickContext::Predict };
+		static constexpr ENetSimCueReplicationTarget ReplicationTarget { ENetSimCueReplicationTarget::None };
+	};
+
+	// Same as above but will only play on predicting client, not authority
+	struct WeakClientOnly
+	{
+		static constexpr uint8 InvokeMask { (uint8)ESimulationTickContext::Predict };
+		static constexpr ENetSimCueReplicationTarget ReplicationTarget { ENetSimCueReplicationTarget::None };
+	};
+
+	// Will only play on the authority path and not replicate to anyone else
+	struct AuthorityOnly
+	{
+		static constexpr uint8 InvokeMask { (uint8)ESimulationTickContext::Authority };
+		static constexpr ENetSimCueReplicationTarget ReplicationTarget { ENetSimCueReplicationTarget::None };
+	};
+
+	// Only invoked on authority and will replicate to everyone else. Not predicted so controlling client will see delays!
+	// Best for events that are critical that cannot be rolled back/undown and do not need to be predicted.
+	struct ReplicatedNonPredicted
+	{
+		static constexpr uint8 InvokeMask { (uint8)ESimulationTickContext::Authority };
+		static constexpr ENetSimCueReplicationTarget ReplicationTarget { ENetSimCueReplicationTarget::All };
+	};
+
+	// Replicated to interpolating proxies, predicted by autonomous proxy
+	// Best for events you want everyone to see but don't need to get perfect in the predicting cases: doesn't need to rollback and cheap on cpu (no NetIdentical tests on predicted path)
+	struct ReplicatedXOrPredicted
+	{
+		static constexpr uint8 InvokeMask { (uint8)ESimulationTickContext::Authority | (uint8)ESimulationTickContext::Predict };
+		static constexpr ENetSimCueReplicationTarget ReplicationTarget { ENetSimCueReplicationTarget::Interpolators };
+	};
+
+	// Invoked and replicated to all. NetIdentical testing to avoid double playing, rollbackable so that it can (re)play during resimulates
+	// Most expensive (bandwidth/CPU) and requires rollback callbacks to be implemented to be correct. But will always be shown "as correct as possible"
+	struct Strong
+	{
+		static constexpr uint8 InvokeMask { (uint8)ESimulationTickContext::All };
+		static constexpr ENetSimCueReplicationTarget ReplicationTarget { ENetSimCueReplicationTarget::All };
+	};
+
+	// Non replicated but if a resimulate happens, the cue is undone and replayed.
+	// This is not common and doesn't really have a clear use case. But the system can support it.
+	struct NonReplicatedResimulated
+	{
+		static constexpr uint8 InvokeMask { (uint8)ESimulationTickContext::All };
+		static constexpr ENetSimCueReplicationTarget ReplicationTarget { ENetSimCueReplicationTarget::None };
+	};
+}
+
+// Explicit trait settings. This can be used to explicitly set your traits without using a preset.
+template<uint8 InInvokeMask, ENetSimCueReplicationTarget InReplicationTarget>
+struct TNetSimCueTraitsExplicit
+{
+	// Who can Invoke this Cue in their simulation (if this test fails, the Invoke call is supressed locally)
+	static constexpr uint8 InvokeMask { InInvokeMask };
+
+	// Does the cue replicate? (from authority). This will also determine if the cue needs to be saved locally for NetIdentical tests (to avoid double playing)
+	static constexpr ENetSimCueReplicationTarget ReplicationTarget { InReplicationTarget };
+};
+
+
+// ----------------------------------------------------------------------------------------------------------------------------------------------
+//	NetSimCue Traits template helpers
+// ----------------------------------------------------------------------------------------------------------------------------------------------
+
+// SFINAE helper
+template<typename T>
+struct TToVoid
+{
+	using type  = void;
+};
+
+// Helper to see if T has a Traits type or not
+template<typename T, typename dummy=void>
+struct THasNetSimCueTraits
+{
+	enum { Value = false };
+};
+
+template<typename T>
+struct THasNetSimCueTraits<T, typename TToVoid<typename T::Traits>::type>
+{
+	enum { Value = true };
+};
+
+// Selects explicit traits set by T
+template<typename T, bool HasTraits = THasNetSimCueTraits<T>::Value>
+struct TSelectNetSimCueTraits
+{
+	using Traits = typename T::Traits;
+};
+
+// Fall back to Default traits
+template<typename T>
+struct TSelectNetSimCueTraits<T, false>
+{
+	using Traits = NetSimCueTraits::Default;
+};
+
+// Actual trait struct that we use to look up traits. The ways this can be set:
+//	1. Explicitly specialize TNetSimCueTraits for your type (non intrusive, but does not inherit)
+//	2. Explicitly set Traits inside your struct. E.g:
+//			using Traits = TNetSimCueTraits_Strong; (intrustive but more concise, does support inherited types)
+//	3. Automatically falls back to NetSimCueTraits::Default if not explicitly set above
+template<typename TCue>
+struct TNetSimCueTraits
+{
+	using Traits = typename TSelectNetSimCueTraits<TCue>::Traits;
+
+	static constexpr uint8 InvokeMask { Traits::InvokeMask };
+	static constexpr ENetSimCueReplicationTarget ReplicationTarget { Traits::ReplicationTarget };
+};
+
+// Type requirements: helper to determine if NetSerialize/NetIdentical functions need to be defined for user types based on the above traits
+template <typename TCue>
+struct TNetSimCueTypeRequirements
+{
+	enum {
+		// NetSerialize is required if we ever need to replicate
+		RequiresNetSerialize = (TNetSimCueTraits<TCue>::ReplicationTarget != ENetSimCueReplicationTarget::None),
+		// Likewise for NetIdentical, but also if we plan to invoke during Resimulate too (even if non replicated, we use NetIdentical for comparisons. though this is probably a non practical use case).
+		RequiresNetIdentical = (TNetSimCueTraits<TCue>::ReplicationTarget != ENetSimCueReplicationTarget::None) || (TNetSimCueTraits<TCue>::InvokeMask & (uint8)ESimulationTickContext::Resimulate)
+	};
+};
+
+// ----------------------------------------------------------------------------------------------------------------------------------------------
+//	Function checking helpers - helps us do clean checks for member function (NetSerialize/NetIdentical) when registering types
+// ----------------------------------------------------------------------------------------------------------------------------------------------
+
+// Helper to compile time check if NetSerialize exists
+GENERATE_MEMBER_FUNCTION_CHECK(NetSerialize, void,, FArchive&);
+
+// Helper to compile time check if NetIdentical exists (since argument is template type, must be wrapped in helper struct)
+template<typename TCue>
+struct THasNetIdenticalHelper
+{
+	GENERATE_MEMBER_FUNCTION_CHECK(NetIdentical, bool, const, const TCue&);
+	enum { Value = THasMemberFunction_NetIdentical<TCue>::Value };
+};
+
+// Helper to call NetIdentical if type defines it
+template<typename TCue, bool Enabled=THasNetIdenticalHelper<TCue>::Value>
+struct TNetCueNetIdenticalHelper
+{
+	static bool CallNetIdenticalOrNot(const TCue& Cue, const TCue& Other) { ensure(false); return false; } // This should never be hit by cue types that don't need to NetIdentical
+};
+
+template<typename TCue>
+struct TNetCueNetIdenticalHelper<TCue, true>
+{
+	static bool CallNetIdenticalOrNot(const TCue& Cue, const TCue& Other) { return Cue.NetIdentical(Other); }
+};
+
+// Helper to call NetSerialize if type defines it
+template<typename TCue, bool Enabled=THasMemberFunction_NetSerialize<TCue>::Value>
+struct TNetCueNetSerializeHelper
+{
+	static void CallNetSerializeOrNot(TCue& Cue, FArchive& Ar) { ensure(false); } // This should never be hit by cue types that don't need to NetSerialize
+};
+
+template<typename TCue>
+struct TNetCueNetSerializeHelper<TCue, true>
+{
+	static void CallNetSerializeOrNot(TCue& Cue, FArchive& Ar) { Cue.NetSerialize(Ar); }
+};
\ No newline at end of file
diff --git a/Engine/Plugins/Runtime/NetworkPrediction/Source/NetworkPrediction/Public/NetworkedSimulationModelCues.h b/Engine/Plugins/Runtime/NetworkPrediction/Source/NetworkPrediction/Public/NetworkedSimulationModelCues.h
--- a/Engine/Plugins/Runtime/NetworkPrediction/Source/NetworkPrediction/Public/NetworkedSimulationModelCues.h
+++ b/Engine/Plugins/Runtime/NetworkPrediction/Source/NetworkPrediction/Public/NetworkedSimulationModelCues.h
@@ -3,6 +3,37 @@
 
 #include "Engine/EngineTypes.h"
 #include "NetworkedSimulationModelTime.h"
+#include "NetworkedSimulationModelCueTraits.h"
+
+/*=============================================================================
+Networked Simulation Cues
+
+Cues are events that are outputted by a Networked Simulation. They have the following properties:
+-Cues are defined by a user struct (custom data payload). They are *invoked* during a SimulationTick and are *dispatched* to a *handler*.
+-They are dispatched *after* the simulation is done ticking (via TNetSimCueDispatcher::DispatchCueRecord which is called during the owning component's tick).
+-They should not affect the simulation. Or rather, if they do affect the simulation, it will be during the actor tick, effectively the same as any "out of band" modifications.
+-They provide automatic replication and invocation settings (traits). They will not "double play" during resimulates.
+-They are time aware. When dispatched, the receiver is given how much time has passed (relative to local "head" time) since the invocation.
+-They are rollback aware. When dispatched, the receiver is given a callback that will be invoked if the cue needs to rollback (undo itself). 
+-The callback is not given in contexts when a rollback is impossible: e.g, on the authority.
+
+Notes on reliability:
+-Cues are unreliable in nature. If you join a game in progress, an actor driving the sim suddenly becomes relevant: you will not get all past events.
+-If network becomes saturated or drops, you may miss events too.
+-Order of received cues is guaranteed, but we can't promise there won't be gaps/missing cues!
+-In other words: do not use cues in stateful ways! Cues should be used for transient events ("NotifyExplosion") not state ("NotifySetDestroyed(true)")
+-For state transitions, just use FinalizeFrame to detect these changes. "State transitions" are 100% reliable, but (obviously) cannot use data that is not in the Sync/Aux buffer.
+
+Notes on "simulation affect events" (E.g, *not* NetSimCues)
+-If you have an event that needs to affect the simulation during the simulation - that is seen as an extension of the simulation and is "up to the user" to implement.
+-In other words, handle the event yourself inline with the simulation. That means directly broadcasting/calling other functions/into other classes/etc inside your simulation.
+-If your event has state mutation on the handler side, that is a hazard (e.g, state that the network sim is not aware of, but is used in the event code which is an extension of the simulation)
+-In these cases I would recommend: A) on the handler side, don't write to any non-networked-sim state if non authority or B) just don't handle the event on non authority. (Expect corrections)
+-If the handler side doesn't have state hazards, say a teleporting volume that always does the same thing: there is no reason everyone can't run the event. Its an extension of the simulation.
+
+See "Mock Cue Example" in NetworkedSimulationModelCues.cpp for minimal example of implementing the Cue types and Handler classes.
+
+=============================================================================*/
 
 NETWORKPREDICTION_API DECLARE_LOG_CATEGORY_EXTERN(LogNetSimCues, Log, All);
 
@@ -12,146 +43,8 @@ NETWORKPREDICTION_API DECLARE_LOG_CATEGORY_EXTERN(LogNetSimCues, Log, All);
 
 using FNetSimCueTypeId = NETSIMCUE_TYPEID_TYPE;
 
-enum class ESimulationTickContext : uint8
-{
-	None			= 0,
-	Authority		= 1 << 0,
-	Predict			= 1 << 1,
-	Resimulate		= 1 << 2,
-
-	AuthorityPredict	= Authority | Predict,
-	PredictResimulate	= Predict | Resimulate,
-
-	All				= (Authority | Predict | Resimulate),
-};
-
-enum class ENetSimCueReplicationTarget : uint8
-{
-	None,			// Do not replicate cue to anyone
-	Interpolators,	// Only replicate to / accept on clients that are interpolating (not running the simulation themselves)
-	All,			// Replicate to everyone
-};
-
-// ----------------------------------------------------------------------------------------------------------------------------------------------
-//	NetSimCue Traits: how to configure how a NetSimCue will dispatch within the NetworkedSimulationModel
-// ----------------------------------------------------------------------------------------------------------------------------------------------
-
-struct TNetSimCueTraitsBase
-{
-	// Who can Invoke this Cue in their simulation (if this test fails, the Invoke call is supressed locally)
-	static constexpr uint8 InvokeMask { (uint8)ESimulationTickContext::All };
-
-	// Does the cue replicate? (from authority). This will also determine if the cue needs to be saved locally for NetIdentical tests (to avoid double playing)
-	static constexpr ENetSimCueReplicationTarget ReplicationTarget { ENetSimCueReplicationTarget::All };
-};
-
-// Preset: non replicated cue that only plays during "latest" simulate. Will not be played during rewind/resimulate.
-// Lightest weight cue. Best used for cosmetic, non critical events. Footsteps, impact effects, etc.
-struct TNetSimCueTraits_Weak
-{
-	static constexpr uint8 InvokeMask { (uint8)ESimulationTickContext::Authority | (uint8)ESimulationTickContext::Predict };
-	static constexpr ENetSimCueReplicationTarget ReplicationTarget { ENetSimCueReplicationTarget::None };
-};
-
-// Preset: Replicated, non predicted. Only invoked on authority and will replicate to everyone else.
-// Best for events that are critical that cannot be rolled back/undown and do not need to be predicted.
-struct TNetSimCueTraits_ReplicatedNonPredicted
-{
-	static constexpr uint8 InvokeMask { (uint8)ESimulationTickContext::Authority };
-	static constexpr ENetSimCueReplicationTarget ReplicationTarget { ENetSimCueReplicationTarget::All };
-};
-
-// Preset: Replicated to interpolating proxies, predicted by autonomous proxy
-// Best for events you want everyone to see but don't need to get perfect in the predicting cases: doesn't need to rollback and cheap on cpu (no unique tests on predicted path)
-struct TNetSimCueTraits_ReplicatedXOrPredicted
-{
-	static constexpr uint8 InvokeMask { (uint8)ESimulationTickContext::Authority | (uint8)ESimulationTickContext::Predict };
-	static constexpr ENetSimCueReplicationTarget ReplicationTarget { ENetSimCueReplicationTarget::Interpolators };
-};
-
-// Preset: Invoked and replicated to all. Uniqueness testing to avoid double playing, rollbackable so that it can (re)play during resimulates
-// Most expensive (bandwidth and CPU for uniqueness testing) and requires rollback callbacks to be implemented to be correct
-struct TNetSimCueTraits_Strong
-{
-	static constexpr uint8 InvokeMask { (uint8)ESimulationTickContext::All };
-	static constexpr ENetSimCueReplicationTarget ReplicationTarget { ENetSimCueReplicationTarget::All };
-};
-
-// Preset: Non replicated but if a resimulate happens, the cue is undone and replayed.
-// This is not common and doesn't really have a clear use case. But the system can support it.
-struct TNetSimCueTraits_NonReplicatedResimulated
-{
-	static constexpr uint8 InvokeMask { (uint8)ESimulationTickContext::All };
-	static constexpr ENetSimCueReplicationTarget ReplicationTarget { ENetSimCueReplicationTarget::None };
-};
-
-// Actual trait struct that we use to look up traits. User cues must specialize this
-template<typename TCue>
-struct TCueHandlerTraits : public TNetSimCueTraitsBase
-{
-};
-
-
-// Type requirements: helper to determine if NetSerialize/NetIdentical functions need to be defined for user types based on traits
-template <typename TCue>
-struct TNetSimCueTypeRequirements
-{
-	enum {
-		// NetSerialize is required if we ever need to replicate
-		RequiresNetSerialize = (TCueHandlerTraits<TCue>::ReplicationTarget != ENetSimCueReplicationTarget::None),
-		// Likewise for NetIdentical, but also if we plan to invoke during Resimulate too (even if non replicated, we use NetIdentical for comparisons. though this is probably a non practical use case).
-		RequiresNetIdentical = (TCueHandlerTraits<TCue>::ReplicationTarget != ENetSimCueReplicationTarget::None) || (TCueHandlerTraits<TCue>::InvokeMask & (uint8)ESimulationTickContext::Resimulate)
-	};
-};
-
-// Helper to compile time check if NetSerialize exists
-GENERATE_MEMBER_FUNCTION_CHECK(NetSerialize, void,, FArchive&);
-
-// Helper to compile time check if NetIdentical exists (since argument is template type, must be wrapped in helper struct)
-template<typename TCue>
-struct THasNetIdenticalHelper
-{
-	GENERATE_MEMBER_FUNCTION_CHECK(NetIdentical, bool, const, const TCue&);
-	enum { Value = THasMemberFunction_NetIdentical<TCue>::Value };
-};
-
-// Helper to call NetIdentical if type defines it
-template<typename TCue, bool Enabled=THasNetIdenticalHelper<TCue>::Value>
-struct TNetCueNetIdenticalHelper
-{
-	static bool CallNetIdenticalOrNot(const TCue& Cue, const TCue& Other) { ensure(false); return false; } // This should never be hit by cue types that don't need to NetIdentical
-};
-
-template<typename TCue>
-struct TNetCueNetIdenticalHelper<TCue, true>
-{
-	static bool CallNetIdenticalOrNot(const TCue& Cue, const TCue& Other) { return Cue.NetIdentical(Other); }
-};
-
-// Helper to call NetSerialize if type defines it
-template<typename TCue, bool Enabled=THasMemberFunction_NetSerialize<TCue>::Value>
-struct TNetCueNetSerializeHelper
-{
-	static void CallNetSerializeOrNot(TCue& Cue, FArchive& Ar) { ensure(false); } // This should never be hit by cue types that don't need to NetSerialize
-};
-
-template<typename TCue>
-struct TNetCueNetSerializeHelper<TCue, true>
-{
-	static void CallNetSerializeOrNot(TCue& Cue, FArchive& Ar) { Cue.NetSerialize(Ar); }
-};
-
-// Traits for TNetSimCueDispatcher
-struct NETWORKPREDICTION_API FCueDispatcherTraitsBase
-{
-	// Window for replicating a NetSimCue. That is, after a cue is invoked, it has ReplicationWindow time before it will be pruned.
-	// If a client does not get a net update for the sim in this window, they will miss the event.
-	static constexpr FNetworkSimTime ReplicationWindow = FNetworkSimTime::FromRealTimeMS(200);
-};
-template<typename T> struct TCueDispatcherTraits : public FCueDispatcherTraitsBase { };
-
 // ----------------------------------------------------------------------------------------------------------------------------------------------
-//	Wrapper
+//	Wrapper: wraps the actual user NetSimCue. We want to avoid virtualizing functions on the actual user types.
 // ----------------------------------------------------------------------------------------------------------------------------------------------
 
 struct FNetSimCueWrapperBase
@@ -189,7 +82,7 @@ struct TNetSimCueWrapper : FNetSimCueWrapperBase
 
 	ENetSimCueReplicationTarget GetReplicationTarget() const override final
 	{
-		return TCueHandlerTraits<TCue>::ReplicationTarget;
+		return TNetSimCueTraits<TCue>::ReplicationTarget;
 	}
 
 	TCue Instance;
@@ -204,6 +97,8 @@ struct FNetSimCueCallbacks
 	/** Callback to rollback any side effects of the cue. */
 	DECLARE_MULTICAST_DELEGATE(FOnRollback)
 	FOnRollback	OnRollback;
+
+	// It may make sense to add an "OnConfirmed" that will let the user know a rollback will no longer be possible
 };
 
 /** System parameters for NetSimCue events */
@@ -441,6 +336,15 @@ TCueDispatchTable<TCueHandler> TCueDispatchTable<TCueHandler>::Singleton;
 //	-Holds recorded cue state for replication
 // ------------------------------------------------------------------------------------------------------------------------------------------------------------
 
+// Traits for TNetSimCueDispatcher
+struct NETWORKPREDICTION_API FCueDispatcherTraitsBase
+{
+	// Window for replicating a NetSimCue. That is, after a cue is invoked, it has ReplicationWindow time before it will be pruned.
+	// If a client does not get a net update for the sim in this window, they will miss the event.
+	static constexpr FNetworkSimTime ReplicationWindow = FNetworkSimTime::FromRealTimeMS(200);
+};
+template<typename T> struct TCueDispatcherTraits : public FCueDispatcherTraitsBase { };
+
 // Non-templated, "networking model independent" base: this is what the pure simulation code gets to invoke cues. 
 struct FNetSimCueDispatcher
 {
@@ -458,38 +362,9 @@ struct FNetSimCueDispatcher
 	{
 		if (EnsureValidContext())
 		{
-			if ((TCueHandlerTraits<T>::InvokeMask & (uint8)Context.TickContext) > 0)
+			if ((TNetSimCueTraits<T>::InvokeMask & (uint8)Context.TickContext) > 0)
 			{
-				/*
-				{
-					// Whether we go in the transient list. Transient cues are dumped after dispatching (not saved over multiple frames for uniqueness comparisons during Invoke or NetSerialize)
-					const bool bTransient = (Context.TickContext == ESimulationTickContext::Authority && TCueHandlerTraits<T>::ReplicationTarget == ENetSimCueReplicationTarget::None)	  // Transient on authority if it doesn't replicate
-						|| ((TCueHandlerTraits<T>::InvokeMask & (uint8)ESimulationTickContext::Resimulate) == 0 && TCueHandlerTraits<T>::ReplicationTarget != ENetSimCueReplicationTarget::All); // Transient elsewhere if it won't replicate to us and it won't be invoked during resimulates
-
-					// Whether this cue should be dispatched with rollback callbacks. This is a trait of the cue type + context (e.g, authority will never rollback)
-					const bool bAllowRollback = TCueHandlerTraits<T>::Rollbackable && !bTransient && Context.TickContext != ESimulationTickContext::Authority;
-
-					// Is this already confirmed? (it is if we can't roll it back, or if it won't be replicated to us)
-					const bool bNetConfirmed = !bAllowRollback || TCueHandlerTraits<T>::ReplicationTarget != ENetSimCueReplicationTarget::All;
-				}
-
-
-				{
-					// Whether this cue should be dispatched with rollback callbacks. Authority never rolls back, clients rollback if they plan to invoke during resims or if the cue is replicated to everyone
-					const bool bRequiresRollbackForReplication = Context.TickContext != ESimulationTickContext::Authority && (TCueHandlerTraits<T>::ReplicationTarget == ENetSimCueReplicationTarget::All);
-					const bool bRequiresRollbackForResimulate = Context.TickContext != ESimulationTickContext::Authority && (TCueHandlerTraits<T>::InvokeMask & (uint8)ESimulationTickContext::Resimulate);
-					const bool bAllowRollback = bRequiresRollbackForReplication || bRequiresRollbackForResimulate;
-
-					// Whether we go in the transient list. Transient cues are dumped after dispatching (not saved over multiple frames for uniqueness comparisons during Invoke or NetSerialize)
-					const bool bTransient = !bAllowRollback &&
-						(Context.TickContext == ESimulationTickContext::Authority && (TCueHandlerTraits<T>::ReplicationTarget == ENetSimCueReplicationTarget::None)) ||
-						(Context.TickContext != ESimulationTickContext::Authority && (TCueHandlerTraits<T>::ReplicationTarget != ENetSimCueReplicationTarget::All));
-
-					const bool bNetConfirmed = !bAllowRollback || TCueHandlerTraits<T>::ReplicationTarget != ENetSimCueReplicationTarget::All;
-				}
-				*/
-
-				const bool bSupportsResimulate = (TCueHandlerTraits<T>::InvokeMask & (uint8)ESimulationTickContext::Resimulate) > 0;
+				const bool bSupportsResimulate = (TNetSimCueTraits<T>::InvokeMask & (uint8)ESimulationTickContext::Resimulate) > 0;
 
 				bool bAllowRollback = false;	// Whether this cue should be dispatched with rollback callbacks.
 				bool bTransient = false;		// Whether we go in the transient list. Transient cues are dumped after dispatching (not saved over multiple frames for uniqueness comparisons during Invoke or NetSerialize)
@@ -498,14 +373,14 @@ struct FNetSimCueDispatcher
 				{
 					// Authority: never rolls back, is already confirmed, and can treat cue as transient if it doesn't have to replicate it
 					bAllowRollback = false;
-					bTransient = TCueHandlerTraits<T>::ReplicationTarget == ENetSimCueReplicationTarget::None;
+					bTransient = TNetSimCueTraits<T>::ReplicationTarget == ENetSimCueReplicationTarget::None;
 					bNetConfirmed = true;
 				}
 				else
 				{
 					// Everyone else that is running the simulation (since this in ::Invoke which is called from within the simulation)
 					// Rollback if it will replicate to us or if we plan to invoke this cue during resimulates. Transient and confirmed follow directly from this in the non authority case.
-					bAllowRollback = (TCueHandlerTraits<T>::ReplicationTarget == ENetSimCueReplicationTarget::All) || bSupportsResimulate;
+					bAllowRollback = (TNetSimCueTraits<T>::ReplicationTarget == ENetSimCueReplicationTarget::All) || bSupportsResimulate;
 					bTransient = !bAllowRollback;
 					bNetConfirmed = !bAllowRollback;
 				}
@@ -547,7 +422,7 @@ struct FNetSimCueDispatcher
 			}
 			else
 			{
-				UE_LOG(LogNetSimCues, Log, TEXT("%s .Suppressing Cue Invocation %s. Mask: %d. TickContext: %d"), *GetDebugName(), *FGlobalCueTypeTable::Get().GetTypeName(T::ID), TCueHandlerTraits<T>::InvokeMask, (int32)Context.TickContext);
+				UE_LOG(LogNetSimCues, Log, TEXT("%s .Suppressing Cue Invocation %s. Mask: %d. TickContext: %d"), *GetDebugName(), *FGlobalCueTypeTable::Get().GetTypeName(T::ID), TNetSimCueTraits<T>::InvokeMask, (int32)Context.TickContext);
 			}
 		}
 	}
diff --git a/Engine/Plugins/Runtime/NetworkPredictionExtras/Source/NetworkPredictionExtras/Public/MockAbilitySimulation.h b/Engine/Plugins/Runtime/NetworkPredictionExtras/Source/NetworkPredictionExtras/Public/MockAbilitySimulation.h
--- a/Engine/Plugins/Runtime/NetworkPredictionExtras/Source/NetworkPredictionExtras/Public/MockAbilitySimulation.h
+++ b/Engine/Plugins/Runtime/NetworkPredictionExtras/Source/NetworkPredictionExtras/Public/MockAbilitySimulation.h
@@ -136,6 +136,8 @@ struct FMockAbilityBlinkActivateCue
 
 	FVector_NetQuantize10 Destination;
 	uint8 RandomType; // Random value used to color code the effect. This is the test/prove out mispredictions
+
+	using Traits = NetSimCueTraits::Strong;
 	
 	void NetSerialize(FArchive& Ar)
 	{
@@ -151,8 +153,11 @@ struct FMockAbilityBlinkActivateCue
 	}
 };
 
-template<>
-struct TCueHandlerTraits<FMockAbilityBlinkActivateCue> : public TNetSimCueTraits_Strong { };
+// This way works too but is less concise:
+//template<>
+//struct TNetSimCueTraits<FMockAbilityBlinkActivateCue> : public TNetSimCueTraits_Strong { };
+
+static_assert( TNetSimCueTraits<FMockAbilityBlinkActivateCue>::ReplicationTarget == NetSimCueTraits::Strong::ReplicationTarget, "Traits error" );
 
 // ----------------------------------------------------------------------------------------------
 
@@ -198,6 +203,8 @@ struct FMockAbilityBlinkCue
 
 	FVector_NetQuantize10 StartLocation;
 	FVector_NetQuantize10 StopLocation;
+
+	using Traits = NetSimCueTraits::Strong;
 	
 	void NetSerialize(FArchive& Ar)
 	{
@@ -214,9 +221,6 @@ struct FMockAbilityBlinkCue
 	}
 };
 
-template<>
-struct TCueHandlerTraits<FMockAbilityBlinkCue> : public TNetSimCueTraits_Strong { };
-
 // -----------------------------------------------------------------------------------------------------
 // Subtypes of the BlinkCue - this is not an expected setup! This is done for testing/debugging so we can 
 // see the differences between the cue type traits in a controlled setup. See FMockAbilitySimulation::SimulationTick
@@ -224,15 +228,16 @@ struct TCueHandlerTraits<FMockAbilityBlinkCue> : public TNetSimCueTraits_Strong
 #define DECLARE_BLINKCUE_SUBTYPE(TYPE, TRAITS) \
  struct TYPE : public FMockAbilityBlinkCue { \
  template <typename... ArgsType> TYPE(ArgsType&&... Args) : FMockAbilityBlinkCue(Forward<ArgsType>(Args)...) { } \
+ using Traits = TRAITS; \
  void NetSerialize(FArchive& Ar) { FMockAbilityBlinkCue::NetSerialize(Ar); } \
  bool NetIdentical(const TYPE& Other) const { return FMockAbilityBlinkCue::NetIdentical(Other); } \
- NETSIMCUE_BODY(); }; \
- template<> struct TCueHandlerTraits<TYPE> : public TRAITS { };
+ NETSIMCUE_BODY(); };
+ 
 
-DECLARE_BLINKCUE_SUBTYPE(FMockAbilityBlinkCue_Weak, TNetSimCueTraits_Weak);
-DECLARE_BLINKCUE_SUBTYPE(FMockAbilityBlinkCue_ReplicatedNonPredicted, TNetSimCueTraits_ReplicatedNonPredicted);
-DECLARE_BLINKCUE_SUBTYPE(FMockAbilityBlinkCue_ReplicatedXOrPredicted, TNetSimCueTraits_ReplicatedXOrPredicted);
-DECLARE_BLINKCUE_SUBTYPE(FMockAbilityBlinkCue_Strong, TNetSimCueTraits_Strong);
+DECLARE_BLINKCUE_SUBTYPE(FMockAbilityBlinkCue_Weak, NetSimCueTraits::Weak);
+DECLARE_BLINKCUE_SUBTYPE(FMockAbilityBlinkCue_ReplicatedNonPredicted, NetSimCueTraits::ReplicatedNonPredicted);
+DECLARE_BLINKCUE_SUBTYPE(FMockAbilityBlinkCue_ReplicatedXOrPredicted, NetSimCueTraits::ReplicatedXOrPredicted);
+DECLARE_BLINKCUE_SUBTYPE(FMockAbilityBlinkCue_Strong, NetSimCueTraits::Strong);
 
 // The set of Cues the MockAbility simulation will invoke
 struct FMockAbilityCueSet
