NetworkPrediction Plugin
-Simplify/cleanup cue traits: specify invoke mask (auth, predict, simulate) + rollback bool and derive final simulation mask from this. This cuts down on number of invalid configurations, makes the code a bit easier to follow and gets rid of awkward uint8 casts.

Core:
-Made EnumHasAllFlags/EnumHasAnyFlags CONSTEXPR

diff --git a/Engine/Plugins/Runtime/NetworkPrediction/Source/NetworkPrediction/Private/NetworkedSimulationModelCues.cpp b/Engine/Plugins/Runtime/NetworkPrediction/Source/NetworkPrediction/Private/NetworkedSimulationModelCues.cpp
--- a/Engine/Plugins/Runtime/NetworkPrediction/Source/NetworkPrediction/Private/NetworkedSimulationModelCues.cpp
+++ b/Engine/Plugins/Runtime/NetworkPrediction/Source/NetworkPrediction/Private/NetworkedSimulationModelCues.cpp
@@ -82,7 +82,7 @@ struct FMockHealingCue
 	int32 SourceID;
 	int32 HealingAmount;
 
-	using Traits = TNetSimCueTraitsExplicit<(uint8)ESimulationTickContext::Authority, ENetSimCueReplicationTarget::All>; // Sets our traits explicitly
+	using Traits = TNetSimCueTraitsExplicit<ENetSimCueInvoker::All, ENetSimCueReplicationTarget::All, true>; // Sets our traits explicitly
 	
 	void NetSerialize(FArchive& Ar)
 	{
@@ -97,7 +97,7 @@ struct FMockHealingCue
 	}
 };
 NETSIMCUE_REGISTER(FMockHealingCue, TEXT("Healing")); // Must go in .cpp
-static_assert(TNetSimCueTraits<FMockHealingCue>::InvokeMask == (uint8)ESimulationTickContext::Authority, "Traits Errors"); // Not required: just asserting what this classes traits are
+static_assert(TNetSimCueTraits<FMockHealingCue>::InvokeMask == ENetSimCueInvoker::All, "Traits Errors"); // Not required: just asserting what this classes traits are
 
 // ---------------------------------------------------------------------
 // (2) Cue Sets: just a collection of cue types that a simulation emits.
diff --git a/Engine/Plugins/Runtime/NetworkPrediction/Source/NetworkPrediction/Public/NetworkedSimulationModelCueTraits.h b/Engine/Plugins/Runtime/NetworkPrediction/Source/NetworkPrediction/Public/NetworkedSimulationModelCueTraits.h
--- a/Engine/Plugins/Runtime/NetworkPrediction/Source/NetworkPrediction/Public/NetworkedSimulationModelCueTraits.h
+++ b/Engine/Plugins/Runtime/NetworkPrediction/Source/NetworkPrediction/Public/NetworkedSimulationModelCueTraits.h
@@ -5,29 +5,52 @@
 //	NetSimCue Traits: compile time settings for NetSimeCue types that determine who can invoke the event and who it replicates to.
 //	See "Mock Cue Example" in NetworkedSimulationModelCues.cpp for minimal examples for using/setting traits.
 //
-//	There are two traits:
+//	There are three traits:
 //
-//	// Who can Invoke this Cue in their simulation (if this test fails, the Invoke call is supressed locally)
-//	static constexpr uint8 InvokeMask { InInvokeMask };
+//	// Who can Invoke this Cue in their simulation (if this test fails, the Invoke call is suppressed locally)
+//	static constexpr ENetSimCueInvoker InvokeMask { InInvokeMask };
 //
-//	// Does the cue replicate? (from authority). This will also determine if the cue needs to be saved locally for NetIdentical tests (to avoid double playing)
+//	// Whether the cue will be invoked during resimulates (which will require cue to be rollbackable)
+//	static constexpr bool Resimulate { InResimulate };
+//
+//	// Does the cue replicate? (from authority). This will also determine if the cue needs to be saved locally for NetIdentical tests (to avoid double playing) and if needs to support rollback
 //	static constexpr ENetSimCueReplicationTarget ReplicationTarget { InReplicationTarget };
 //
 // ----------------------------------------------------------------------------------------------------------------------------------------------
 
+// When we run a SimulationTick, it will be done under one of these contexts. This isn't NetSimCue specific and probably makes sense to move to the NetSimModel types
 enum class ESimulationTickContext : uint8
 {
-	None			= 0,
-	Authority		= 1 << 0,
-	Predict			= 1 << 1,
-	Resimulate		= 1 << 2,
+	None				= 0,
+	Authority			= 1 << 0,	// The authority (usually "server" but could also mean client authoritative client)
+	Predict				= 1 << 1,	// Predicting client: autonomous proxy ("controlling client")
+	Resimulate			= 1 << 2,	// Predicting client during resimulate (simulation rollback into resimulate steps)
+	SimExtrapolate		= 1 << 3,	// Simulation extrapolation: non controlling client (simulated proxy) running the simulation to extrapolate
+	ResimExtrapolate	= 1 << 4,	// Simulation extrapolation during a reconcile (sim rolled back to server state then stepped again to "catch back up")
+};
+ENUM_CLASS_FLAGS(ESimulationTickContext);
 
-	AuthorityPredict	= Authority | Predict,
-	PredictResimulate	= Predict | Resimulate,
+// High level "who can invoke this". Does not take resimulate into account yet. All combinations of this enum are valid (though Authority | SimExtrapolate is a bit weird maybe)
+enum class ENetSimCueInvoker : uint8
+{
+	Authority			= (uint8)ESimulationTickContext::Authority,
+	Predict				= (uint8)ESimulationTickContext::Predict,
+	SimExtrapolate		= (uint8)ESimulationTickContext::SimExtrapolate,
 
-	All				= (Authority | Predict | Resimulate),
+	All					= Authority | Predict | SimExtrapolate
 };
+ENUM_CLASS_FLAGS(ENetSimCueInvoker);
 
+// Turns "Who can invoke this" + "Does this play during resimulate" into the final ESimulationTickContext mask we use at runtime. This is done so users can't make invalid configurations (like "Authority | Resimulate")
+static constexpr ESimulationTickContext GetSimTickMask(const ENetSimCueInvoker Invoker, const bool bAllowResimulate)
+{
+	ESimulationTickContext Mask = (ESimulationTickContext)Invoker;
+	ESimulationTickContext AutoMask = (bAllowResimulate && EnumHasAllFlags(Invoker, ENetSimCueInvoker::Predict)) ? ESimulationTickContext::Resimulate : ESimulationTickContext::None;
+	ESimulationTickContext SimMask = (bAllowResimulate && EnumHasAllFlags(Invoker, ENetSimCueInvoker::SimExtrapolate)) ? ESimulationTickContext::ResimExtrapolate : ESimulationTickContext::None;
+	return Mask | AutoMask | SimMask;
+}
+
+// Who a NetSimCue should replicate/be accepted by
 enum class ENetSimCueReplicationTarget : uint8
 {
 	None,			// Do not replicate cue to anyone
@@ -35,7 +58,6 @@ enum class ENetSimCueReplicationTarget : uint8
 	All,			// Replicate to everyone
 };
 
-
 // ----------------------------------------------------------------------------------------------------------------------------------------------
 //	NetSimCue Traits Presets. These are reasonable settings that may be appropriate in common cases. Presets are provided so that individual settings
 //	are not duplicated throughout the code base and to establish consistent vocabulary for the common types.
@@ -57,21 +79,24 @@ namespace NetSimCueTraits
 	// Lightest weight cue. Best used for cosmetic, non critical events. Footsteps, impact effects, etc.
 	struct Weak
 	{
-		static constexpr uint8 InvokeMask { (uint8)ESimulationTickContext::Authority | (uint8)ESimulationTickContext::Predict };
+		static constexpr ENetSimCueInvoker InvokeMask { ENetSimCueInvoker::All };
+		static constexpr bool Resimulate { false };
 		static constexpr ENetSimCueReplicationTarget ReplicationTarget { ENetSimCueReplicationTarget::None };
 	};
 
 	// Same as above but will only play on predicting client, not authority
 	struct WeakClientOnly
 	{
-		static constexpr uint8 InvokeMask { (uint8)ESimulationTickContext::Predict };
+		static constexpr ENetSimCueInvoker InvokeMask { ENetSimCueInvoker::Predict };
+		static constexpr bool Resimulate { false };
 		static constexpr ENetSimCueReplicationTarget ReplicationTarget { ENetSimCueReplicationTarget::None };
 	};
 
 	// Will only play on the authority path and not replicate to anyone else
 	struct AuthorityOnly
 	{
-		static constexpr uint8 InvokeMask { (uint8)ESimulationTickContext::Authority };
+		static constexpr ENetSimCueInvoker InvokeMask { ENetSimCueInvoker::Authority };
+		static constexpr bool Resimulate { false };
 		static constexpr ENetSimCueReplicationTarget ReplicationTarget { ENetSimCueReplicationTarget::None };
 	};
 
@@ -79,15 +104,17 @@ namespace NetSimCueTraits
 	// Best for events that are critical that cannot be rolled back/undown and do not need to be predicted.
 	struct ReplicatedNonPredicted
 	{
-		static constexpr uint8 InvokeMask { (uint8)ESimulationTickContext::Authority };
+		static constexpr ENetSimCueInvoker InvokeMask { ENetSimCueInvoker::Authority };
+		static constexpr bool Resimulate { false };
 		static constexpr ENetSimCueReplicationTarget ReplicationTarget { ENetSimCueReplicationTarget::All };
 	};
 
-	// Replicated to interpolating proxies, predicted by autonomous proxy
+	// Replicated to interpolating proxies, predicted by autonomous/simulated proxy
 	// Best for events you want everyone to see but don't need to get perfect in the predicting cases: doesn't need to rollback and cheap on cpu (no NetIdentical tests on predicted path)
 	struct ReplicatedXOrPredicted
 	{
-		static constexpr uint8 InvokeMask { (uint8)ESimulationTickContext::Authority | (uint8)ESimulationTickContext::Predict };
+		static constexpr ENetSimCueInvoker InvokeMask { ENetSimCueInvoker::All };
+		static constexpr bool Resimulate { false };
 		static constexpr ENetSimCueReplicationTarget ReplicationTarget { ENetSimCueReplicationTarget::Interpolators };
 	};
 
@@ -95,7 +122,8 @@ namespace NetSimCueTraits
 	// Most expensive (bandwidth/CPU) and requires rollback callbacks to be implemented to be correct. But will always be shown "as correct as possible"
 	struct Strong
 	{
-		static constexpr uint8 InvokeMask { (uint8)ESimulationTickContext::All };
+		static constexpr ENetSimCueInvoker InvokeMask { ENetSimCueInvoker::All };
+		static constexpr bool Resimulate { true };
 		static constexpr ENetSimCueReplicationTarget ReplicationTarget { ENetSimCueReplicationTarget::All };
 	};
 
@@ -103,17 +131,21 @@ namespace NetSimCueTraits
 	// This is not common and doesn't really have a clear use case. But the system can support it.
 	struct NonReplicatedResimulated
 	{
-		static constexpr uint8 InvokeMask { (uint8)ESimulationTickContext::All };
+		static constexpr ENetSimCueInvoker InvokeMask { ENetSimCueInvoker::All };
+		static constexpr bool Resimulate { true };
 		static constexpr ENetSimCueReplicationTarget ReplicationTarget { ENetSimCueReplicationTarget::None };
 	};
 }
 
 // Explicit trait settings. This can be used to explicitly set your traits without using a preset.
-template<uint8 InInvokeMask, ENetSimCueReplicationTarget InReplicationTarget>
+template<ENetSimCueInvoker InInvokeMask, ENetSimCueReplicationTarget InReplicationTarget, bool InResimulate>
 struct TNetSimCueTraitsExplicit
 {
-	// Who can Invoke this Cue in their simulation (if this test fails, the Invoke call is supressed locally)
-	static constexpr uint8 InvokeMask { InInvokeMask };
+	// Who can Invoke this Cue in their simulation (if this test fails, the Invoke call is suppressed locally)
+	static constexpr ENetSimCueInvoker InvokeMask { InInvokeMask };
+
+	// Whether the cue will be invoked during resimulates (which will require cue to be rollbackable)
+	static constexpr bool Resimulate { InResimulate };
 
 	// Does the cue replicate? (from authority). This will also determine if the cue needs to be saved locally for NetIdentical tests (to avoid double playing)
 	static constexpr ENetSimCueReplicationTarget ReplicationTarget { InReplicationTarget };
@@ -168,19 +200,22 @@ struct TNetSimCueTraits
 {
 	using Traits = typename TSelectNetSimCueTraits<TCue>::Traits;
 
-	static constexpr uint8 InvokeMask { Traits::InvokeMask };
+	static constexpr ENetSimCueInvoker InvokeMask { Traits::InvokeMask };
+	static constexpr bool Resimulate { Traits::Resimulate };
 	static constexpr ENetSimCueReplicationTarget ReplicationTarget { Traits::ReplicationTarget };
+	static constexpr ESimulationTickContext SimTickMask { GetSimTickMask( (ENetSimCueInvoker)InvokeMask, Resimulate) };
 };
 
 // Type requirements: helper to determine if NetSerialize/NetIdentical functions need to be defined for user types based on the above traits
 template <typename TCue>
 struct TNetSimCueTypeRequirements
 {
-	enum {
+	enum
+	{
 		// NetSerialize is required if we ever need to replicate
 		RequiresNetSerialize = (TNetSimCueTraits<TCue>::ReplicationTarget != ENetSimCueReplicationTarget::None),
 		// Likewise for NetIdentical, but also if we plan to invoke during Resimulate too (even if non replicated, we use NetIdentical for comparisons. though this is probably a non practical use case).
-		RequiresNetIdentical = (TNetSimCueTraits<TCue>::ReplicationTarget != ENetSimCueReplicationTarget::None) || (TNetSimCueTraits<TCue>::InvokeMask & (uint8)ESimulationTickContext::Resimulate)
+		RequiresNetIdentical = (TNetSimCueTraits<TCue>::ReplicationTarget != ENetSimCueReplicationTarget::None) || (TNetSimCueTraits<TCue>::Resimulate)
 	};
 };
 
diff --git a/Engine/Plugins/Runtime/NetworkPrediction/Source/NetworkPrediction/Public/NetworkedSimulationModelCues.h b/Engine/Plugins/Runtime/NetworkPrediction/Source/NetworkPrediction/Public/NetworkedSimulationModelCues.h
--- a/Engine/Plugins/Runtime/NetworkPrediction/Source/NetworkPrediction/Public/NetworkedSimulationModelCues.h
+++ b/Engine/Plugins/Runtime/NetworkPrediction/Source/NetworkPrediction/Public/NetworkedSimulationModelCues.h
@@ -362,9 +362,9 @@ struct FNetSimCueDispatcher
 	{
 		if (EnsureValidContext())
 		{
-			if ((TNetSimCueTraits<T>::InvokeMask & (uint8)Context.TickContext) > 0)
+			if (EnumHasAnyFlags(TNetSimCueTraits<T>::SimTickMask, Context.TickContext))
 			{
-				const bool bSupportsResimulate = (TNetSimCueTraits<T>::InvokeMask & (uint8)ESimulationTickContext::Resimulate) > 0;
+				constexpr bool bSupportsResimulate = TNetSimCueTraits<T>::Resimulate;
 
 				bool bAllowRollback = false;	// Whether this cue should be dispatched with rollback callbacks.
 				bool bTransient = false;		// Whether we go in the transient list. Transient cues are dumped after dispatching (not saved over multiple frames for uniqueness comparisons during Invoke or NetSerialize)
@@ -422,7 +422,7 @@ struct FNetSimCueDispatcher
 			}
 			else
 			{
-				UE_LOG(LogNetSimCues, Log, TEXT("%s .Suppressing Cue Invocation %s. Mask: %d. TickContext: %d"), *GetDebugName(), *FGlobalCueTypeTable::Get().GetTypeName(T::ID), TNetSimCueTraits<T>::InvokeMask, (int32)Context.TickContext);
+				UE_LOG(LogNetSimCues, Log, TEXT("%s .Suppressing Cue Invocation %s. Mask: %d. TickContext: %d"), *GetDebugName(), *FGlobalCueTypeTable::Get().GetTypeName(T::ID), TNetSimCueTraits<T>::SimTickMask, (int32)Context.TickContext);
 			}
 		}
 	}
diff --git a/Engine/Source/Runtime/Core/Public/Misc/EnumClassFlags.h b/Engine/Source/Runtime/Core/Public/Misc/EnumClassFlags.h
--- a/Engine/Source/Runtime/Core/Public/Misc/EnumClassFlags.h
+++ b/Engine/Source/Runtime/Core/Public/Misc/EnumClassFlags.h
@@ -14,13 +14,13 @@
 	inline CONSTEXPR Enum  operator~ (Enum  E)             { return (Enum)~(__underlying_type(Enum))E; }
 
 template<typename Enum>
-inline bool EnumHasAllFlags(Enum Flags, Enum Contains)
+CONSTEXPR bool EnumHasAllFlags(Enum Flags, Enum Contains)
 {
 	return ( ( ( __underlying_type(Enum) )Flags ) & ( __underlying_type(Enum) )Contains ) == ( ( __underlying_type(Enum) )Contains );
 }
 
 template<typename Enum>
-inline bool EnumHasAnyFlags(Enum Flags, Enum Contains)
+CONSTEXPR bool EnumHasAnyFlags(Enum Flags, Enum Contains)
 {
 	return ( ( ( __underlying_type(Enum) )Flags ) & ( __underlying_type(Enum) )Contains ) != 0;
 }
