NetworkPrediction Plugin
-Cue replication target now a mask with proper distinction between auto/sim proxy. This makes it so cues not intended for a client won't actually be sent. (However simulated extrap vs interpolation is still done client side since it is a client side setting. The win here is that we won't replicate these cues to the auto proxy in those cases)

Core: replaced CONSTEXPR with constexpr based on feedback

#rnx

diff --git a/Engine/Plugins/Runtime/NetworkPrediction/Source/NetworkPrediction/Private/NetworkedSimulationModelCues.cpp b/Engine/Plugins/Runtime/NetworkPrediction/Source/NetworkPrediction/Private/NetworkedSimulationModelCues.cpp
--- a/Engine/Plugins/Runtime/NetworkPrediction/Source/NetworkPrediction/Private/NetworkedSimulationModelCues.cpp
+++ b/Engine/Plugins/Runtime/NetworkPrediction/Source/NetworkPrediction/Private/NetworkedSimulationModelCues.cpp
@@ -201,7 +201,7 @@ void TestCues()
 	// -------------------------------------
 
 	FNetBitWriter TempWriter(1024 << 3);
-	ServerDispatcher.NetSerializeSavedCues(TempWriter, false);
+	ServerDispatcher.NetSerializeSavedCues(TempWriter, ENetSimCueReplicationTarget::All);
 
 	// -------------------------------------
 	// Receive
@@ -209,7 +209,7 @@ void TestCues()
 
 	FNetBitReader TempReader(nullptr, TempWriter.GetData(), TempWriter.GetNumBits());
 	TNetSimCueDispatcher<void> ClientDispatcher;
-	ClientDispatcher.NetSerializeSavedCues(TempReader, false);
+	ClientDispatcher.NetSerializeSavedCues(TempReader, ENetSimCueReplicationTarget::All);
 
 	TestHandlerChild MyObject;
 
diff --git a/Engine/Plugins/Runtime/NetworkPrediction/Source/NetworkPrediction/Public/NetworkedSimulationModelCueTraits.h b/Engine/Plugins/Runtime/NetworkPrediction/Source/NetworkPrediction/Public/NetworkedSimulationModelCueTraits.h
--- a/Engine/Plugins/Runtime/NetworkPrediction/Source/NetworkPrediction/Public/NetworkedSimulationModelCueTraits.h
+++ b/Engine/Plugins/Runtime/NetworkPrediction/Source/NetworkPrediction/Public/NetworkedSimulationModelCueTraits.h
@@ -53,10 +53,14 @@ static constexpr ESimulationTickContext GetSimTickMask(const ENetSimCueInvoker I
 // Who a NetSimCue should replicate/be accepted by
 enum class ENetSimCueReplicationTarget : uint8
 {
-	None,			// Do not replicate cue to anyone
-	Interpolators,	// Only replicate to / accept on clients that are interpolating (not running the simulation themselves)
-	All,			// Replicate to everyone
+	None			= 0,		// Do not replicate cue to anyone
+	AutoProxy		= 1 << 0,	// Replicate to autonomous proxies (controlling clients)
+	SimulatedProxy	= 1 << 1,	// Replicate to simulated proxies that are running the simulation
+	Interpolators	= 1 << 2,	// Replicate to simulated proxies that are *not* running the simulation themselves (E.,g interpolating)
+	
+	All = AutoProxy | SimulatedProxy | Interpolators,
 };
+ENUM_CLASS_FLAGS(ENetSimCueReplicationTarget);
 
 // ----------------------------------------------------------------------------------------------------------------------------------------------
 //	NetSimCue Traits Presets. These are reasonable settings that may be appropriate in common cases. Presets are provided so that individual settings
@@ -67,6 +71,7 @@ enum class ENetSimCueReplicationTarget : uint8
 //		NetSimCueTraits::Strong - The most robust trait type. Will replicate to everyone and will rollback/resimulate. Must implement NetSerialize/NetIdentical. Most expensive.
 //
 //	The other presets fall somewhere in the middle and require thought/nuance to decide if they are right for your case.
+//	Other configs are possible but its not clear their usefullness. For example, a cue that only plays on simulated clients.
 //
 // ----------------------------------------------------------------------------------------------------------------------------------------------
 
@@ -84,14 +89,24 @@ namespace NetSimCueTraits
 		static constexpr ENetSimCueReplicationTarget ReplicationTarget { ENetSimCueReplicationTarget::None };
 	};
 
-	// Same as above but will only play on predicting client, not authority
-	struct WeakClientOnly
+	// Same as above but will only play on the owning, autonomous proxy client (not on authority, not on simulated clients)
+	// Useful if the Cue is needed for only the controlling player, maybe like a HUD/UI notification pop
+	struct WeakOwningClientOnly
 	{
 		static constexpr ENetSimCueInvoker InvokeMask { ENetSimCueInvoker::Predict };
 		static constexpr bool Resimulate { false };
 		static constexpr ENetSimCueReplicationTarget ReplicationTarget { ENetSimCueReplicationTarget::None };
 	};
 
+	// Same as above but will play on all clients, just not the authority.
+	// Useful for something that is purely cosmetic. E.g, does not run on the server but all clients should see it if they are running the sim
+	struct WeakClientsOnly
+	{
+		static constexpr ENetSimCueInvoker InvokeMask { ENetSimCueInvoker::Predict | ENetSimCueInvoker::SimExtrapolate };
+		static constexpr bool Resimulate { false };
+		static constexpr ENetSimCueReplicationTarget ReplicationTarget { ENetSimCueReplicationTarget::None };
+	};
+
 	// Will only play on the authority path and not replicate to anyone else
 	struct AuthorityOnly
 	{
diff --git a/Engine/Plugins/Runtime/NetworkPrediction/Source/NetworkPrediction/Public/NetworkedSimulationModelCues.h b/Engine/Plugins/Runtime/NetworkPrediction/Source/NetworkPrediction/Public/NetworkedSimulationModelCues.h
--- a/Engine/Plugins/Runtime/NetworkPrediction/Source/NetworkPrediction/Public/NetworkedSimulationModelCues.h
+++ b/Engine/Plugins/Runtime/NetworkPrediction/Source/NetworkPrediction/Public/NetworkedSimulationModelCues.h
@@ -208,7 +208,7 @@ struct FSavedCue
 	FSavedCue& operator=(FSavedCue&&) = default;
 	
 	FSavedCue(const FNetSimCueTypeId& InId, const FNetworkSimTime& InTime, const bool& bInAllowRollback, const bool& bInNetConfirmed, const bool& bInResimulates, FNetSimCueWrapperBase* Cue)
-		: ID(InId), Time(InTime), CueInstance(Cue), bAllowRollback(bInAllowRollback), bNetConfirmed(bInNetConfirmed), bResimulates(bInResimulates)
+		: ID(InId), Time(InTime), CueInstance(Cue), bAllowRollback(bInAllowRollback), bNetConfirmed(bInNetConfirmed), bResimulates(bInResimulates), ReplicationTarget(Cue->GetReplicationTarget())
 	{
 		
 	}
@@ -230,6 +230,7 @@ struct FSavedCue
 			CueInstance->NetSerialize(Ar);
 
 			Time.NetSerialize(Ar);
+			ReplicationTarget = CueInstance->GetReplicationTarget();
 		}
 	}
 
@@ -263,6 +264,8 @@ struct FSavedCue
 	
 	bool bResimulates = false;				 // Whether this cue supports invocation during resimulation. Needed to set bResimulatePendingRollback
 	bool bPendingResimulateRollback = false; // Rollback is pending due to resimulation (unless CUe is matched with an Invocation during the resimulate)
+
+	ENetSimCueReplicationTarget ReplicationTarget;
 };
 
 
@@ -456,20 +459,36 @@ template<typename Model=void>
 struct TNetSimCueDispatcher : public FNetSimCueDispatcher
 {
 	// Serializes all saved cues
-	void NetSerializeSavedCues(FArchive& Ar, bool bIsInterpolatingSim)
+	void NetSerializeSavedCues(FArchive& Ar, ENetSimCueReplicationTarget ReplicationMask)
 	{
-		FNetSimCueTypeId NumCues = SavedCues.Num();
-		Ar << NumCues;
-		
 		if (Ar.IsSaving())
 		{
+			// FIXME: requires two passes to count how many elements are valid for this replication mask.
+			// We could count this as saved cues are added or possibly modify the bitstream after writing the elements (tricky and would require casting to FNetBitWriter which feels real bad)
+			FNetSimCueTypeId NumCues = 0;
 			for (FSavedCue& SavedCue : SavedCues)
 			{
-				SavedCue.NetSerialize(Ar);
+				if (EnumHasAnyFlags(SavedCue.ReplicationTarget, ReplicationMask))
+				{
+					NumCues++;
+				}
+			}
+
+			Ar << NumCues;
+
+			for (FSavedCue& SavedCue : SavedCues)
+			{
+				if (EnumHasAnyFlags(SavedCue.ReplicationTarget, ReplicationMask))
+				{
+					SavedCue.NetSerialize(Ar);
+				}
 			}
 		}
 		else
 		{
+			FNetSimCueTypeId NumCues;
+			Ar << NumCues;
+
 			// This is quite inefficient right now. 
 			//	-We are replicating cues in the last X seconds (ReplicationWindow) redundantly
 			//	-Client has to deserialize them (+ heap allocation) and check for uniqueness (have they already processed)
@@ -484,11 +503,10 @@ struct TNetSimCueDispatcher : public FNetSimCueDispatcher
 				SerializedCue.NetSerialize(Ar);
 
 				// Decide if we should accept the cue:
-
 				// ReplicationTarget: Cues can be set to only replicate to interpolators
-				if (SerializedCue.CueInstance->GetReplicationTarget() == ENetSimCueReplicationTarget::Interpolators && !bIsInterpolatingSim)
+				if ( EnumHasAnyFlags(SerializedCue.ReplicationTarget, ReplicationMask) == false )
 				{
-					UE_LOG(LogNetSimCues, Log, TEXT("Discarding replicated NetSimCue %s intended for interpolators."), *SerializedCue.GetDebugName());
+					UE_LOG(LogNetSimCues, Log, TEXT("Discarding replicated NetSimCue %s that is not intended for us. CueMask: %d. Our Mask: %d"), *SerializedCue.GetDebugName(), SerializedCue.ReplicationTarget, ReplicationMask);
 					continue;
 				}
 				
diff --git a/Engine/Plugins/Runtime/NetworkPrediction/Source/NetworkPrediction/Public/NetworkedSimulationModelRepControllers.h b/Engine/Plugins/Runtime/NetworkPrediction/Source/NetworkPrediction/Public/NetworkedSimulationModelRepControllers.h
--- a/Engine/Plugins/Runtime/NetworkPrediction/Source/NetworkPrediction/Public/NetworkedSimulationModelRepControllers.h
+++ b/Engine/Plugins/Runtime/NetworkPrediction/Source/NetworkPrediction/Public/NetworkedSimulationModelRepControllers.h
@@ -429,10 +429,15 @@ struct TRepController_Simulated : public TBase
 		SyncState->NetSerialize(Ar);
 		AuxState->NetSerialize(Ar);
 
-		Buffers.CueDispatcher.NetSerializeSavedCues(Ar, GetSimulatedUpdateMode() == ESimulatedUpdateMode::Interpolate);
-
-		if (Ar.IsLoading())
+		if (Ar.IsSaving())
 		{
+			// Server: send SimProxy and Interpolators. Whether this is in interpolation mode is really a client side thing (we want client to make this decision and transition between the two when necessary)
+			Buffers.CueDispatcher.NetSerializeSavedCues(Ar, ENetSimCueReplicationTarget::SimulatedProxy | ENetSimCueReplicationTarget::Interpolators);
+		}
+		else
+		{
+			Buffers.CueDispatcher.NetSerializeSavedCues(Ar, GetSimulatedUpdateMode() == ESimulatedUpdateMode::Interpolate ? ENetSimCueReplicationTarget::Interpolators : ENetSimCueReplicationTarget::SimulatedProxy);
+
 			LastSerializedInputCmd = *InputCmd;
 			LastSerializedSyncState = *SyncState;
 			LastSerializedAuxState = *AuxState;
@@ -653,7 +658,7 @@ struct TRepController_Autonomous: public TBase
 			SerializedAuxState.NetSerialize(Ar);
 		}
 
-		Buffers.CueDispatcher.NetSerializeSavedCues(Ar, false);
+		Buffers.CueDispatcher.NetSerializeSavedCues(Ar, ENetSimCueReplicationTarget::AutoProxy);
 
 		if (Ar.IsLoading())
 		{
diff --git a/Engine/Source/Runtime/Core/Public/Misc/EnumClassFlags.h b/Engine/Source/Runtime/Core/Public/Misc/EnumClassFlags.h
--- a/Engine/Source/Runtime/Core/Public/Misc/EnumClassFlags.h
+++ b/Engine/Source/Runtime/Core/Public/Misc/EnumClassFlags.h
@@ -7,20 +7,20 @@
 	inline           Enum& operator|=(Enum& Lhs, Enum Rhs) { return Lhs = (Enum)((__underlying_type(Enum))Lhs | (__underlying_type(Enum))Rhs); } \
 	inline           Enum& operator&=(Enum& Lhs, Enum Rhs) { return Lhs = (Enum)((__underlying_type(Enum))Lhs & (__underlying_type(Enum))Rhs); } \
 	inline           Enum& operator^=(Enum& Lhs, Enum Rhs) { return Lhs = (Enum)((__underlying_type(Enum))Lhs ^ (__underlying_type(Enum))Rhs); } \
-	inline CONSTEXPR Enum  operator| (Enum  Lhs, Enum Rhs) { return (Enum)((__underlying_type(Enum))Lhs | (__underlying_type(Enum))Rhs); } \
-	inline CONSTEXPR Enum  operator& (Enum  Lhs, Enum Rhs) { return (Enum)((__underlying_type(Enum))Lhs & (__underlying_type(Enum))Rhs); } \
-	inline CONSTEXPR Enum  operator^ (Enum  Lhs, Enum Rhs) { return (Enum)((__underlying_type(Enum))Lhs ^ (__underlying_type(Enum))Rhs); } \
-	inline CONSTEXPR bool  operator! (Enum  E)             { return !(__underlying_type(Enum))E; } \
-	inline CONSTEXPR Enum  operator~ (Enum  E)             { return (Enum)~(__underlying_type(Enum))E; }
+	inline constexpr Enum  operator| (Enum  Lhs, Enum Rhs) { return (Enum)((__underlying_type(Enum))Lhs | (__underlying_type(Enum))Rhs); } \
+	inline constexpr Enum  operator& (Enum  Lhs, Enum Rhs) { return (Enum)((__underlying_type(Enum))Lhs & (__underlying_type(Enum))Rhs); } \
+	inline constexpr Enum  operator^ (Enum  Lhs, Enum Rhs) { return (Enum)((__underlying_type(Enum))Lhs ^ (__underlying_type(Enum))Rhs); } \
+	inline constexpr bool  operator! (Enum  E)             { return !(__underlying_type(Enum))E; } \
+	inline constexpr Enum  operator~ (Enum  E)             { return (Enum)~(__underlying_type(Enum))E; }
 
 template<typename Enum>
-CONSTEXPR bool EnumHasAllFlags(Enum Flags, Enum Contains)
+constexpr bool EnumHasAllFlags(Enum Flags, Enum Contains)
 {
 	return ( ( ( __underlying_type(Enum) )Flags ) & ( __underlying_type(Enum) )Contains ) == ( ( __underlying_type(Enum) )Contains );
 }
 
 template<typename Enum>
-CONSTEXPR bool EnumHasAnyFlags(Enum Flags, Enum Contains)
+constexpr bool EnumHasAnyFlags(Enum Flags, Enum Contains)
 {
 	return ( ( ( __underlying_type(Enum) )Flags ) & ( __underlying_type(Enum) )Contains ) != 0;
 }
