NetworkPrediction Plugin
-fix uninit variable
#rnx

diff --git a/Engine/Plugins/Runtime/NetworkPrediction/Source/NetworkPrediction/Public/Movement/ParametricMovement.h b/Engine/Plugins/Runtime/NetworkPrediction/Source/NetworkPrediction/Public/Movement/ParametricMovement.h
--- a/Engine/Plugins/Runtime/NetworkPrediction/Source/NetworkPrediction/Public/Movement/ParametricMovement.h
+++ b/Engine/Plugins/Runtime/NetworkPrediction/Source/NetworkPrediction/Public/Movement/ParametricMovement.h
@@ -65,8 +65,8 @@ struct FParametricInputCmd
 // State we are evolving frame to frame and keeping in sync
 struct FParametricSyncState
 {
-	float Position;
-	float PlayRate;
+	float Position=0.f;
+	float PlayRate=1.f;
 
 	void NetSerialize(const FNetSerializeParams& P)
 	{
