NetworkPrediction Plugin
-Fix force reconcile not working
-for crumb visual logging, don't output full debug string as it is too spammy

#rnx

diff --git a/Engine/Plugins/Runtime/NetworkPrediction/Source/NetworkPrediction/Private/NetworkPredictionTypes.cpp b/Engine/Plugins/Runtime/NetworkPrediction/Source/NetworkPrediction/Private/NetworkPredictionTypes.cpp
--- a/Engine/Plugins/Runtime/NetworkPrediction/Source/NetworkPrediction/Private/NetworkPredictionTypes.cpp
+++ b/Engine/Plugins/Runtime/NetworkPrediction/Source/NetworkPrediction/Private/NetworkPredictionTypes.cpp
@@ -67,7 +67,7 @@ void FVisualLoggingHelpers::VisualLogActor(AActor* Owner, FTransform& Transform,
 		if (NetSimVLogCVars::UseVLogger)
 		{
 			static float CrumbRadius = 3.f;
-			UE_VLOG_LOCATION(Owner, LogNetworkSim, Log, Transform.GetLocation(), CrumbRadius, DrawColor, TEXT("%s"), *Params.DebugString);
+			UE_VLOG_LOCATION(Owner, LogNetworkSim, Log, Transform.GetLocation(), CrumbRadius, DrawColor, TEXT("%d"), Params.Frame);
 		}
 	}
 	else
diff --git a/Engine/Plugins/Runtime/NetworkPrediction/Source/NetworkPrediction/Public/NetworkedSimulationModelRepControllers.h b/Engine/Plugins/Runtime/NetworkPrediction/Source/NetworkPrediction/Public/NetworkedSimulationModelRepControllers.h
--- a/Engine/Plugins/Runtime/NetworkPrediction/Source/NetworkPrediction/Public/NetworkedSimulationModelRepControllers.h
+++ b/Engine/Plugins/Runtime/NetworkPrediction/Source/NetworkPrediction/Public/NetworkedSimulationModelRepControllers.h
@@ -673,7 +673,7 @@ struct TRepController_Autonomous: public TBase
 			{
 				const bool bForceReconcile = (NetworkSimulationModelCVars::ForceReconcile() > 0) || (NetworkSimulationModelCVars::ForceReconcileSingle() > 0);
 				
-				if (Model::ShouldReconcile(SerializedSyncState, SerializedAuxState, *ClientSync, *ClientAux))
+				if (bForceReconcile || Model::ShouldReconcile(SerializedSyncState, SerializedAuxState, *ClientSync, *ClientAux))
 				{
 					NetworkSimulationModelCVars::SetForceReconcileSingle(0);
 					UE_CLOG(!Buffers.Input.IsValidFrame(SerializedFrame-1), LogNetworkSim, Error, TEXT("::NetSerialize: Client InputBuffer does not contain data for frame %d. {%s} {%s}"), SerializedFrame, *Buffers.Input.GetBasicDebugStr(), *Buffers.Sync.GetBasicDebugStr());
