NetworkPrediction Plugin
Added NETWORKPREDICTION_API to the FNetworkSimTime so game code can use it when making movement components

#rnx

diff --git a/Engine/Plugins/Runtime/NetworkPrediction/Source/NetworkPrediction/Public/NetworkedSimulationModelTime.h b/Engine/Plugins/Runtime/NetworkPrediction/Source/NetworkPrediction/Public/NetworkedSimulationModelTime.h
--- a/Engine/Plugins/Runtime/NetworkPrediction/Source/NetworkPrediction/Public/NetworkedSimulationModelTime.h
+++ b/Engine/Plugins/Runtime/NetworkPrediction/Source/NetworkPrediction/Public/NetworkedSimulationModelTime.h
@@ -8,7 +8,7 @@
 // ----------------------------------------------------------------------------------------------------------------------------------------------
 
 // The main Simulation time type. All sims use this to talk about time.
-struct FNetworkSimTime
+struct NETWORKPREDICTION_API FNetworkSimTime
 {
 	using FSimTime = int32; // Underlying type used to store simulation time
 	using FRealTime = float; // Underlying type used when dealing with realtime (usually coming in from the engine tick).
