NetworkPrediction Plugin
-cvar to supress canvas drawing in debugger (so you can get clearer view of in world debugging)

#rnx

diff --git a/Engine/Plugins/Runtime/NetworkPrediction/Source/NetworkPrediction/Public/NetworkedSimulationModelDebugger.h b/Engine/Plugins/Runtime/NetworkPrediction/Source/NetworkPrediction/Public/NetworkedSimulationModelDebugger.h
--- a/Engine/Plugins/Runtime/NetworkPrediction/Source/NetworkPrediction/Public/NetworkedSimulationModelDebugger.h
+++ b/Engine/Plugins/Runtime/NetworkPrediction/Source/NetworkPrediction/Public/NetworkedSimulationModelDebugger.h
@@ -21,6 +21,7 @@ DEFINE_LOG_CATEGORY_STATIC(LogNetworkSimDebug, Log, All);
 
 namespace NetworkSimulationModelDebugCVars
 {
+	NETSIM_DEVCVAR_SHIPCONST_INT(DrawCanvas, 1, "nsm.debug.DrawCanvas", "Draws Canvas Debug HUD");
 	NETSIM_DEVCVAR_SHIPCONST_INT(DrawFrames, 1, "nsm.debug.DrawFrames", "Draws frame data (text) in debug graphs");
 	NETSIM_DEVCVAR_SHIPCONST_INT(DrawNetworkSendLines, 1, "nsm.debug.DrawNetworkSendLines", "Draws lines representing network traffic in debugger");
 	NETSIM_DEVCVAR_SHIPCONST_INT(GatherServerSidePIE, 1, "nsm.debug.GatherServerSide", "Whenever we gather debug info from a client side actor, also gather server side equivelent. Only works in PIE.");
@@ -118,6 +119,11 @@ struct NETWORKPREDICTION_API FNetworkSimulationModelDebuggerManager: public FTic
 		FDisplayDebugManager& DisplayDebugManager = C->DisplayDebugManager;
 		DisplayDebugManager.Initialize(C, GEngine->GetSmallFont(), FVector2D(4.0f, 150.0f));
 
+		if (!NetworkSimulationModelDebugCVars::DrawCanvas())
+		{
+			return;
+		}
+
 		if (Lines.Num() > 0)
 		{
 			const float TextScale = FMath::Max(C->SizeX / 1920.0f, 1.0f);
diff --git a/Engine/Plugins/Runtime/NetworkPrediction/readme.txt b/Engine/Plugins/Runtime/NetworkPrediction/readme.txt
--- a/Engine/Plugins/Runtime/NetworkPrediction/readme.txt
+++ b/Engine/Plugins/Runtime/NetworkPrediction/readme.txt
@@ -106,6 +106,10 @@ Smoothing: Taking the output of the simulation and applying an additional layer
 // Release notes
 // ----------------------------------------------------------------------------------------------------------
 
+Update (12-13-19)
+-Only small updates (if any) for the rest of the year. System is going through code reviews as we plan for writing a new general movement system on top of Network Prediction.
+-Things are still liable to change but for the most part every major feature is in.
+
 Update (11-21-19)
 -Rewindable NetSimCue support added. This just about wraps up the initial implemention of NetSimCues and the "Game API / Event System".
 -Still some cleanup and debugging/logging stuff to do but will be transitioning into porting the legacy CMC movement system over soon.
