Move non-templated implementations to cpp, this avoids Mac linker issues where users must depend on RenderCore.
Small adjustment to Emit and EmitText to allow for Move semenatics
#rnx

diff --git a/Engine/Plugins/Runtime/NetworkPrediction/Source/NetworkPrediction/Private/NetworkedSimulationModelDebugger.cpp b/Engine/Plugins/Runtime/NetworkPrediction/Source/NetworkPrediction/Private/NetworkedSimulationModelDebugger.cpp
--- a/Engine/Plugins/Runtime/NetworkPrediction/Source/NetworkPrediction/Private/NetworkedSimulationModelDebugger.cpp
+++ b/Engine/Plugins/Runtime/NetworkPrediction/Source/NetworkPrediction/Private/NetworkedSimulationModelDebugger.cpp
@@ -18,6 +18,224 @@ FNetworkSimulationModelDebuggerManager& FNetworkSimulationModelDebuggerManager::
 	return Manager;
 }
 
+FNetworkSimulationModelDebuggerManager::FNetworkSimulationModelDebuggerManager()
+{
+	DrawDebugServicesHandle = UDebugDrawService::Register(TEXT("Game"), FDebugDrawDelegate::CreateRaw(this, &FNetworkSimulationModelDebuggerManager::DrawDebugService));
+	check(DrawDebugServicesHandle.IsValid());
+}
+
+FNetworkSimulationModelDebuggerManager::~FNetworkSimulationModelDebuggerManager()
+{
+	if (Graph.IsValid())
+	{
+		Graph->RemoveFromRoot();
+	}
+}
+
+void FNetworkSimulationModelDebuggerManager::SetDebuggerActive(AActor* OwningActor, bool InActive)
+{
+	if (INetworkSimulationModelDebugger* Debugger = Find(OwningActor))
+	{
+		Debugger->SetActive(InActive);
+	}
+	ResetCache();
+	Gather(LastCanvas.Get());
+}
+
+void FNetworkSimulationModelDebuggerManager::ToggleDebuggerActive(AActor* OwningActor)
+{
+	if (INetworkSimulationModelDebugger* Debugger = Find(OwningActor))
+	{
+		Debugger->SetActive(!Debugger->IsActive());
+	}
+	ResetCache();
+	Gather(LastCanvas.Get());
+}
+
+void FNetworkSimulationModelDebuggerManager::SetContinousGather(bool InGather)
+{
+	bContinousGather = InGather;
+	if (!bContinousGather)
+	{
+		Gather(LastCanvas.Get());
+	}
+}
+
+void FNetworkSimulationModelDebuggerManager::DrawDebugService(UCanvas* C, APlayerController* PC)
+{
+	LastCanvas = C;
+	if (bContinousGather)
+	{
+		Gather(C);
+	}
+
+	FDisplayDebugManager& DisplayDebugManager = C->DisplayDebugManager;
+	DisplayDebugManager.Initialize(C, GEngine->GetSmallFont(), FVector2D(4.0f, 150.0f));
+
+	if (!NetworkSimulationModelDebugCVars::DrawCanvas())
+	{
+		return;
+	}
+
+	if (Lines.Num() > 0)
+	{
+		const float TextScale = FMath::Max(C->SizeX / 1920.0f, 1.0f);
+		FCanvasTileItem TextBackgroundTile(FVector2D(0.0f, 120.0f), FVector2D(400.0f, 1800.0f) * TextScale, FColor(0, 0, 0, 100));
+		TextBackgroundTile.BlendMode = SE_BLEND_Translucent;
+		C->DrawItem(TextBackgroundTile);
+	}
+
+	// --------------------------------------------------------
+	//	Lines
+	// --------------------------------------------------------
+
+	for (FDebugLine& Line : Lines)
+	{
+		DisplayDebugManager.SetDrawColor(Line.Color);
+		DisplayDebugManager.DrawString(Line.Str);
+	}
+
+	// --------------------------------------------------------
+	//	Canvas Items (graphs+text)
+	// --------------------------------------------------------
+
+	for (TUniquePtr<FCanvasItem>& Item : CanvasItems[0])
+	{
+		C->DrawItem(*Item.Get());
+	}
+
+	if (NetworkSimulationModelDebugCVars::DrawFrames() > 0)
+	{
+		for (TUniquePtr<FCanvasItem>& Item : CanvasItems[1])
+		{
+			C->DrawItem(*Item.Get());
+		}
+	}
+}
+
+void FNetworkSimulationModelDebuggerManager::Tick(float DeltaTime)
+{
+	for (auto It = DebuggerMap.CreateIterator(); It; ++It)
+	{
+		const AActor* Owner = It.Key().Get();
+		if (!Owner)
+		{
+			It.RemoveCurrent();
+			continue;;
+		}
+
+		if (It.Value()->IsActive())
+		{
+			It.Value()->Tick(DeltaTime);
+		}
+	}
+}
+
+TStatId FNetworkSimulationModelDebuggerManager::GetStatId() const
+{
+	RETURN_QUICK_DECLARE_CYCLE_STAT(FNetworkSimulationModelDebuggerManager, STATGROUP_TaskGraphTasks);
+}
+
+void FNetworkSimulationModelDebuggerManager::LogSingleFrame(FOutputDevice& Ar)
+{
+	Gather(LastCanvas.Get());
+
+	for (const FDebugLine& Line : Lines)
+	{
+		Ar.Logf(TEXT("%s"), *Line.Str);
+	}
+}
+
+void FNetworkSimulationModelDebuggerManager::Emit(FString Str, FColor Color, float XOffset, float YOffset)
+{
+	Lines.Emplace(FDebugLine{ MoveTemp(Str), Color, XOffset, YOffset });
+}
+
+void FNetworkSimulationModelDebuggerManager::EmitQuad(FVector2D ScreenPosition, FVector2D ScreenSize, FColor Color)
+{
+	FVector2D Quad[4];
+
+	Quad[0].X = ScreenPosition.X;
+	Quad[0].Y = ScreenPosition.Y;
+
+	Quad[1].X = ScreenPosition.X;
+	Quad[1].Y = ScreenPosition.Y + ScreenSize.Y;
+
+	Quad[2].X = ScreenPosition.X + ScreenSize.X;
+	Quad[2].Y = ScreenPosition.Y + ScreenSize.Y;
+
+	Quad[3].X = ScreenPosition.X + ScreenSize.X;
+	Quad[3].Y = ScreenPosition.Y;
+
+	CanvasItems[0].Emplace(MakeUnique<FCanvasTriangleItem>(Quad[0], Quad[1], Quad[2], GWhiteTexture));
+	CanvasItems[0].Last()->SetColor(Color);
+
+	CanvasItems[0].Emplace(MakeUnique<FCanvasTriangleItem>(Quad[2], Quad[3], Quad[0], GWhiteTexture));
+	CanvasItems[0].Last()->SetColor(Color);
+}
+
+void FNetworkSimulationModelDebuggerManager::EmitText(FVector2D ScreenPosition, FColor Color, FString Str)
+{
+	CanvasItems[1].Emplace(MakeUnique<FCanvasTextItem>(ScreenPosition, FText::FromString(MoveTemp(Str)), GEngine->GetTinyFont(), Color));
+}
+
+void FNetworkSimulationModelDebuggerManager::EmitLine(FVector2D StartPosition, FVector2D EndPosition, FColor Color, float Thickness)
+{
+	CanvasItems[0].Emplace(MakeUnique<FCanvasLineItem>(StartPosition, EndPosition));
+	CanvasItems[0].Last()->SetColor(Color);
+	((FCanvasLineItem*)CanvasItems[0].Last().Get())->LineThickness = Thickness;
+}
+
+INetworkSimulationModelDebugger* FNetworkSimulationModelDebuggerManager::Find(const AActor* Actor)
+{
+	if (!Actor)
+	{
+		return nullptr;
+	}
+
+	INetworkSimulationModelDebugger* Debugger = DebuggerMap.FindRef(TWeakObjectPtr<const AActor>(Actor));
+	if (!Debugger)
+	{
+		UE_LOG(LogNetworkSimDebug, Warning, TEXT("Could not find NetworkSimulationModel associated with %s"), *GetPathNameSafe(Actor));
+	}
+	return Debugger;
+}
+
+void FNetworkSimulationModelDebuggerManager::Gather(UCanvas* C)
+{
+	ResetCache();
+
+	for (auto It = DebuggerMap.CreateIterator(); It; ++It)
+	{
+		const AActor* Owner = It.Key().Get();
+		if (!Owner)
+		{
+			It.RemoveCurrent();
+			continue;;
+		}
+
+		if (It.Value()->IsActive())
+		{
+			It.Value()->GatherCurrent(*this, C);
+			if (NetworkSimulationModelDebugCVars::GatherServerSidePIE() > 0)
+			{
+				if (const AActor* ServerSideActor = Cast<AActor>(FindReplicatedObjectOnPIEServer(Owner)))
+				{
+					if (INetworkSimulationModelDebugger* ServerSideSim = Find(ServerSideActor))
+					{
+						Emit();
+						Emit();
+						ServerSideSim->GatherCurrent(*this, nullptr); // Dont do graphs for server side state
+					}
+				}
+			}
+
+			// Only gather first active debugger (it would be great to have more control over this when debugging multiples)
+			break;
+		}
+	}
+}
+
 UObject* FindReplicatedObjectOnPIEServer(const UObject* ClientObject)
 {
 	if (ClientObject == nullptr)
diff --git a/Engine/Plugins/Runtime/NetworkPrediction/Source/NetworkPrediction/Public/NetworkedSimulationModelDebugger.h b/Engine/Plugins/Runtime/NetworkPrediction/Source/NetworkPrediction/Public/NetworkedSimulationModelDebugger.h
--- a/Engine/Plugins/Runtime/NetworkPrediction/Source/NetworkPrediction/Public/NetworkedSimulationModelDebugger.h
+++ b/Engine/Plugins/Runtime/NetworkPrediction/Source/NetworkPrediction/Public/NetworkedSimulationModelDebugger.h
@@ -53,19 +53,8 @@ struct NETWORKPREDICTION_API FNetworkSimulationModelDebuggerManager: public FTic
 {
 	static FNetworkSimulationModelDebuggerManager& Get();
 
-	~FNetworkSimulationModelDebuggerManager()
-	{
-		if (Graph.IsValid())
-		{
-			Graph->RemoveFromRoot();
-		}
-	}
-
-	FNetworkSimulationModelDebuggerManager()
-	{
-		DrawDebugServicesHandle = UDebugDrawService::Register(TEXT("Game"), FDebugDrawDelegate::CreateRaw(this, &FNetworkSimulationModelDebuggerManager::DrawDebugService));
-		check(DrawDebugServicesHandle.IsValid());
-	}
+	FNetworkSimulationModelDebuggerManager();
+	~FNetworkSimulationModelDebuggerManager();
 
 	// ---------------------------------------------------------------------------------------------------------------------------------------
 	//	Outside API (registration, console commands, draw services, etc)
@@ -74,135 +63,32 @@ struct NETWORKPREDICTION_API FNetworkSimulationModelDebuggerManager: public FTic
 	template <typename T>
 	void RegisterNetworkSimulationModel(T* NetworkSim, const AActor* OwningActor);
 
-	void SetDebuggerActive(AActor* OwningActor, bool InActive)
-	{
-		if (INetworkSimulationModelDebugger* Debugger = Find(OwningActor))
-		{
-			Debugger->SetActive(InActive);
-		}
-		ResetCache();
-		Gather(LastCanvas.Get());
-	}
+	void SetDebuggerActive(AActor* OwningActor, bool InActive);
 
-	void ToggleDebuggerActive(AActor* OwningActor)
-	{
-		if (INetworkSimulationModelDebugger* Debugger = Find(OwningActor))
-		{
-			Debugger->SetActive(!Debugger->IsActive());
-		}
-		ResetCache();
-		Gather(LastCanvas.Get());
-	}
+	void ToggleDebuggerActive(AActor* OwningActor);
 
-	void SetContinousGather(bool InGather)
-	{
-		bContinousGather = InGather;
-		if (!bContinousGather)
-		{
-			Gather(LastCanvas.Get());
-		}
-	}
+	void SetContinousGather(bool InGather);
 
 	void ToggleContinousGather()
 	{
 		SetContinousGather(!bContinousGather);
 	}
 
-	void DrawDebugService(UCanvas* C, APlayerController* PC)
-	{
-		LastCanvas = C;
-		if (bContinousGather)
-		{
-			Gather(C);
-		}
-		
-		FDisplayDebugManager& DisplayDebugManager = C->DisplayDebugManager;
-		DisplayDebugManager.Initialize(C, GEngine->GetSmallFont(), FVector2D(4.0f, 150.0f));
-
-		if (!NetworkSimulationModelDebugCVars::DrawCanvas())
-		{
-			return;
-		}
-
-		if (Lines.Num() > 0)
-		{
-			const float TextScale = FMath::Max(C->SizeX / 1920.0f, 1.0f);
-			FCanvasTileItem TextBackgroundTile(FVector2D(0.0f, 120.0f), FVector2D(400.0f, 1800.0f) * TextScale, FColor(0, 0, 0, 100));
-			TextBackgroundTile.BlendMode = SE_BLEND_Translucent;
-			C->DrawItem(TextBackgroundTile);
-		}
-
-		// --------------------------------------------------------
-		//	Lines
-		// --------------------------------------------------------
-
-		for (FDebugLine& Line : Lines)
-		{
-			DisplayDebugManager.SetDrawColor(Line.Color);
-			DisplayDebugManager.DrawString(Line.Str);
-		}
-
-		// --------------------------------------------------------
-		//	Canvas Items (graphs+text)
-		// --------------------------------------------------------
-		
-		for (auto& Item : CanvasItems[0])
-		{
-			C->DrawItem(*Item.Get());
-		}
-
-		if (NetworkSimulationModelDebugCVars::DrawFrames() > 0)
-		{
-			for (auto& Item : CanvasItems[1])
-			{
-				C->DrawItem(*Item.Get());
-			}
-		}
-	}
+	void DrawDebugService(UCanvas* C, APlayerController* PC);
 
-	virtual void Tick( float DeltaTime )
-	{
-		for (auto It = DebuggerMap.CreateIterator(); It; ++It)
-		{
-			const AActor* Owner = It.Key().Get();
-			if (!Owner)
-			{
-				It.RemoveCurrent();
-				continue;;
-			}
-			
-			if (It.Value()->IsActive())
-			{
-				It.Value()->Tick(DeltaTime);
-			}
-		}
-	}	
+	virtual void Tick(float DeltaTime) override;
 
 	/** return the stat id to use for this tickable **/
-	virtual TStatId GetStatId() const override
-	{
-		RETURN_QUICK_DECLARE_CYCLE_STAT(FNetworkSimulationModelDebuggerManager, STATGROUP_TaskGraphTasks);
-	}
+	virtual TStatId GetStatId() const override;
 
 	/** Gathers latest and Logs single frame */
-	void LogSingleFrame(FOutputDevice& Ar)
-	{
-		Gather(LastCanvas.Get());
-		
-		for (FDebugLine& Line : Lines)
-		{
-			Ar.Logf(TEXT("%s"), *Line.Str);
-		}
-	}
+	void LogSingleFrame(FOutputDevice& Ar);
 
 	// ---------------------------------------------------------------------------------------------------------------------------------------
 	//	Debugging API used by TNetworkSimulationModelDebugger
 	// ---------------------------------------------------------------------------------------------------------------------------------------
 
-	void Emit(const FString& Str = FString(), FColor Color = FColor::White, float XOffset=0.f, float YOffset=0.f)
-	{
-		Lines.Emplace(FDebugLine{ Str, Color, XOffset, YOffset });
-	}
+	void Emit(FString Str = FString(), FColor Color = FColor::White, float XOffset = 0.f, float YOffset = 0.f);
 
 	template <typename TBuffer>
 	void EmitElement(TBuffer& Buffer, const FStandardLoggingParameters& Parameters)
@@ -222,97 +108,22 @@ struct NETWORKPREDICTION_API FNetworkSimulationModelDebuggerManager: public FTic
 			StrOut.ParseIntoArrayLines(StrLines, true);
 			for (FString& Str : StrLines)
 			{
-				Emit(Str);
+				Emit(MoveTemp(Str));
 			}
 		}
 	}
 
-	void EmitQuad(FVector2D ScreenPosition, FVector2D ScreenSize, FColor Color)
-	{
-		FVector2D Quad[4];
-		
-		Quad[0].X = ScreenPosition.X;
-		Quad[0].Y = ScreenPosition.Y;
-
-		Quad[1].X = ScreenPosition.X;
-		Quad[1].Y = ScreenPosition.Y + ScreenSize.Y;
-
-		Quad[2].X = ScreenPosition.X + ScreenSize.X;
-		Quad[2].Y = ScreenPosition.Y + ScreenSize.Y;
+	void EmitQuad(FVector2D ScreenPosition, FVector2D ScreenSize, FColor Color);
 
-		Quad[3].X = ScreenPosition.X + ScreenSize.X;
-		Quad[3].Y = ScreenPosition.Y;
-		
-		CanvasItems[0].Emplace( MakeUnique<FCanvasTriangleItem>(Quad[0], Quad[1], Quad[2], GWhiteTexture) );
-		CanvasItems[0].Last()->SetColor(Color);
-
-		CanvasItems[0].Emplace( MakeUnique<FCanvasTriangleItem>(Quad[2], Quad[3], Quad[0], GWhiteTexture) );
-		CanvasItems[0].Last()->SetColor(Color);
-	}
-
-	void EmitText(FVector2D ScreenPosition, FColor Color, const FString& Str)
-	{
-		CanvasItems[1].Emplace( MakeUnique<FCanvasTextItem>(ScreenPosition, FText::FromString(Str), GEngine->GetTinyFont(), Color) );
-	}
+	void EmitText(FVector2D ScreenPosition, FColor Color, FString Str);
 
-	void EmitLine(FVector2D StartPosition, FVector2D EndPosition, FColor Color, float Thickness=1.f)
-	{
-		CanvasItems[0].Emplace( MakeUnique<FCanvasLineItem>(StartPosition, EndPosition) );
-		CanvasItems[0].Last()->SetColor(Color);
-		((FCanvasLineItem*)CanvasItems[0].Last().Get())->LineThickness = Thickness;
-	}
+	void EmitLine(FVector2D StartPosition, FVector2D EndPosition, FColor Color, float Thickness = 1.f);
 
 private:
 
-	INetworkSimulationModelDebugger* Find(const AActor* Actor)
-	{
-		if (!Actor)
-		{
-			return nullptr;
-		}
-
-		INetworkSimulationModelDebugger* Debugger = DebuggerMap.FindRef(TWeakObjectPtr<const AActor>(Actor));
-		if (!Debugger)
-		{
-			UE_LOG(LogNetworkSimDebug, Warning, TEXT("Could not find NetworkSimulationModel associated with %s"), *GetPathNameSafe(Actor));
-		}
-		return Debugger;
-	}
-
-	void Gather(UCanvas* C)
-	{
-		ResetCache();
-
-		for (auto It = DebuggerMap.CreateIterator(); It; ++It)
-		{
-			const AActor* Owner = It.Key().Get();
-			if (!Owner)
-			{
-				It.RemoveCurrent();
-				continue;;
-			}
-			
-			if (It.Value()->IsActive())
-			{
-				It.Value()->GatherCurrent(*this, C);
-				if (NetworkSimulationModelDebugCVars::GatherServerSidePIE() > 0)
-				{
-					if (const AActor* ServerSideActor = Cast<AActor>(FindReplicatedObjectOnPIEServer(Owner)))
-					{
-						if (INetworkSimulationModelDebugger* ServerSideSim = Find(ServerSideActor))
-						{
-							Emit();
-							Emit();
-							ServerSideSim->GatherCurrent(*this, nullptr); // Dont do graphs for server side state
-						}
-					}
-				}
+	INetworkSimulationModelDebugger* Find(const AActor* Actor);
 
-				// Only gather first active debugger (it would be great to have more control over this when debugging multiples)
-				break;
-			}
-		}
-	}
+	void Gather(UCanvas* C);
 
 	void ResetCache()
 	{
@@ -527,7 +338,7 @@ struct TNetworkSimulationModelDebugger : public INetworkSimulationModelDebugger
 		}
 	}
 
-	void GatherCurrent(FNetworkSimulationModelDebuggerManager& Out, UCanvas* Canvas) override
+	virtual void GatherCurrent(FNetworkSimulationModelDebuggerManager& Out, UCanvas* Canvas) override
 	{
 		const AActor* Owner = WeakOwningActor.Get();
 		if (!ensure(Owner))
@@ -567,23 +378,23 @@ struct TNetworkSimulationModelDebugger : public INetworkSimulationModelDebugger
 				Color = FColor::Red;
 			}
 
-			Out.Emit(*ConfirmedFrameStr, Color);
+			Out.Emit(MoveTemp(ConfirmedFrameStr), Color);
 
 			FString SimulationTimeString = FString::Printf(TEXT("Local SimulationTime: %s. SerializedSimulationTime: %s. Difference MS: %s"), *NetworkSim->Ticker.GetTotalProcessedSimulationTime().ToString(),
 				*NetworkSim->RepProxy_Autonomous.GetLastSerializedSimTime().ToString(), *(NetworkSim->Ticker.GetTotalProcessedSimulationTime() - NetworkSim->RepProxy_Autonomous.GetLastSerializedSimTime()).ToString());
-			Out.Emit(*SimulationTimeString, Color);
+			Out.Emit(MoveTemp(SimulationTimeString), Color);
 
 			FString AllowedSimulationTimeString = FString::Printf(TEXT("Allowed Simulation Time: %s. Frame: %d/%d/%d"), *NetworkSim->Ticker.GetRemainingAllowedSimulationTime().ToString(), NetworkSim->Ticker.MaxAllowedFrame, NetworkSim->Ticker.PendingFrame, NetworkSim->Buffers.Input.HeadFrame());
-			Out.Emit(*AllowedSimulationTimeString, Color);
+			Out.Emit(MoveTemp(AllowedSimulationTimeString), Color);
 		}
 		else if (Owner->GetLocalRole() == ROLE_SimulatedProxy)
 		{
 			FColor Color = FColor::White;
 			FString TimeString = FString::Printf(TEXT("Total Processed Simulation Time: %s. Last Serialized Simulation Time: %s. Delta: %s"), *NetworkSim->Ticker.GetTotalProcessedSimulationTime().ToString(), *NetworkSim->RepProxy_Simulated.GetLastSerializedSimulationTime().ToString(), *(NetworkSim->RepProxy_Simulated.GetLastSerializedSimulationTime() - NetworkSim->Ticker.GetTotalProcessedSimulationTime()).ToString());
-			Out.Emit(*TimeString, Color);
+			Out.Emit(MoveTemp(TimeString), Color);
 		}
 
-		auto EmitBuffer = [&Out](FString BufferName, auto& Buffer)
+		auto EmitBuffer = [&Out](const FString& BufferName, auto& Buffer)
 		{
 			Out.Emit();
 			Out.Emit(FString::Printf(TEXT("//////////////// %s ///////////////"), *BufferName), FColor::Yellow);
@@ -657,7 +468,7 @@ struct TNetworkSimulationModelDebugger : public INetworkSimulationModelDebugger
 		}
 	}
 
-	void Tick( float DeltaTime ) override
+	virtual void Tick( float DeltaTime ) override
 	{
 		const AActor* Owner = WeakOwningActor.Get();
 		if (!Owner)
