Network Prediction Plugin
-Added "snapTo" in interpolator. This fixes weirdness around simulations becoming relevant and interpolation taking too long to catch up.
-Added development checksum to network stream to hopefully catch serialization desync
-added cvar to disable debug buffer replication in development builds
-Added more debugging info to interpolation debugging

#rnx

diff --git a/Engine/Plugins/Runtime/NetworkPrediction/Source/NetworkPrediction/Private/NetworkPredictionTypes.cpp b/Engine/Plugins/Runtime/NetworkPrediction/Source/NetworkPrediction/Private/NetworkPredictionTypes.cpp
--- a/Engine/Plugins/Runtime/NetworkPrediction/Source/NetworkPrediction/Private/NetworkPredictionTypes.cpp
+++ b/Engine/Plugins/Runtime/NetworkPrediction/Source/NetworkPrediction/Private/NetworkPredictionTypes.cpp
@@ -107,6 +107,7 @@ void FVisualLoggingHelpers::VisualLogActor(AActor* Owner, FTransform& Transform,
 				LocalSpaceBox.GetCenterAndExtents(ActorOrigin, ActorExtent);
 				ActorExtent *= Transform.GetScale3D();
 				DrawDebugBox(World, Transform.GetLocation(), ActorExtent, Transform.GetRotation(), DrawColor, false, LifetimeSeconds, 0, Thickness);
+				DrawDebugString(World, Transform.GetLocation(), Params.DebugString, nullptr, DrawColor, PersistentLines ? NetSimVLogCVars::DrawDebugDefaultLifeTime : 0.f, false);
 			}
 
 			if (NetSimVLogCVars::UseVLogger)
diff --git a/Engine/Plugins/Runtime/NetworkPrediction/Source/NetworkPrediction/Public/NetworkedSimulationModelInterpolator.h b/Engine/Plugins/Runtime/NetworkPrediction/Source/NetworkPrediction/Public/NetworkedSimulationModelInterpolator.h
--- a/Engine/Plugins/Runtime/NetworkPrediction/Source/NetworkPrediction/Public/NetworkedSimulationModelInterpolator.h
+++ b/Engine/Plugins/Runtime/NetworkPrediction/Source/NetworkPrediction/Public/NetworkedSimulationModelInterpolator.h
@@ -20,9 +20,14 @@ namespace NetworkInterpolationDebugCVars
 
 	NETSIM_DEVCVAR_SHIPCONST_FLOAT(WaitSlack, 0.05, "ni.WaitSlack", "How much slack to wait for when waiting");
 
+	// Catchup means "interpolate faster than realtime"
 	NETSIM_DEVCVAR_SHIPCONST_FLOAT(CatchUpThreshold, 0.300, "ni.CatchUpThreshold", "When we start catching up (seconds from head)");
-	NETSIM_DEVCVAR_SHIPCONST_FLOAT(CatchUpGoal, 0.010, "ni.CatchUpGoal", "When we stop cathcing up (seconds from head)");
+	NETSIM_DEVCVAR_SHIPCONST_FLOAT(CatchUpGoal, 0.010, "ni.CatchUpGoal", "When we stop catching up (seconds from head)");
 	NETSIM_DEVCVAR_SHIPCONST_FLOAT(CatchUpFactor, 1.50, "ni.CatchUpFactor", "Factor we use to catch up");
+
+	// Snap means "we are too far behind, just jump to a specific point"
+	NETSIM_DEVCVAR_SHIPCONST_FLOAT(SnapThreshold, 0.500, "ni.SnapThreshold", "If we are < this, we will snap");
+	NETSIM_DEVCVAR_SHIPCONST_FLOAT(SnapGoal, 0.010, "ni.SnapGoal", "Where we snap to");
 }
 
 template<typename TSync, typename TAux>
@@ -136,23 +141,30 @@ struct TNetSimInterpolator
 		FRealTime NewInterpolationTime = InterpolationTime;
 		{
 			FRealTime Step = DeltaSeconds;
+
+			// Snap if way far behind
+			FRealTime SnapThreshold = SimulationTimeBuffer.HeadElement()->ToRealTimeSeconds() - NetworkInterpolationDebugCVars::SnapThreshold();
+			if (NewInterpolationTime < SnapThreshold)
+			{
+				NewInterpolationTime = SimulationTimeBuffer.HeadElement()->ToRealTimeSeconds() - NetworkInterpolationDebugCVars::SnapGoal();
+				UE_VLOG(LogOwner, LogNetInterpolation, Log, TEXT("Snapping to %s"), *LexToString(NewInterpolationTime));
+			}
 			
 			// Speed up if we are too far behind
 			FRealTime CatchUpThreshold = SimulationTimeBuffer.HeadElement()->ToRealTimeSeconds() - NetworkInterpolationDebugCVars::CatchUpThreshold();
-			if (CatchUpUntilTime <= 0.f && InterpolationTime < CatchUpThreshold)
+			if (CatchUpUntilTime <= 0.f && NewInterpolationTime < CatchUpThreshold)
 			{
 				CatchUpUntilTime = SimulationTimeBuffer.HeadElement()->ToRealTimeSeconds() - NetworkInterpolationDebugCVars::CatchUpGoal();
 			}
 
 			if (CatchUpUntilTime > 0.f)
 			{
-				if (InterpolationTime  < CatchUpUntilTime)
+				if (NewInterpolationTime  < CatchUpUntilTime)
 				{
 					Step *= NetworkInterpolationDebugCVars::CatchUpFactor();
 					LoggingContext = EVisualLoggingContext::InterpolationSpeedUp;
 
-					UE_VLOG(LogOwner, LogNetInterpolation, Log, TEXT("Catching up! %s < %s"), *LexToString(InterpolationTime), *LexToString(CatchUpUntilTime));
-					//!UE_LOG(LogNetInterpolation, Warning, TEXT("Catching up! %s < %s"), *LexToString(InterpolationTime), *LexToString(CatchUpUntilTime));
+					UE_VLOG(LogOwner, LogNetInterpolation, Log, TEXT("Catching up! %s < %s"), *LexToString(NewInterpolationTime), *LexToString(CatchUpUntilTime));
 				}
 				else
 				{
@@ -234,27 +246,27 @@ struct TNetSimInterpolator
 					const TSyncState* DebugTail = Buffers.Sync.TailElement();
 					const TSyncState* DebugHead = Buffers.Sync.HeadElement();
 
-					auto VLogHelper = [&](int32 Frame, EVisualLoggingContext Context)
+					auto VLogHelper = [&](int32 Frame, EVisualLoggingContext Context, const FString& DebugStr)
 					{
-						FVisualLoggingParameters VLogParams(Context, Frame, EVisualLoggingLifetime::Transient);
+						FVisualLoggingParameters VLogParams(Context, Frame, EVisualLoggingLifetime::Transient, DebugStr);
 						Driver->InvokeVisualLog(Buffers.Input[Frame], Buffers.Sync[Frame], Buffers.Aux[Frame], VLogParams);
 					};
 
-					VLogHelper(Buffers.Sync.TailFrame(), EVisualLoggingContext::InterpolationBufferTail);
-					VLogHelper(Buffers.Sync.HeadFrame(), EVisualLoggingContext::InterpolationBufferHead);
+					VLogHelper(Buffers.Sync.TailFrame(), EVisualLoggingContext::InterpolationBufferTail, LexToString(TickInfo.SimulationTimeBuffer.HeadElement()->ToRealTimeMS()));
+					VLogHelper(Buffers.Sync.HeadFrame(), EVisualLoggingContext::InterpolationBufferHead, LexToString(TickInfo.SimulationTimeBuffer.TailElement()->ToRealTimeMS()));
 
 					{
-						FVisualLoggingParameters VLogParams(EVisualLoggingContext::InterpolationFrom, InterpolationFrame-1, EVisualLoggingLifetime::Transient);
+						FVisualLoggingParameters VLogParams(EVisualLoggingContext::InterpolationFrom, InterpolationFrame-1, EVisualLoggingLifetime::Transient, LexToString(FromRealTime));
 						Driver->InvokeVisualLog(Buffers.Input[InterpolationFrame-1], &FromState.Sync, &FromState.Aux, VLogParams);
 					}
 
 					{
-						FVisualLoggingParameters VLogParams(EVisualLoggingContext::InterpolationTo, InterpolationFrame, EVisualLoggingLifetime::Transient);
+						FVisualLoggingParameters VLogParams(EVisualLoggingContext::InterpolationTo, InterpolationFrame, EVisualLoggingLifetime::Transient, LexToString(ToRealTime));
 						Driver->InvokeVisualLog(Buffers.Input[InterpolationFrame], ToState, ToAuxState, VLogParams);
 					}
 
 					{
-						FVisualLoggingParameters VLogParams(LoggingContext, InterpolationFrame, EVisualLoggingLifetime::Transient);
+						FVisualLoggingParameters VLogParams(LoggingContext, InterpolationFrame, EVisualLoggingLifetime::Transient, LexToString(NewInterpolationTime));
 						Driver->InvokeVisualLog(Buffers.Input[InterpolationFrame], &OutputState.Sync, &OutputState.Aux, VLogParams);
 					}
 				}
diff --git a/Engine/Plugins/Runtime/NetworkPrediction/Source/NetworkPrediction/Public/NetworkedSimulationModelRepControllers.h b/Engine/Plugins/Runtime/NetworkPrediction/Source/NetworkPrediction/Public/NetworkedSimulationModelRepControllers.h
--- a/Engine/Plugins/Runtime/NetworkPrediction/Source/NetworkPrediction/Public/NetworkedSimulationModelRepControllers.h
+++ b/Engine/Plugins/Runtime/NetworkPrediction/Source/NetworkPrediction/Public/NetworkedSimulationModelRepControllers.h
@@ -17,8 +17,16 @@ namespace NetworkSimulationModelCVars
 	NETSIM_DEVCVAR_SHIPCONST_INT(EnableSimulatedExtrapolation, 1, "ns.EnableSimulatedExtrapolation", "Toggle simulated proxy extrapolation.");
 	NETSIM_DEVCVAR_SHIPCONST_INT(ForceReconcile, 0, "ns.ForceReconcile", "Forces reconcile even if state does not differ. E.g, force resimulation after every netupdate.");
 	NETSIM_DEVCVAR_SHIPCONST_INT(ForceReconcileSingle, 0, "ns.ForceReconcileSingle", "Forces a since reconcile to happen on the next frame");
+	NETSIM_DEVCVAR_SHIPCONST_INT(EnableDebugBufferReplication, 1, "ns.DebugBufferReplication", "Replicate debug buffer (takes lots of bandwidth. Always compiled out of ship/test)");
 }
 
+#define NETSIM_ENABLE_CHECKSUMS !(UE_BUILD_SHIPPING || UE_BUILD_TEST)
+#if NETSIM_ENABLE_CHECKSUMS 
+#define NETSIM_CHECKSUM(Ser) SerializeChecksum(Ser,0xA186A384, false);
+#else
+#define NETSIM_CHECKSUM
+#endif
+
 #ifndef NETSIM_NETCONSTANT_NUM_BITS_FRAME
 	#define NETSIM_NETCONSTANT_NUM_BITS_FRAME 8	// Allows you to override this setting via UBT, but access via cleaner FActorMotionNetworkingConstants::NUM_BITS_FRAME
 #endif
@@ -175,6 +183,8 @@ struct TRepController_SimTime : public TBase
 		TBase::NetSerialize(P, Buffers, Ticker);
 		SerializedTime = Ticker.GetTotalProcessedSimulationTime();
 		SerializedTime.NetSerialize(P.Ar);
+
+		NETSIM_CHECKSUM(P.Ar);
 	}
 
 	FNetworkSimTime SerializedTime;
@@ -208,6 +218,8 @@ struct TRepController_Sequence : public TBase
 
 	void NetSerialize(const FNetSerializeParams& P, TNetworkSimBufferContainer<Model>& Buffers, TSimulationTicker<TTickSettings>& Ticker)
 	{
+		NETSIM_CHECKSUM(P.Ar);
+
 		TBase::NetSerialize(P, Buffers, Ticker);
 
 		auto& Buffer = Buffers.template Get<BufferId>();
@@ -237,6 +249,8 @@ struct TRepController_Sequence : public TBase
 		}
 
 		LastSerializedFrame = HeadFrame;
+
+		NETSIM_CHECKSUM(P.Ar);
 	}
 
 	int32 GetLastSerializedFrame() const { return LastSerializedFrame; }
@@ -333,6 +347,7 @@ struct TRepController_Simulated : public TBase
 	void NetSerialize(const FNetSerializeParams& P, TNetworkSimBufferContainer<Model>& Buffers, TSimulationTicker<TTickSettings>& Ticker)
 	{
 		FArchive& Ar = P.Ar;
+		NETSIM_CHECKSUM(P.Ar);
 
 		FNetworkSimTime PrevLastSerializedSimulationTime = LastSerializedSimulationTime;
 
@@ -442,6 +457,8 @@ struct TRepController_Simulated : public TBase
 			LastSerializedSyncState = *SyncState;
 			LastSerializedAuxState = *AuxState;
 		}
+
+		NETSIM_CHECKSUM(P.Ar);
 	}
 
 	void Reconcile(TSimulation* Simulation, TDriver* Driver, TNetworkSimBufferContainer<Model>& Buffers, TSimulationTicker<TTickSettings>& Ticker)
@@ -639,6 +656,7 @@ struct TRepController_Autonomous: public TBase
 	// --------------------------------------------------------------------
 	void NetSerialize(const FNetSerializeParams& P, TNetworkSimBufferContainer<Model>& Buffers, TSimulationTicker<TTickSettings>& Ticker)
 	{
+		NETSIM_CHECKSUM(P.Ar);
 		FArchive& Ar = P.Ar;
 		
 		SerializedFrame = FNetworkSimulationSerialization::SerializeFrame(Ar, Buffers.Sync.HeadFrame());
@@ -695,6 +713,7 @@ struct TRepController_Autonomous: public TBase
 				}
 			}
 		}
+		NETSIM_CHECKSUM(P.Ar);
 	}
 
 	// --------------------------------------------------------------------
@@ -880,11 +899,17 @@ struct TRepController_Debug : public TBase
 
 	int32 GetProxyDirtyCount(TNetworkSimBufferContainer<Model>& Buffers) const 
 	{
+		if (NetworkSimulationModelCVars::EnableDebugBufferReplication() == 0)
+		{
+			return 0;
+		}
+
 		return Buffers.Debug.GetDirtyCount() ^ (TBase::GetProxyDirtyCount(Buffers) << 2); 
 	}
 
 	void NetSerialize(const FNetSerializeParams& P, TNetworkSimBufferContainer<Model>& Buffers, TSimulationTicker<TTickSettings>& Ticker)
 	{
+		NETSIM_CHECKSUM(P.Ar);
 		TBase::NetSerialize(P, Buffers, Ticker);
 		FArchive& Ar = P.Ar;
 
@@ -913,6 +938,7 @@ struct TRepController_Debug : public TBase
 			auto* Cmd = Ar.IsLoading() ? Buffer.WriteFrame(Frame) : Buffer[Frame];
 			Cmd->NetSerialize(P);
 		}
+		NETSIM_CHECKSUM(P.Ar);
 	}
 
 	TDebugBuffer ReceivedBuffer;
