Network Prediction Plugin
-Fix a few issues around interpolation: it should now better reflect the actual sync state even when in wait mode.

#rnx

diff --git a/Engine/Plugins/Runtime/NetworkPrediction/Source/NetworkPrediction/Public/NetworkedSimulationModelInterpolator.h b/Engine/Plugins/Runtime/NetworkPrediction/Source/NetworkPrediction/Public/NetworkedSimulationModelInterpolator.h
--- a/Engine/Plugins/Runtime/NetworkPrediction/Source/NetworkPrediction/Public/NetworkedSimulationModelInterpolator.h
+++ b/Engine/Plugins/Runtime/NetworkPrediction/Source/NetworkPrediction/Public/NetworkedSimulationModelInterpolator.h
@@ -79,9 +79,10 @@ struct TNetSimInterpolator
 
 		const AActor* LogOwner = Driver->GetVLogOwner();
 
-		// Interpolation disabled
-		if (NetworkInterpolationDebugCVars::Disable() > 0)
-		{			
+		// Interpolation disabled or we don't have 2 elements yet
+		if (NetworkInterpolationDebugCVars::Disable() > 0 || TickInfo.SimulationTimeBuffer.Num() <= 1)
+		{
+			// We are still responsible for calling FinalizeFrame though
 			if (const TSyncState* HeadState = Buffers.Sync.HeadElement())
 			{
 				const TAuxState* AuxState = Buffers.Aux.HeadElement();
@@ -89,14 +90,9 @@ struct TNetSimInterpolator
 
 				Driver->FinalizeFrame(*HeadState, *AuxState);
 			}
+			
 			return TickInfo.GetTotalProcessedSimulationTime().ToRealTimeSeconds();
 		}
-		
-		if (TickInfo.SimulationTimeBuffer.Num() <= 1)
-		{
-			// Cant interpolate yet	
-			return 0.f;
-		}
 
 		auto& SimulationTimeBuffer = TickInfo.SimulationTimeBuffer;
 
@@ -111,6 +107,9 @@ struct TNetSimInterpolator
 			FromState.Aux = *Buffers.Aux[InterpolationFrame];
 		}
 
+		EVisualLoggingContext LoggingContext = EVisualLoggingContext::InterpolationLatest;
+		FRealTime NewInterpolationTime = InterpolationTime;
+
 		// Wait if we were too far ahead
 		if (WaitUntilTime > 0.f)
 		{
@@ -131,15 +130,11 @@ struct TNetSimInterpolator
 					FVector2D LocalTimeVsInterpolationTime(LogOwner->GetWorld()->GetTimeSeconds(), (int32)(InterpolationTime * 1000.f));
 					UE_VLOG_HISTOGRAM(LogOwner, LogNetInterpolation, Log, "ServerSimulationTimeGraph", LocalInterpolationTimeName, LocalTimeVsInterpolationTime);
 				}
-				return InterpolationTime;
 			}
 		}
-
-		EVisualLoggingContext LoggingContext = EVisualLoggingContext::InterpolationLatest;
-
-		// Calc new interpolation time
-		FRealTime NewInterpolationTime = InterpolationTime;
+		else
 		{
+			// Calc new interpolation time
 			FRealTime Step = DeltaSeconds;
 
 			// Snap if way far behind
@@ -206,74 +201,73 @@ struct TNetSimInterpolator
 			const FRealTime FromRealTime = InterpolationTime;
 			const FRealTime ToRealTime = ToTime;
 			const FRealTime InterpolationInterval = ToRealTime - FromRealTime;
-		
-			if (ensure(FMath::Abs(InterpolationInterval) > 0.f))
-			{
-				const float InterpolationPCT = (NewInterpolationTime - FromRealTime) / InterpolationInterval;
-				ensureMsgf(InterpolationPCT >= 0.f && InterpolationPCT <= 1.f, TEXT("Calculated InterpolationPCT not in expected range. NewInterpolationTime: %s. From: %s. To: %s"), *LexToString(NewInterpolationTime), *LexToString(FromRealTime), *LexToString(ToRealTime));
+			
+			const bool bValidInterval = FMath::Abs(InterpolationInterval) > 0.f;
 
-				auto& FromState = GetFromInterpolationState();
-				auto& OutputState = GetNextInterpolationState();
+			const float InterpolationPCT = bValidInterval ? ((NewInterpolationTime - FromRealTime) / InterpolationInterval) : 1.f;
+			ensureMsgf(InterpolationPCT >= 0.f && InterpolationPCT <= 1.f, TEXT("Calculated InterpolationPCT not in expected range. NewInterpolationTime: %s. From: %s. To: %s"), *LexToString(NewInterpolationTime), *LexToString(FromRealTime), *LexToString(ToRealTime));
 
-				Model::Interpolate({ { FromState.Sync, FromState.Aux }, { *ToState, *ToAuxState }, InterpolationPCT, { OutputState.Sync, OutputState.Aux } });
-				Driver->FinalizeFrame(OutputState.Sync, OutputState.Aux);
+			auto& FromState = GetFromInterpolationState();
+			auto& OutputState = GetNextInterpolationState();
+
+			Model::Interpolate({ { FromState.Sync, FromState.Aux }, { *ToState, *ToAuxState }, InterpolationPCT, { OutputState.Sync, OutputState.Aux } });
+			Driver->FinalizeFrame(OutputState.Sync, OutputState.Aux);
 				
-				if (bDoVLog)
-				{
-					UE_VLOG(LogOwner, LogNetInterpolation, Log, TEXT("%s - %s - %s.  InterpolationPCT: %f"), *LexToString(FromRealTime), *LexToString(NewInterpolationTime), *LexToString(ToRealTime), InterpolationPCT);
+			if (bDoVLog)
+			{
+				UE_VLOG(LogOwner, LogNetInterpolation, Log, TEXT("%s - %s - %s.  InterpolationPCT: %f"), *LexToString(FromRealTime), *LexToString(NewInterpolationTime), *LexToString(ToRealTime), InterpolationPCT);
 
-					// Graph Interpolation Time vs Buffer Head/Tail times
-					const FName ServerSimulationGraphName("ServerSimulationTimeGraph");
-					const FName ServerSimTimeName("Server Simulation Time");
-					FVector2D LocalTimeVsServerSimTime(LogOwner->GetWorld()->GetTimeSeconds(), (int32)TickInfo.SimulationTimeBuffer.HeadElement()->ToRealTimeMS());
-					UE_VLOG_HISTOGRAM(LogOwner, LogNetInterpolation, Log, ServerSimulationGraphName, ServerSimTimeName, LocalTimeVsServerSimTime);
+				// Graph Interpolation Time vs Buffer Head/Tail times
+				const FName ServerSimulationGraphName("ServerSimulationTimeGraph");
+				const FName ServerSimTimeName("Server Simulation Time");
+				FVector2D LocalTimeVsServerSimTime(LogOwner->GetWorld()->GetTimeSeconds(), (int32)TickInfo.SimulationTimeBuffer.HeadElement()->ToRealTimeMS());
+				UE_VLOG_HISTOGRAM(LogOwner, LogNetInterpolation, Log, ServerSimulationGraphName, ServerSimTimeName, LocalTimeVsServerSimTime);
 
-					const FName BufferTailSimTimeName("Buffer Tail Simulation Time");
-					FVector2D LocalTimeVsBufferTailSim(LogOwner->GetWorld()->GetTimeSeconds(), (int32)TickInfo.SimulationTimeBuffer.TailElement()->ToRealTimeMS());
-					UE_VLOG_HISTOGRAM(LogOwner, LogNetInterpolation, Log, ServerSimulationGraphName, BufferTailSimTimeName, LocalTimeVsBufferTailSim);
+				const FName BufferTailSimTimeName("Buffer Tail Simulation Time");
+				FVector2D LocalTimeVsBufferTailSim(LogOwner->GetWorld()->GetTimeSeconds(), (int32)TickInfo.SimulationTimeBuffer.TailElement()->ToRealTimeMS());
+				UE_VLOG_HISTOGRAM(LogOwner, LogNetInterpolation, Log, ServerSimulationGraphName, BufferTailSimTimeName, LocalTimeVsBufferTailSim);
 
-					const FName LocalInterpolationTimeName("Local Interpolation Time");
-					FVector2D LocalTimeVsInterpolationTime(LogOwner->GetWorld()->GetTimeSeconds(), (int32)(NewInterpolationTime * 1000.f));
-					UE_VLOG_HISTOGRAM(LogOwner, LogNetInterpolation, Log, ServerSimulationGraphName, LocalInterpolationTimeName, LocalTimeVsInterpolationTime);
+				const FName LocalInterpolationTimeName("Local Interpolation Time");
+				FVector2D LocalTimeVsInterpolationTime(LogOwner->GetWorld()->GetTimeSeconds(), (int32)(NewInterpolationTime * 1000.f));
+				UE_VLOG_HISTOGRAM(LogOwner, LogNetInterpolation, Log, ServerSimulationGraphName, LocalInterpolationTimeName, LocalTimeVsInterpolationTime);
 					
-					FVector2D LocalTimeVsCatchUpThreshold(LogOwner->GetWorld()->GetTimeSeconds(), (SimulationTimeBuffer.HeadElement()->ToRealTimeSeconds() - NetworkInterpolationDebugCVars::CatchUpThreshold()) * 1000.f);
-					UE_VLOG_HISTOGRAM(LogOwner, LogNetInterpolation, Log, ServerSimulationGraphName, "Catch Up Threshold", LocalTimeVsCatchUpThreshold);
-
-					FVector2D LocalTimeVsCatchUpGoal(LogOwner->GetWorld()->GetTimeSeconds(), (SimulationTimeBuffer.HeadElement()->ToRealTimeSeconds() - NetworkInterpolationDebugCVars::CatchUpGoal()) * 1000.f);
-					UE_VLOG_HISTOGRAM(LogOwner, LogNetInterpolation, Log, ServerSimulationGraphName, "Catch Up Goal", LocalTimeVsCatchUpGoal);
-
-					// VLog the actual motion states
-					const TSyncState* DebugTail = Buffers.Sync.TailElement();
-					const TSyncState* DebugHead = Buffers.Sync.HeadElement();
-
-					auto VLogHelper = [&](int32 Frame, EVisualLoggingContext Context, const FString& DebugStr)
-					{
-						FVisualLoggingParameters VLogParams(Context, Frame, EVisualLoggingLifetime::Transient, DebugStr);
-						Driver->InvokeVisualLog(Buffers.Input[Frame], Buffers.Sync[Frame], Buffers.Aux[Frame], VLogParams);
-					};
-
-					VLogHelper(Buffers.Sync.TailFrame(), EVisualLoggingContext::InterpolationBufferTail, LexToString(TickInfo.SimulationTimeBuffer.HeadElement()->ToRealTimeMS()));
-					VLogHelper(Buffers.Sync.HeadFrame(), EVisualLoggingContext::InterpolationBufferHead, LexToString(TickInfo.SimulationTimeBuffer.TailElement()->ToRealTimeMS()));
-
-					{
-						FVisualLoggingParameters VLogParams(EVisualLoggingContext::InterpolationFrom, InterpolationFrame-1, EVisualLoggingLifetime::Transient, LexToString(FromRealTime));
-						Driver->InvokeVisualLog(Buffers.Input[InterpolationFrame-1], &FromState.Sync, &FromState.Aux, VLogParams);
-					}
-
-					{
-						FVisualLoggingParameters VLogParams(EVisualLoggingContext::InterpolationTo, InterpolationFrame, EVisualLoggingLifetime::Transient, LexToString(ToRealTime));
-						Driver->InvokeVisualLog(Buffers.Input[InterpolationFrame], ToState, ToAuxState, VLogParams);
-					}
-
-					{
-						FVisualLoggingParameters VLogParams(LoggingContext, InterpolationFrame, EVisualLoggingLifetime::Transient, LexToString(NewInterpolationTime));
-						Driver->InvokeVisualLog(Buffers.Input[InterpolationFrame], &OutputState.Sync, &OutputState.Aux, VLogParams);
-					}
+				FVector2D LocalTimeVsCatchUpThreshold(LogOwner->GetWorld()->GetTimeSeconds(), (SimulationTimeBuffer.HeadElement()->ToRealTimeSeconds() - NetworkInterpolationDebugCVars::CatchUpThreshold()) * 1000.f);
+				UE_VLOG_HISTOGRAM(LogOwner, LogNetInterpolation, Log, ServerSimulationGraphName, "Catch Up Threshold", LocalTimeVsCatchUpThreshold);
+
+				FVector2D LocalTimeVsCatchUpGoal(LogOwner->GetWorld()->GetTimeSeconds(), (SimulationTimeBuffer.HeadElement()->ToRealTimeSeconds() - NetworkInterpolationDebugCVars::CatchUpGoal()) * 1000.f);
+				UE_VLOG_HISTOGRAM(LogOwner, LogNetInterpolation, Log, ServerSimulationGraphName, "Catch Up Goal", LocalTimeVsCatchUpGoal);
+
+				// VLog the actual motion states
+				const TSyncState* DebugTail = Buffers.Sync.TailElement();
+				const TSyncState* DebugHead = Buffers.Sync.HeadElement();
+
+				auto VLogHelper = [&](int32 Frame, EVisualLoggingContext Context, const FString& DebugStr)
+				{
+					FVisualLoggingParameters VLogParams(Context, Frame, EVisualLoggingLifetime::Transient, DebugStr);
+					Driver->InvokeVisualLog(Buffers.Input[Frame], Buffers.Sync[Frame], Buffers.Aux[Frame], VLogParams);
+				};
+
+				VLogHelper(Buffers.Sync.TailFrame(), EVisualLoggingContext::InterpolationBufferTail, LexToString(TickInfo.SimulationTimeBuffer.HeadElement()->ToRealTimeMS()));
+				VLogHelper(Buffers.Sync.HeadFrame(), EVisualLoggingContext::InterpolationBufferHead, LexToString(TickInfo.SimulationTimeBuffer.TailElement()->ToRealTimeMS()));
+
+				{
+					FVisualLoggingParameters VLogParams(EVisualLoggingContext::InterpolationFrom, InterpolationFrame-1, EVisualLoggingLifetime::Transient, LexToString(FromRealTime));
+					Driver->InvokeVisualLog(Buffers.Input[InterpolationFrame-1], &FromState.Sync, &FromState.Aux, VLogParams);
+				}
+
+				{
+					FVisualLoggingParameters VLogParams(EVisualLoggingContext::InterpolationTo, InterpolationFrame, EVisualLoggingLifetime::Transient, LexToString(ToRealTime));
+					Driver->InvokeVisualLog(Buffers.Input[InterpolationFrame], ToState, ToAuxState, VLogParams);
+				}
+
+				{
+					FVisualLoggingParameters VLogParams(LoggingContext, InterpolationFrame, EVisualLoggingLifetime::Transient, LexToString(NewInterpolationTime));
+					Driver->InvokeVisualLog(Buffers.Input[InterpolationFrame], &OutputState.Sync, &OutputState.Aux, VLogParams);
 				}
-				
-				InterpolationTime = NewInterpolationTime;
-				InternalIdx ^= 1;
 			}
+				
+			InterpolationTime = NewInterpolationTime;
+			InternalIdx ^= 1;
 		}
 
 		return InterpolationTime;
