NetworkPrediction: visual log last serialized state in addition to latest predicted state for simulated

#rnx

diff --git a/Engine/Plugins/Runtime/NetworkPrediction/Source/NetworkPrediction/Public/NetworkedSimulationModelDebugger.h b/Engine/Plugins/Runtime/NetworkPrediction/Source/NetworkPrediction/Public/NetworkedSimulationModelDebugger.h
--- a/Engine/Plugins/Runtime/NetworkPrediction/Source/NetworkPrediction/Public/NetworkedSimulationModelDebugger.h
+++ b/Engine/Plugins/Runtime/NetworkPrediction/Source/NetworkPrediction/Public/NetworkedSimulationModelDebugger.h
@@ -510,9 +510,16 @@ struct TNetworkSimulationModelDebugger : public INetworkSimulationModelDebugger
 		}
 		else if (Owner->GetLocalRole() == ROLE_SimulatedProxy)
 		{
-			FVisualLoggingParameters VLogParams(EVisualLoggingContext::LastConfirmed, NetworkSim->Buffers.Sync.HeadFrame(), EVisualLoggingLifetime::Transient);
-			int32 HeadFrame = NetworkSim->Buffers.Sync.HeadFrame();
-			NetworkSim->Driver->InvokeVisualLog(NetworkSim->Buffers.Input[HeadFrame], NetworkSim->Buffers.Sync[HeadFrame], NetworkSim->Buffers.Aux[HeadFrame], VLogParams);
+			{
+				FVisualLoggingParameters VLogParams(EVisualLoggingContext::LastPredicted, NetworkSim->Buffers.Sync.HeadFrame(), EVisualLoggingLifetime::Transient);
+				int32 HeadFrame = NetworkSim->Buffers.Sync.HeadFrame();
+				NetworkSim->Driver->InvokeVisualLog(NetworkSim->Buffers.Input[HeadFrame], NetworkSim->Buffers.Sync[HeadFrame], NetworkSim->Buffers.Aux[HeadFrame], VLogParams);
+			}
+
+			{
+				FVisualLoggingParameters VLogParams(EVisualLoggingContext::LastConfirmed, NetworkSim->Buffers.Sync.HeadFrame(), EVisualLoggingLifetime::Transient);
+				NetworkSim->Driver->InvokeVisualLog(&NetworkSim->RepProxy_Simulated.GetLastSerializedInputCmd(), &NetworkSim->RepProxy_Simulated.GetLastSerializedSyncState(), &NetworkSim->RepProxy_Simulated.GetLastSerializedAuxState(), VLogParams);
+			}
 			
 			if (NetworkSim->GetSimulatedUpdateMode() != ESimulatedUpdateMode::Interpolate)
 			{
diff --git a/Engine/Plugins/Runtime/NetworkPrediction/Source/NetworkPrediction/Public/NetworkedSimulationModelRepControllers.h b/Engine/Plugins/Runtime/NetworkPrediction/Source/NetworkPrediction/Public/NetworkedSimulationModelRepControllers.h
--- a/Engine/Plugins/Runtime/NetworkPrediction/Source/NetworkPrediction/Public/NetworkedSimulationModelRepControllers.h
+++ b/Engine/Plugins/Runtime/NetworkPrediction/Source/NetworkPrediction/Public/NetworkedSimulationModelRepControllers.h
@@ -324,6 +324,8 @@ struct TRepController_Simulated : public TBase
 	// Last Serialized time and state
 	FNetworkSimTime GetLastSerializedSimulationTime() const { return LastSerializedSimulationTime; }
 	const TSyncState& GetLastSerializedSyncState() const { return LastSerializedSyncState; }
+	const TAuxState& GetLastSerializedAuxState() const { return LastSerializedAuxState; }
+	const TInputCmd& GetLastSerializedInputCmd() const { return LastSerializedInputCmd; }
 	
 	ESimulatedUpdateMode GetSimulatedUpdateMode() const
 	{
