Fix linux compile error. There is a fix submitted to another stream but who knows how long it'll take to propagate, so reapplying the fix here so that CIS can succeed.

#rnx

diff --git a/Engine/Plugins/Runtime/NetworkPrediction/Source/NetworkPrediction/Public/NetworkSimulationModelReplicators.h b/Engine/Plugins/Runtime/NetworkPrediction/Source/NetworkPrediction/Public/NetworkSimulationModelReplicators.h
--- a/Engine/Plugins/Runtime/NetworkPrediction/Source/NetworkPrediction/Public/NetworkSimulationModelReplicators.h
+++ b/Engine/Plugins/Runtime/NetworkPrediction/Source/NetworkPrediction/Public/NetworkSimulationModelReplicators.h
@@ -3,7 +3,7 @@
 #pragma once
 
 #include "Net/UnrealNetwork.h" // For MakeRelative
-#include "NetworkSimulationModelCvars.h"
+#include "NetworkSimulationModelCVars.h"
 
 // ----------------------------------------------------------------------------------------------------------------------------------------------------------------------
 //	CVars and compile time constants
