Fix another linux compile error. There is a fix submitted to another stream but who knows how long it'll take to propagate, so reapplying the fix here so that CIS can succeed.

#rnx

diff --git a/Engine/Plugins/Runtime/NetworkPrediction/Source/NetworkPrediction/Public/NetworkSimulationModelDebugger.h b/Engine/Plugins/Runtime/NetworkPrediction/Source/NetworkPrediction/Public/NetworkSimulationModelDebugger.h
--- a/Engine/Plugins/Runtime/NetworkPrediction/Source/NetworkPrediction/Public/NetworkSimulationModelDebugger.h
+++ b/Engine/Plugins/Runtime/NetworkPrediction/Source/NetworkPrediction/Public/NetworkSimulationModelDebugger.h
@@ -15,7 +15,7 @@
 #include "Containers/Array.h"
 #include "Templates/UniquePtr.h"
 #include "NetworkSimulationModel.h"
-#include "NetworkSimulationModelCvars.h"
+#include "NetworkSimulationModelCVars.h"
 
 DEFINE_LOG_CATEGORY_STATIC(LogNetworkSimDebug, Log, All);
 
