Network Prediction Insights
Runtime changes for supporting new tracing system for debugging.

#RNX

diff --git a/Engine/Plugins/Runtime/NetworkPrediction/Source/NetworkPrediction/NetworkPrediction.Build.cs b/Engine/Plugins/Runtime/NetworkPrediction/Source/NetworkPrediction/NetworkPrediction.Build.cs
--- a/Engine/Plugins/Runtime/NetworkPrediction/Source/NetworkPrediction/NetworkPrediction.Build.cs
+++ b/Engine/Plugins/Runtime/NetworkPrediction/Source/NetworkPrediction/NetworkPrediction.Build.cs
@@ -9,7 +9,8 @@ namespace UnrealBuildTool.Rules
 			PublicIncludePaths.AddRange(
 				new string[] {
                     ModuleDirectory + "/Public",
-				}
+                    "Runtime/TraceLog/Public",
+                }
 				);
 
 			PrivateIncludePaths.AddRange(
@@ -31,7 +32,11 @@ namespace UnrealBuildTool.Rules
 			PrivateDependencyModuleNames.AddRange(
 				new string[]
 				{
-				}
+					"TraceLog",
+					"TraceAnalysis",
+					"TraceServices",
+					"TraceInsights",
+                }
 				);
 
 			DynamicallyLoadedModuleNames.AddRange(
@@ -40,6 +45,17 @@ namespace UnrealBuildTool.Rules
 				}
 				);
 
-		}
+
+
+            // Only needed for the PIE delegate in FNetworkPredictionModule::StartupModule
+            if (Target.Type == TargetType.Editor) {
+                PrivateDependencyModuleNames.AddRange(
+                new string[]
+                {
+                    "UnrealEd",
+                });
+            }
+
+        }
 	}
 }
diff --git a/Engine/Plugins/Runtime/NetworkPrediction/Source/NetworkPrediction/Private/Movement/ParametricMovement.cpp b/Engine/Plugins/Runtime/NetworkPrediction/Source/NetworkPrediction/Private/Movement/ParametricMovement.cpp
--- a/Engine/Plugins/Runtime/NetworkPrediction/Source/NetworkPrediction/Private/Movement/ParametricMovement.cpp
+++ b/Engine/Plugins/Runtime/NetworkPrediction/Source/NetworkPrediction/Private/Movement/ParametricMovement.cpp
@@ -194,6 +194,12 @@ void UParametricMovementComponent::ProduceInput(const FNetworkSimTime DeltaTimeS
 	PendingPlayRate.Reset();
 }
 
+void UParametricMovementComponent::EnableInterpolationMode(bool bValue)
+{
+	bEnableInterpolation = bValue;
+	GetNetworkSimulation()->SetEnableSimulationExtrapolation(!bEnableInterpolation);
+}
+
 // -------------------------------------------------------------------------------------------
 // -------------------------------------------------------------------------------------------
 
diff --git a/Engine/Plugins/Runtime/NetworkPrediction/Source/NetworkPrediction/Private/NetworkPredictionModule.cpp b/Engine/Plugins/Runtime/NetworkPrediction/Source/NetworkPrediction/Private/NetworkPredictionModule.cpp
--- a/Engine/Plugins/Runtime/NetworkPrediction/Source/NetworkPrediction/Private/NetworkPredictionModule.cpp
+++ b/Engine/Plugins/Runtime/NetworkPrediction/Source/NetworkPrediction/Private/NetworkPredictionModule.cpp
@@ -7,6 +7,15 @@
 #include "Engine/World.h"
 #include "NetworkedSimulationModelCues.h"
 #include "Misc/CoreDelegates.h"
+#include "Trace/NetworkPredictionTrace.h"
+#include "Trace/Trace.h"
+#include "Misc/CommandLine.h"
+#include "Misc/Parse.h"
+#include "String/ParseTokens.h"
+
+#if WITH_EDITOR
+#include "Editor.h"
+#endif
 
 #define LOCTEXT_NAMESPACE "FNetworkPredictionModule"
 
@@ -15,19 +24,43 @@ class FNetworkPredictionModule : public INetworkPredictionModule
 	/** IModuleInterface implementation */
 	virtual void StartupModule() override;
 	virtual void ShutdownModule() override;
+
+	FDelegateHandle PieHandle;
 };
 
 void FNetworkPredictionModule::StartupModule()
 {
+	// Disable by default unless in the command line args. This is temp as the existing insights -trace parsing happen before the plugin is loaded
+	Trace::ToggleChannel(TEXT("NetworkPredictionChannel"), false);
+
+	FString EnabledChannels;
+	FParse::Value(FCommandLine::Get(), TEXT("-trace="), EnabledChannels, false);
+	UE::String::ParseTokens(EnabledChannels, TEXT(","), [](FStringView Token) {
+		if (Token.Compare(TEXT("NetworkPrediction"), ESearchCase::IgnoreCase)==0 || Token.Compare(TEXT("NP"), ESearchCase::IgnoreCase)==0)
+		{
+			Trace::ToggleChannel(TEXT("NetworkPredictionChannel"), true);
+		}
+	});
+
 	FCoreDelegates::OnPostEngineInit.AddLambda([]()
 	{
 		FGlobalCueTypeTable::Get().FinalizeTypes();
 	});
+
+#if WITH_EDITOR
+	PieHandle = FEditorDelegates::PreBeginPIE.AddLambda([this](const bool bBegan)
+	{
+		UE_NP_TRACE_PIE_START();
+	});
+#endif
 }
 
 
 void FNetworkPredictionModule::ShutdownModule()
 {
+#if WITH_EDITOR
+	FEditorDelegates::PreBeginPIE.Remove(PieHandle);
+#endif
 }
 
 IMPLEMENT_MODULE( FNetworkPredictionModule, NetworkPrediction )
diff --git a/Engine/Plugins/Runtime/NetworkPrediction/Source/NetworkPrediction/Private/NetworkPredictionTypes.cpp b/Engine/Plugins/Runtime/NetworkPrediction/Source/NetworkPrediction/Private/NetworkPredictionTypes.cpp
--- a/Engine/Plugins/Runtime/NetworkPrediction/Source/NetworkPrediction/Private/NetworkPredictionTypes.cpp
+++ b/Engine/Plugins/Runtime/NetworkPrediction/Source/NetworkPrediction/Private/NetworkPredictionTypes.cpp
@@ -11,6 +11,11 @@
 
 DEFINE_LOG_CATEGORY(LogNetworkSim);
 
+uint32 INetworkedSimulationModel::GetNextSimulationId()
+{
+	static uint32 SimulationIdCounter = 1;
+	return SimulationIdCounter++;
+}
 
 namespace NetSimVLogCVars
 {
diff --git a/Engine/Plugins/Runtime/NetworkPrediction/Source/NetworkPrediction/Private/Trace/NetworkPredictionTrace.cpp b/Engine/Plugins/Runtime/NetworkPrediction/Source/NetworkPrediction/Private/Trace/NetworkPredictionTrace.cpp
new file mode 100644
--- /dev/null
+++ b/Engine/Plugins/Runtime/NetworkPrediction/Source/NetworkPrediction/Private/Trace/NetworkPredictionTrace.cpp
@@ -0,0 +1,346 @@
+// Copyright Epic Games, Inc. All Rights Reserved.
+
+
+#include "Trace/NetworkPredictionTrace.h"
+#include "CoreMinimal.h"
+#include "TraceServices/Model/AnalysisSession.h"
+#include "HAL/PlatformTime.h"
+#include "UObject/ObjectKey.h"
+#include "Engine/EngineBaseTypes.h"
+#include "Engine/GameInstance.h"
+#include "Engine/World.h"
+#include "GameFramework/Actor.h"
+#include "Engine/NetConnection.h"
+#include "UObject/CoreNet.h"
+#include "Engine/PackageMapClient.h"
+#include "Logging/LogMacros.h"
+//#include "Trace/Trace.inl" this will need to be included after next copy up from Dev-Core
+
+UE_TRACE_CHANNEL_DEFINE(NetworkPredictionChannel)
+
+UE_TRACE_EVENT_BEGIN(NetworkPrediction, GameInstanceRegister)
+	UE_TRACE_EVENT_FIELD(uint32, GameInstanceId)
+	UE_TRACE_EVENT_FIELD(bool, IsServer)
+UE_TRACE_EVENT_END()
+
+// Trace a simulation creation. GroupName is attached as attachment.
+UE_TRACE_EVENT_BEGIN(NetworkPrediction, SimulationCreated)
+	UE_TRACE_EVENT_FIELD(uint32, SimulationId)
+	UE_TRACE_EVENT_FIELD(uint32, GameInstanceId)
+UE_TRACE_EVENT_END()
+
+// Trace a simulation's local netrole, which can change during the sim's liftetime.
+UE_TRACE_EVENT_BEGIN(NetworkPrediction, SimulationNetRole)
+	UE_TRACE_EVENT_FIELD(uint64, EngineFrameNumber)
+	UE_TRACE_EVENT_FIELD(uint32, SimulationId)
+	UE_TRACE_EVENT_FIELD(uint8, NetRole)
+UE_TRACE_EVENT_END()
+
+// Trace a simulation's NetGUID assignment, which unfortunately hasn't always happened at creation.
+UE_TRACE_EVENT_BEGIN(NetworkPrediction, SimulationNetGUID)
+	UE_TRACE_EVENT_FIELD(uint32, SimulationId)
+	UE_TRACE_EVENT_FIELD(uint32, NetGUID)
+UE_TRACE_EVENT_END()
+
+// Traces a simulation tick in any context
+UE_TRACE_EVENT_BEGIN(NetworkPrediction, SimulationTick)
+	UE_TRACE_EVENT_FIELD(uint32, SimulationId)
+	UE_TRACE_EVENT_FIELD(uint64, EngineFrameNumber)
+	UE_TRACE_EVENT_FIELD(int32, StartMS)
+	UE_TRACE_EVENT_FIELD(int32, EndMS)
+	UE_TRACE_EVENT_FIELD(int32, OutputFrame)	
+
+UE_TRACE_EVENT_END()
+
+UE_TRACE_EVENT_BEGIN(NetworkPrediction, SimulationEOF)
+	UE_TRACE_EVENT_FIELD(uint32, SimulationId)
+	UE_TRACE_EVENT_FIELD(uint64, EngineFrameNumber)
+	UE_TRACE_EVENT_FIELD(double, EngineFrameDeltaTime)	// FIXME: Need better tracking of engine frame vs simulation EOF updates
+	UE_TRACE_EVENT_FIELD(double, EngineCurrentTime)		// FIXME: ^^
+	UE_TRACE_EVENT_FIELD(int32, TotalProcessedMS)
+	UE_TRACE_EVENT_FIELD(int32, TotalAllowedMS)
+	UE_TRACE_EVENT_FIELD(int32, LastSentKeyframe)
+	UE_TRACE_EVENT_FIELD(int32, LastReceivedKeyframe)
+UE_TRACE_EVENT_END()
+
+UE_TRACE_EVENT_BEGIN(NetworkPrediction, NetSerializeRecv)
+	UE_TRACE_EVENT_FIELD(uint32, SimulationId)
+	UE_TRACE_EVENT_FIELD(uint64, EngineFrameNumber)
+	UE_TRACE_EVENT_FIELD(int32, ReceivedSimTimeMS)
+	UE_TRACE_EVENT_FIELD(int32, ReceivedFrame)
+UE_TRACE_EVENT_END()
+
+UE_TRACE_EVENT_BEGIN(NetworkPrediction, InputCmd)
+	UE_TRACE_EVENT_FIELD(uint32, SimulationId)
+	UE_TRACE_EVENT_FIELD(int32, SimulationFrame)
+	UE_TRACE_EVENT_FIELD(uint64, EngineFrameNumber)
+UE_TRACE_EVENT_END()
+
+UE_TRACE_EVENT_BEGIN(NetworkPrediction, SyncState)
+	UE_TRACE_EVENT_FIELD(uint32, SimulationId)
+	UE_TRACE_EVENT_FIELD(int32, SimulationFrame)
+	UE_TRACE_EVENT_FIELD(uint64, EngineFrameNumber)
+UE_TRACE_EVENT_END()
+
+UE_TRACE_EVENT_BEGIN(NetworkPrediction, AuxState)
+	UE_TRACE_EVENT_FIELD(uint32, SimulationId)
+	UE_TRACE_EVENT_FIELD(int32, SimulationFrame)
+	UE_TRACE_EVENT_FIELD(uint64, EngineFrameNumber)
+UE_TRACE_EVENT_END()
+
+UE_TRACE_EVENT_BEGIN(NetworkPrediction, NetSerializeCommit)
+	UE_TRACE_EVENT_FIELD(uint32, SimulationId)
+UE_TRACE_EVENT_END()
+
+UE_TRACE_EVENT_BEGIN(NetworkPrediction,NetSerializeFault)
+	UE_TRACE_EVENT_FIELD(uint32,SimulationId)
+UE_TRACE_EVENT_END()
+
+UE_TRACE_EVENT_BEGIN(NetworkPrediction, OOBStateMod)
+	UE_TRACE_EVENT_FIELD(uint32, SimulationId)
+UE_TRACE_EVENT_END()
+
+UE_TRACE_EVENT_BEGIN(NetworkPrediction, ProduceInput)
+	UE_TRACE_EVENT_FIELD(uint32, SimulationId)
+UE_TRACE_EVENT_END()
+
+UE_TRACE_EVENT_BEGIN(NetworkPrediction, SynthInput)
+UE_TRACE_EVENT_FIELD(uint32, SimulationId)
+UE_TRACE_EVENT_END()
+
+UE_TRACE_EVENT_BEGIN(NetworkPrediction, PieBegin)
+UE_TRACE_EVENT_END()
+
+
+// ----------------------------------------------------------------------------
+//	Internal stack to track which simulation is active. 
+//		-Avoid passing or duplicating SimulationId to rep controllers, interpolaters, etc
+//		-Allows us to trace dependent simulations a bit easier
+// ----------------------------------------------------------------------------
+
+TArray<uint32, TInlineAllocator<4>> SimulationIdStack;
+uint32 PeakSimulationIdChecked()
+{
+	check(SimulationIdStack.Num() > 0);
+	return SimulationIdStack.Last();
+}
+
+// Assign Id to UGameInstance ObjectKey. Assignment
+struct FGameInstanceIdMap
+{
+	uint32 GetId(UGameInstance* Instance)
+	{
+		FObjectKey Key(Instance);
+		
+		int32 FoundIdx=INDEX_NONE;
+		if (AssignedInstances.Find(Key, FoundIdx))
+		{
+			return (uint32)(FoundIdx+1);
+		}
+
+		AssignedInstances.Add(Key);
+		uint32 Id = (uint32)AssignedInstances.Num();
+
+		const bool bIsServer = Instance->GetWorld()->GetNetMode() != ENetMode::NM_Client;
+
+		UE_TRACE_LOG(NetworkPrediction, GameInstanceRegister, NetworkPredictionChannel)
+			<< GameInstanceRegister.GameInstanceId(Id)
+			<< GameInstanceRegister.IsServer(bIsServer);
+		
+		return Id;
+	}
+
+private:
+	TArray<FObjectKey> AssignedInstances;
+
+} GameInstanceMap;
+
+
+//Tracks which SimulationIDs we've successfully traced NetGUIDs of
+TArray<uint32> TracedSimulationNetGUIDs;
+TMap<uint32, TWeakObjectPtr<const AActor>> OwningActorMap;
+
+// ---------------------------------------------------------------------------
+
+FStringOutputDevice FNetworkPredictionTrace::StrOut;
+
+void FNetworkPredictionTrace::PushSimulationId(uint32 SimulationId)
+{
+	SimulationIdStack.Push(SimulationId);
+}
+
+void FNetworkPredictionTrace::PopSimulationId()
+{
+	SimulationIdStack.Pop();
+}
+
+void FNetworkPredictionTrace::TraceSimulationCreated(const AActor* OwningActor, const FName& GroupName)
+{
+	const uint32 SimulationId = PeakSimulationIdChecked();
+
+	uint32 GameInstanceId = GameInstanceMap.GetId(OwningActor->GetGameInstance());
+	FString DisplayNameStr = FString::Printf(TEXT("%s, %s"), *OwningActor->GetName(), *GroupName.ToString());
+
+	const uint16 AttachmentSize = (uint16)((DisplayNameStr.Len()+1) * sizeof(TCHAR));
+
+	OwningActorMap.Add(SimulationId, OwningActor);
+
+	UE_TRACE_LOG(NetworkPrediction, SimulationCreated, NetworkPredictionChannel, AttachmentSize)
+		<< SimulationCreated.SimulationId(SimulationId)
+		<< SimulationCreated.GameInstanceId(GameInstanceId)
+		<< SimulationCreated.Attachment(*DisplayNameStr, AttachmentSize);
+}
+
+void FNetworkPredictionTrace::TraceSimulationNetRole(uint32 SimulationId, ENetRole NetRole)
+{
+	UE_TRACE_LOG(NetworkPrediction, SimulationNetRole, NetworkPredictionChannel)
+		<< SimulationNetRole.EngineFrameNumber(GFrameNumber)
+		<< SimulationNetRole.SimulationId(SimulationId)
+		<< SimulationNetRole.NetRole((uint8)NetRole);
+}
+
+void FNetworkPredictionTrace::TraceSimulationTick(int32 OutputFrame, const FNetworkSimTime& FrameDeltaTime, const FSimulationTickState& TickState)
+{
+	UE_TRACE_LOG(NetworkPrediction, SimulationTick, NetworkPredictionChannel)
+		<< SimulationTick.SimulationId(PeakSimulationIdChecked())
+		<< SimulationTick.EngineFrameNumber(GFrameNumber)
+		<< SimulationTick.StartMS(TickState.GetTotalProcessedSimulationTime())
+		<< SimulationTick.EndMS(TickState.GetTotalProcessedSimulationTime() + FrameDeltaTime)
+		<< SimulationTick.OutputFrame(OutputFrame);
+}
+
+void FNetworkPredictionTrace::TraceSimulationEOF(const FSimulationTickState& TickState, int32 LastSentKeyframe, int32 LastReceivedKeyframe)
+{
+	const uint32 SimulationId = PeakSimulationIdChecked();
+
+	UE_TRACE_LOG(NetworkPrediction, SimulationEOF, NetworkPredictionChannel)
+		<< SimulationEOF.SimulationId(SimulationId)
+		<< SimulationEOF.EngineFrameNumber(GFrameNumber)
+		<< SimulationEOF.EngineFrameDeltaTime(FApp::GetDeltaTime())	// FIXME
+		<< SimulationEOF.EngineCurrentTime(FApp::GetCurrentTime())	// FIXME
+		<< SimulationEOF.TotalProcessedMS(TickState.GetTotalProcessedSimulationTime())
+		<< SimulationEOF.TotalAllowedMS(TickState.GetTotalAllowedSimulationTime())
+		<< SimulationEOF.LastSentKeyframe(LastSentKeyframe)
+		<< SimulationEOF.LastReceivedKeyframe(LastReceivedKeyframe);
+	
+
+
+	// Do we need to trace the NetGUID? This stinks having to be here
+	if (TracedSimulationNetGUIDs.Contains(SimulationId) == false)
+	{
+		uint32 NetGUID = 0;
+		if (const AActor* OwningActor = OwningActorMap.FindChecked(SimulationId).Get())
+		{
+			if (UWorld* World = OwningActor->GetWorld())
+			{
+				if (UNetDriver* NetDriver = World->GetNetDriver())
+				{
+					if (NetDriver->GuidCache.IsValid())
+					{
+						NetGUID = NetDriver->GuidCache->GetNetGUID(OwningActor).Value;
+					}
+				}
+			}
+		}
+
+		if (NetGUID > 1)
+		{
+			UE_TRACE_LOG(NetworkPrediction, SimulationNetGUID, NetworkPredictionChannel)
+				<< SimulationNetGUID.SimulationId(SimulationId)
+				<< SimulationNetGUID.NetGUID(NetGUID);
+		}
+	}
+}
+
+void FNetworkPredictionTrace::TraceNetSerializeRecv(const FNetworkSimTime& ReceivedTime, int32 ReceivedFrame)
+{
+	const uint32 SimulationId = PeakSimulationIdChecked();
+
+	UE_TRACE_LOG(NetworkPrediction, NetSerializeRecv, NetworkPredictionChannel)
+		<< NetSerializeRecv.SimulationId(SimulationId)
+		<< NetSerializeRecv.EngineFrameNumber(GFrameNumber)
+		<< NetSerializeRecv.ReceivedSimTimeMS(ReceivedTime)
+		<< NetSerializeRecv.ReceivedFrame(ReceivedFrame);
+}
+
+void FNetworkPredictionTrace::TraceNetSerializeCommit()
+{
+	const uint32 SimulationId = PeakSimulationIdChecked();
+
+	UE_TRACE_LOG(NetworkPrediction, NetSerializeCommit, NetworkPredictionChannel)
+		<< NetSerializeCommit.SimulationId(SimulationId);
+}
+
+void FNetworkPredictionTrace::TraceNetSerializeFault()
+{
+	const uint32 SimulationId = PeakSimulationIdChecked();
+
+	UE_TRACE_LOG(NetworkPrediction, NetSerializeFault, NetworkPredictionChannel)
+		<< NetSerializeFault.SimulationId(SimulationId);
+}
+
+void FNetworkPredictionTrace::TraceProduceInput()
+{
+	const uint32 SimulationId = PeakSimulationIdChecked();
+
+	UE_TRACE_LOG(NetworkPrediction, ProduceInput, NetworkPredictionChannel)
+		<< ProduceInput.SimulationId(SimulationId);
+}
+
+void FNetworkPredictionTrace::TraceSynthInput()
+{
+	const uint32 SimulationId = PeakSimulationIdChecked();
+
+	UE_TRACE_LOG(NetworkPrediction, SynthInput, NetworkPredictionChannel)
+		<< SynthInput.SimulationId(SimulationId);
+}
+
+void FNetworkPredictionTrace::TraceOOBStateMod()
+{
+	const uint32 SimulationId = PeakSimulationIdChecked();
+
+	UE_TRACE_LOG(NetworkPrediction, OOBStateMod, NetworkPredictionChannel)
+		<< OOBStateMod.SimulationId(SimulationId);
+}
+
+void FNetworkPredictionTrace::TraceUserState_Internal(int32 Frame, ETraceUserState StateType)
+{
+	const uint32 SimulationId = PeakSimulationIdChecked();
+	const uint16 AttachmentSize = (uint16)((StrOut.Len()+1) * sizeof(TCHAR));
+
+	switch(StateType)
+	{
+		case ETraceUserState::Input:
+		{
+			UE_TRACE_LOG(NetworkPrediction, InputCmd, NetworkPredictionChannel, AttachmentSize)
+				<< InputCmd.SimulationId(SimulationId)
+				<< InputCmd.SimulationFrame(Frame)
+				<< InputCmd.EngineFrameNumber(GFrameNumber)
+				<< InputCmd.Attachment(*StrOut, AttachmentSize);
+			break;
+		}
+		case ETraceUserState::Sync:
+		{
+			UE_TRACE_LOG(NetworkPrediction, SyncState, NetworkPredictionChannel, AttachmentSize)
+				<< SyncState.SimulationId(SimulationId)
+				<< SyncState.SimulationFrame(Frame)
+				<< SyncState.EngineFrameNumber(GFrameNumber)
+				<< SyncState.Attachment(*StrOut, AttachmentSize);
+			break;
+		}
+		case ETraceUserState::Aux:
+		{
+			UE_TRACE_LOG(NetworkPrediction, AuxState, NetworkPredictionChannel, AttachmentSize)
+				<< AuxState.SimulationId(SimulationId)
+				<< AuxState.SimulationFrame(Frame)
+				<< AuxState.EngineFrameNumber(GFrameNumber)
+				<< AuxState.Attachment(*StrOut, AttachmentSize);
+			break;
+		}
+	}
+}
+
+void FNetworkPredictionTrace::TracePIEStart()
+{
+	UE_TRACE_LOG(NetworkPrediction, PieBegin, NetworkPredictionChannel);
+}
\ No newline at end of file
diff --git a/Engine/Plugins/Runtime/NetworkPrediction/Source/NetworkPrediction/Public/Movement/FlyingMovement.h b/Engine/Plugins/Runtime/NetworkPrediction/Source/NetworkPrediction/Public/Movement/FlyingMovement.h
--- a/Engine/Plugins/Runtime/NetworkPrediction/Source/NetworkPrediction/Public/Movement/FlyingMovement.h
+++ b/Engine/Plugins/Runtime/NetworkPrediction/Source/NetworkPrediction/Public/Movement/FlyingMovement.h
@@ -111,11 +111,9 @@ struct FFlyingMovementAuxState
 
 	void Log(FStandardLoggingParameters& P) const
 	{
-		if (P.Context == EStandardLoggingContext::HeaderOnly)
-		{
-			P.Ar->Logf(TEXT(" %d "), P.Frame);
-		}
-		else if (P.Context == EStandardLoggingContext::Full)
+		P.Ar->Logf(TEXT("Frame: %d "), P.Frame);
+		
+		if (P.Context == EStandardLoggingContext::Full)
 		{
 			P.Ar->Logf(TEXT("MaxSpeed: %.2f"), MaxSpeed);
 			P.Ar->Logf(TEXT("TurningBoost: %.2f"), TurningBoost);
diff --git a/Engine/Plugins/Runtime/NetworkPrediction/Source/NetworkPrediction/Public/Movement/ParametricMovement.h b/Engine/Plugins/Runtime/NetworkPrediction/Source/NetworkPrediction/Public/Movement/ParametricMovement.h
--- a/Engine/Plugins/Runtime/NetworkPrediction/Source/NetworkPrediction/Public/Movement/ParametricMovement.h
+++ b/Engine/Plugins/Runtime/NetworkPrediction/Source/NetworkPrediction/Public/Movement/ParametricMovement.h
@@ -19,13 +19,13 @@ struct FSimpleParametricMotion
 	// Advance parametric time. This is meant to do simple things like looping/reversing etc.
 	void AdvanceParametricTime(const float InPosition, const float InPlayRate, float &OutPosition, float& OutPlayRate, const float DeltaTimeSeconds) const;
 
-	UPROPERTY(EditAnywhere, BlueprintReadOnly, Category=ParametricMovement)
+	UPROPERTY(EditAnywhere, BlueprintReadWrite, Category=ParametricMovement)
 	FVector ParametricDelta = FVector(0.f, 0.f, 500.f);
 
-	UPROPERTY(EditAnywhere, BlueprintReadOnly, Category=ParametricMovement)
+	UPROPERTY(EditAnywhere, BlueprintReadWrite, Category=ParametricMovement)
 	float MinTime = -1.f;
 
-	UPROPERTY(EditAnywhere, BlueprintReadOnly, Category=ParametricMovement)
+	UPROPERTY(EditAnywhere, BlueprintReadWrite, Category=ParametricMovement)
 	float MaxTime = 1.f;
 
 	FTransform CachedStartingTransform;
@@ -48,16 +48,18 @@ struct FParametricInputCmd
 
 	void Log(FStandardLoggingParameters& P) const
 	{
-		if (P.Context == EStandardLoggingContext::HeaderOnly)
-		{
-			P.Ar->Logf(TEXT(" %d "), P.Frame);
-		}
-		else if (P.Context == EStandardLoggingContext::Full)
+		P.Ar->Logf(TEXT("Frame: %d "), P.Frame);
+		
+		if (P.Context == EStandardLoggingContext::Full)
 		{
 			if (PlayRate.IsSet())
 			{
 				P.Ar->Logf(TEXT("PlayRate: %.2f"), PlayRate.GetValue());
 			}
+			else
+			{
+				P.Ar->Logf(TEXT("PlayRate: Unset"));
+			}
 		}
 	}
 };
@@ -76,11 +78,8 @@ struct FParametricSyncState
 
 	void Log(FStandardLoggingParameters& Params) const
 	{
-		if (Params.Context == EStandardLoggingContext::HeaderOnly)
-		{
-			Params.Ar->Logf(TEXT(" %d "), Params.Frame);
-		}
-		else if (Params.Context == EStandardLoggingContext::Full)
+		Params.Ar->Logf(TEXT("Frame: %d "), Params.Frame);
+		if (Params.Context == EStandardLoggingContext::Full)
 		{
 			Params.Ar->Logf(TEXT("Frame: %d"), Params.Frame);
 			Params.Ar->Logf(TEXT("Pos: %.2f"), Position);
@@ -101,11 +100,8 @@ struct FParametricAuxState
 
 	void Log(FStandardLoggingParameters& Params) const
 	{
-		if (Params.Context == EStandardLoggingContext::HeaderOnly)
-		{
-			Params.Ar->Logf(TEXT(" %d "), Params.Frame);
-		}
-		else if (Params.Context == EStandardLoggingContext::Full)
+		Params.Ar->Logf(TEXT("Frame %d "), Params.Frame);
+		if (Params.Context == EStandardLoggingContext::Full)
 		{
 			Params.Ar->Logf(TEXT("Multiplier: %f"), Multiplier);
 		}
@@ -184,6 +180,9 @@ class NETWORKPREDICTION_API UParametricMovementComponent : public UBaseMovementC
 	void ProduceInput(const FNetworkSimTime SimTime, FParametricInputCmd& Cmd);
 	void FinalizeFrame(const FParametricSyncState& SyncState, const FParametricAuxState& AuxState) override;
 
+	UFUNCTION(BlueprintCallable, Category="Networking")
+	void EnableInterpolationMode(bool bValue);
+
 protected:
 
 	TNetSimStateAccessor<FParametricSyncState> MovementSyncState;
@@ -222,7 +221,7 @@ protected:
 	UPROPERTY(EditAnywhere, BlueprintReadOnly, Category=ParametricMovement)
 	bool bDisableParametricMovementSimulation = false;
 
-	UPROPERTY(EditAnywhere, BlueprintReadOnly, Category=ParametricMovement)
+	UPROPERTY(EditAnywhere, BlueprintReadWrite, Category=ParametricMovement, meta=(ExposeOnSpawn=true))
 	FSimpleParametricMotion ParametricMotion;
 
 	UPROPERTY(EditAnywhere, BlueprintReadOnly, Category=ParametricMovementNetworking)
diff --git a/Engine/Plugins/Runtime/NetworkPrediction/Source/NetworkPrediction/Public/NetworkPredictionTypes.h b/Engine/Plugins/Runtime/NetworkPrediction/Source/NetworkPrediction/Public/NetworkPredictionTypes.h
--- a/Engine/Plugins/Runtime/NetworkPrediction/Source/NetworkPrediction/Public/NetworkPredictionTypes.h
+++ b/Engine/Plugins/Runtime/NetworkPrediction/Source/NetworkPrediction/Public/NetworkPredictionTypes.h
@@ -236,8 +236,12 @@ public:
 	virtual bool ShouldSendServerRPC(float DeltaSeconds) = 0;
 	virtual void SetDesiredServerRPCSendFrequency(float DesiredHz) = 0;
 
+	// Enables simulation extrapolation (run sim code to extrapolate state outside network updates). False = interpolation mode.
+	// Note dependent simulation mode will take precedence over this
+	virtual void SetEnableSimulationExtrapolation(bool bNewValue) = 0;
+
 	// ----------------------------------------------------------------------
-	// Functions for depedent simulation (forward predicting a simulated proxy sim along with an auto proxy sim)
+	// Functions for dependent simulation (forward predicting a simulated proxy sim along with an auto proxy sim)
 	// ----------------------------------------------------------------------
 	
 	// Main function to call on simulated proxy sim
@@ -255,6 +259,11 @@ public:
 	virtual void StepRollback(const struct FNetworkSimTime& Step, const int32 ParentFrame, const bool FinalStep) = 0;
 
 	void ProcessPendingNetSimCues() { ProcessPendingNetSimCuesFunc(); }
+	   	
+	// Unique, global, per process Id of simulation. Used for trace debugging purposes only. Not replicated.
+	virtual uint32 GetSimulationId() const = 0;
+
+	static uint32 NETWORKPREDICTION_API GetNextSimulationId();
 
 protected:
 	TFunction<void()> ProcessPendingNetSimCuesFunc;
diff --git a/Engine/Plugins/Runtime/NetworkPrediction/Source/NetworkPrediction/Public/NetworkedSimulationModel.h b/Engine/Plugins/Runtime/NetworkPrediction/Source/NetworkPrediction/Public/NetworkedSimulationModel.h
--- a/Engine/Plugins/Runtime/NetworkPrediction/Source/NetworkPrediction/Public/NetworkedSimulationModel.h
+++ b/Engine/Plugins/Runtime/NetworkPrediction/Source/NetworkPrediction/Public/NetworkedSimulationModel.h
@@ -88,9 +88,9 @@ struct TNetSimModelRepControllerTraitsBase
 // Template struct to specialize per-model
 template<typename T> struct TNetSimModelRepControllerTraits : public TNetSimModelRepControllerTraitsBase { };
 
-// The actual NetworkedSimulationModel templated class.
+// The actual NetworkedSimulationModel templated class. InModelDef will be a user defined subtype of FNetSimModelDefBase
 template <typename InModelDef>
-class TNetworkedSimulationModel : public InModelDef::Base
+class TNetworkedSimulationModel : public InModelDef::Base // InModelDef::Base defaults to INetworkedSimulationModel
 {
 public:
 
@@ -144,6 +144,13 @@ public:
 		};
 
 		this->Buffers.CueDispatcher.GetDebugName = [this]() { return Driver->GetDebugName(); };
+		
+		SimulationId = this->GetNextSimulationId();
+		UE_NP_TRACE_SET_SCOPE_SIM(SimulationId);
+		UE_NP_TRACE_SIM_CREATED(Driver->GetVLogOwner(), GetSimulationGroupName());
+		UE_NP_TRACE_OOB_STATE_MOD();
+		UE_NP_TRACE_USER_STATE_SYNC(InitialSyncState, 0);
+		UE_NP_TRACE_USER_STATE_AUX(InitialAuxState, 0);
 	}
 
 	virtual ~TNetworkedSimulationModel()
@@ -154,6 +161,8 @@ public:
 
 	void Tick(const FNetSimTickParameters& Parameters) final override
 	{
+		UE_NP_TRACE_SET_SCOPE_SIM(SimulationId);
+
 		// Update previous DebugState based on what we (might) have sent *after* our last Tick 
 		// (property replication and ServerRPC get sent after the tick, rather than forcing awkward callback into the NetSim post replication, we can check it here)
 		if (auto* DebugBuffer = GetLocalDebugBuffer())
@@ -236,23 +245,9 @@ public:
 					InSyncState = Buffers.Sync.WriteFrameInitializedFromHead(InputFrame);
 				}
 
-				const TAuxState* InAuxState = Buffers.Aux[InputFrame];
-				checkf(InAuxState, TEXT("No AuxState available for Frame %d. PendingFrame: %d. MaxAllowedFrame: %d."), InputFrame, Ticker.PendingFrame, Ticker.MaxAllowedFrame);
-
-				TSyncState* OutSyncState = Buffers.Sync.WriteFrame(OutputFrame);
-				check(OutSyncState);
-
-				{
-					ESimulationTickContext Context = Parameters.Role == ROLE_Authority ? ESimulationTickContext::Authority : ESimulationTickContext::Predict;
-					TScopedSimulationTick<Model> UpdateScope(Ticker, Buffers.CueDispatcher, Context, OutputFrame, InputCmd->GetFrameDeltaTime());
-					
-					Simulation->SimulationTick( 
-						{ InputCmd->GetFrameDeltaTime(), Ticker },
-						{ *InputCmd, *InSyncState, *InAuxState },
-						{ *OutSyncState, Buffers.Aux.LazyWriter(OutputFrame), Buffers.CueDispatcher } );
-													
-				}
-
+				ESimulationTickContext Context = Parameters.Role == ROLE_Authority ? ESimulationTickContext::Authority : ESimulationTickContext::Predict;
+				TSimulationDoTickFunc<Model>::DoTick(*Simulation.Get(), Ticker, Buffers, OutputFrame, Context);
+				
 				if (DebugState)
 				{
 					DebugState->ProcessedFrames.Add(InputFrame);
@@ -302,10 +297,33 @@ public:
 			HistoricData->Sync.CopyAndMerge(Buffers.Sync);
 			HistoricData->Aux.CopyAndMerge(Buffers.Aux);
 		}
+		
+		int32 LastSentKeyframe = 0;
+		int32 LastReceivedKeyframe = 0;
+		switch (Parameters.Role)
+		{
+			case ROLE_Authority:
+				LastSentKeyframe = RepProxy_Autonomous.GetLastSerializedFrame();
+				LastReceivedKeyframe = RepProxy_ServerRPC.GetLastSerializedFrame();
+			break;
+
+			case ROLE_AutonomousProxy:
+				LastSentKeyframe = RepProxy_ServerRPC.GetLastSerializedFrame();
+				LastReceivedKeyframe = RepProxy_Autonomous.GetLastSerializedFrame();
+			break;
+
+			case ROLE_SimulatedProxy:
+				LastReceivedKeyframe = RepProxy_Simulated.GetLastSerializedFrame();
+			break;
+		}
+
+		UE_NP_TRACE_SIM_EOF(Ticker, LastSentKeyframe, LastReceivedKeyframe);
 	}
 
 	virtual void Reconcile(const ENetRole Role) final override
 	{
+		UE_NP_TRACE_SET_SCOPE_SIM(SimulationId);
+
 		// --------------------------------------------------------------------------------------------------------------------------
 		//	Reconcile
 		//	This will eventually be called outside the Tick loop, only after processing a network bunch
@@ -329,6 +347,8 @@ public:
 	
 	void InitializeForNetworkRole(const ENetRole Role, const FNetworkSimulationModelInitParameters& Parameters) final override
 	{
+		UE_NP_TRACE_NETROLE(SimulationId, Role);
+
 		// FIXME: buffer sizes are now inlined allocated but we want to support role based buffer sizes
 
 		//Buffers.Input.SetBufferSize(Parameters.InputBufferSize);
@@ -352,6 +372,8 @@ public:
 
 	void NetSerializeProxy(EReplicationProxyTarget Target, const FNetSerializeParams& Params) final override
 	{
+		UE_NP_TRACE_SET_SCOPE_SIM(SimulationId);
+
 		switch(Target)
 		{
 		case EReplicationProxyTarget::ServerRPC:
@@ -443,6 +465,11 @@ public:
 		return nullptr;
 	}
 
+	void SetEnableSimulationExtrapolation(bool bNewValue) final override
+	{
+		RepProxy_Simulated.bAllowSimulatedExtrapolation = bNewValue;
+	}
+
 	// ------------------------------------------------------------------------------------------------------
 	//	Dependent Simulations
 	// ------------------------------------------------------------------------------------------------------
@@ -485,11 +512,13 @@ public:
 
 	void BeginRollback(const FNetworkSimTime& RollbackDeltaTime, const int32 ParentFrame) final override
 	{
+		UE_NP_TRACE_SET_SCOPE_SIM(SimulationId);
 		RepProxy_Simulated.DependentRollbackBegin(Simulation.Get(), Driver, Buffers, Ticker, RollbackDeltaTime, ParentFrame);
 	}
 
 	void StepRollback(const FNetworkSimTime& Step, const int32 ParentFrame, const bool bFinalStep) final override
 	{
+		UE_NP_TRACE_SET_SCOPE_SIM(SimulationId);
 		RepProxy_Simulated.DependentRollbackStep(Simulation.Get(), Driver, Buffers, Ticker, Step, ParentFrame, bFinalStep);
 	}
 
@@ -546,10 +575,15 @@ public:
 	// Simulation class should have a static const FName GroupName member
 	FName GetSimulationGroupName() const final override { return Model::GroupName; }
 
+	// Unique, global, per process Id of simulation. Used for trace debugging purposes only. Not replicated.
+	uint32 GetSimulationId() const final override { return SimulationId; }
+
 private:
 	float ServerRPCAccumulatedTimeSeconds = 0.f;
 	float ServerRPCThresholdTimeSeconds = 1.f / 999.f; // Default is to send at a max of 999hz. This part of the system needs to be build out more (better handling of super high FPS clients and fixed rate servers)
 
+	uint32 SimulationId = 0;
+
 	// ------------------------------------------------------------------
 	//	Debugging
 	// ------------------------------------------------------------------
diff --git a/Engine/Plugins/Runtime/NetworkPrediction/Source/NetworkPrediction/Public/NetworkedSimulationModelBuffer.h b/Engine/Plugins/Runtime/NetworkPrediction/Source/NetworkPrediction/Public/NetworkedSimulationModelBuffer.h
--- a/Engine/Plugins/Runtime/NetworkPrediction/Source/NetworkPrediction/Public/NetworkedSimulationModelBuffer.h
+++ b/Engine/Plugins/Runtime/NetworkPrediction/Source/NetworkPrediction/Public/NetworkedSimulationModelBuffer.h
@@ -1,6 +1,7 @@
 // Copyright Epic Games, Inc. All Rights Reserved.
 
 #pragma once
+#include "Containers/UnrealString.h"
 
 // Generic iterator for iterating over each frame from tail to head
 template<typename TBuffer, typename ElementType>
diff --git a/Engine/Plugins/Runtime/NetworkPrediction/Source/NetworkPrediction/Public/NetworkedSimulationModelRepControllers.h b/Engine/Plugins/Runtime/NetworkPrediction/Source/NetworkPrediction/Public/NetworkedSimulationModelRepControllers.h
--- a/Engine/Plugins/Runtime/NetworkPrediction/Source/NetworkPrediction/Public/NetworkedSimulationModelRepControllers.h
+++ b/Engine/Plugins/Runtime/NetworkPrediction/Source/NetworkPrediction/Public/NetworkedSimulationModelRepControllers.h
@@ -78,8 +78,14 @@ void GenerateLocalInputCmdAtFrame(typename TNetSimModelTraits<Model>::Driver* Dr
 		InputCmd->SetFrameDeltaTime(DeltaSimTime);
 		if (bProduceInputViaDriver)
 		{
+			UE_NP_TRACE_PRODUCE_INPUT();
 			Driver->ProduceInput(DeltaSimTime, *InputCmd);
 		}
+		else
+		{
+			UE_NP_TRACE_SYNTH_INPUT();
+		}
+		UE_NP_TRACE_USER_STATE_INPUT(*InputCmd, Frame);
 	}
 }
 
@@ -228,19 +234,19 @@ struct TRepController_Sequence : public TBase
 		Ar << SerializedNumElements;
 
 		const int32 HeadFrame = FNetworkSimulationSerialization::SerializeFrame(Ar, Buffer.HeadFrame());
-		const int32 StartingFrame = FMath::Max(0, HeadFrame - SerializedNumElements + 1);
+		LastSerializedFrameStart = FMath::Max(0, HeadFrame - SerializedNumElements + 1);
 
 		if (Ar.IsLoading())
 		{
 			const int32 PrevHead = Buffer.HeadFrame();
-			if (PrevHead < StartingFrame && PrevHead >= 0)
+			if (PrevHead < LastSerializedFrameStart && PrevHead >= 0)
 			{
 				// There is a gap in the stream. In some cases, we want this to be a "fault" and bubble up. We may want to synthesize state or maybe we just skip ahead.
-				UE_LOG(LogNetworkSim, Warning, TEXT("Fault: gap in received %s buffer. PrevHead: %d. Received: %d-%d. Reseting previous buffer contents"), *LexToString(BufferId), PrevHead, StartingFrame, HeadFrame);
+				UE_LOG(LogNetworkSim, Warning, TEXT("Fault: gap in received %s buffer. PrevHead: %d. Received: %d-%d. Reseting previous buffer contents"), *LexToString(BufferId), PrevHead, LastSerializedFrameStart, HeadFrame);
 			}
 		}
 
-		for (int32 Frame = StartingFrame; Frame <= HeadFrame; ++Frame)
+		for (int32 Frame = LastSerializedFrameStart; Frame <= HeadFrame; ++Frame)
 		{
 			// This, as is, is bad. The intention is that these functions serialize multiple items in some delta compressed fashion.
 			// As is, we are just serializing the elements individually.
@@ -257,14 +263,15 @@ struct TRepController_Sequence : public TBase
 
 protected:
 
+	int32 LastSerializedFrameStart = 0;
 	int32 LastSerializedFrame = 0;
 };
 
 // -------------------------------------------------------------------------------------------------------
-//	Server Replication Controller (FIXME: should this be called something like "client autonomous proxy -> Server RPC" ?
+//	Server Replication Controller. FIXME: should this be called something like "client autonomous proxy -> Server RPC" ?
 //
 // -------------------------------------------------------------------------------------------------------
-template<typename Model, typename TBase=TRepController_Sequence<Model, ENetworkSimBufferTypeId::Input, 3>>
+template<typename Model, int32 NumSendPerUpdate=3, typename TBase=TRepController_Empty<Model>>
 struct TRepController_Server : public TBase
 {
 	using typename TBase::TSimulation;
@@ -272,7 +279,88 @@ struct TRepController_Server : public TBase
 	using typename TBase::TDriver;
 	using typename TBase::TBufferTypes;
 	using typename TBase::TInputCmd;
+
+	int32 GetProxyDirtyCount(TNetworkSimBufferContainer<Model>& Buffers) const 
+	{
+		return Buffers.Input.GetDirtyCount();
+	}
+	// -------------------------------------------------------------------------------------------------
+	// TRepController_Server::NetSerialize
+	//	This is the client sending input data to the server
+	// -------------------------------------------------------------------------------------------------
+	void NetSerialize(const FNetSerializeParams& P, TNetworkSimBufferContainer<Model>& Buffers, TSimulationTicker<TTickSettings>& Ticker)
+	{
+		NETSIM_CHECKSUM(P.Ar);
+
+		auto& InputBuffer = Buffers.Input;
+
+		FArchive& Ar = P.Ar;
+		uint8 SerializedNumElements = FMath::Min<uint8>(NumSendPerUpdate, InputBuffer.Num());
+		Ar << SerializedNumElements;
+
+		const int32 HeadFrame = FNetworkSimulationSerialization::SerializeFrame(Ar, InputBuffer.HeadFrame());
+		const int32 StartFrame = FMath::Max(0, HeadFrame - SerializedNumElements + 1);
+
+		// Serialize Input Cmds redundantly. This could be better:
+		//	-Sending: We could delta compress the data that doesn't change in subsequent cmds
+		//	-Receiving: we could skip cmds that we already have (right now we must NetSerialize redundantly, we can't skip ahead in the stream)
+		
+		if (Ar.IsSaving())
+		{
+			for (int32 Frame = StartFrame; Frame <= HeadFrame; ++Frame)
+			{
+				InputBuffer[Frame]->NetSerialize(P);
+			}
+		}
+		else
+		{
+			const int32 ExistingInputHeadFrame = InputBuffer.HeadFrame(); // Anything we receive before this frame is redundant
+			
+			// SimulationTimeBuffer[N] is the sim time after running frame N.
+			// Sync[N]/Aux[N] is the output for frame N and input to frame N+1.
+			// Input[N] is the input cmd used in the sim tick that will generate frame N+1.			
+			// Therefor, Input buffer can be no more than 1 behind the simulation SImulationTimeBuffer.
+			ensureMsgf(ExistingInputHeadFrame+1 >= Ticker.SimulationTimeBuffer.HeadFrame(), TEXT("Simulation Ahead of Input Buffer: %d >= %d"), ExistingInputHeadFrame, Ticker.SimulationTimeBuffer.HeadFrame());
+
+			// Figure out what the sim time will be at input head (local input buffer may be ahead of processed simulation times)
+			FNetworkSimTime SimTime = Ticker.GetTotalProcessedSimulationTime();
+			for (int32 TimeFrame = Ticker.SimulationTimeBuffer.HeadFrame(); TimeFrame <= ExistingInputHeadFrame; ++TimeFrame)
+			{
+				//if (ensure(InputBuffer.IsValidFrame(TimeFrame))) // ensure
+				if (InputBuffer.IsValidFrame(TimeFrame)) // ensure
+				{
+					SimTime += InputBuffer[TimeFrame]->GetFrameDeltaTime();
+				}
+			}
+
+			for (int32 Frame = StartFrame; Frame <= HeadFrame; ++Frame)
+			{
+				auto* InputCmd = InputBuffer.WriteFrame(Frame);
+				InputCmd->NetSerialize(P);
+				
+				if (Frame > ExistingInputHeadFrame)
+				{
+					UE_NP_TRACE_NET_SERIALIZE_RECV(SimTime, Frame);
+					UE_NP_TRACE_USER_STATE_INPUT(*InputCmd, Frame);
+					UE_NP_TRACE_NET_SERIALIZE_COMMIT();
+
+					SimTime += InputCmd->GetFrameDeltaTime();
+				}
+				else
+				{
+					// Redundant cmd, don't trace the user state
+				}
+			}
+		}
+
+		LastSerializedFrame = HeadFrame;
+		NETSIM_CHECKSUM(P.Ar);
+	}
 	
+	// -------------------------------------------------------------------------------------------------
+	// TRepController_Server::Reconcile
+	//	This is the server reconciling the simulation state with what the client sent him
+	// -------------------------------------------------------------------------------------------------
 	void Reconcile(TSimulation* Simulation, TDriver* Driver, TNetworkSimBufferContainer<Model>& Buffers, TSimulationTicker<TTickSettings>& Ticker)
 	{
 		// After receiving input, server may process up to the latest received frames.
@@ -288,6 +376,10 @@ struct TRepController_Server : public TBase
 		}
 	}
 	
+	// -------------------------------------------------------------------------------------------------
+	// TRepController_Server::PreSimTick
+	//	This is the server deciding if he should generate local input for the simulation
+	// -------------------------------------------------------------------------------------------------
 	void PreSimTick(TDriver* Driver, TNetworkSimBufferContainer<Model>& Buffers, TSimulationTicker<TTickSettings>& Ticker, const FNetSimTickParameters& TickParameters)
 	{
 		Ticker.GiveSimulationTime(TickParameters.LocalDeltaTimeSeconds);
@@ -296,6 +388,13 @@ struct TRepController_Server : public TBase
 			TryGenerateLocalInput(Driver, Buffers, Ticker);
 		}
 	}
+
+	int32 GetLastSerializedFrame() const { return LastSerializedFrame; }
+
+private:
+	
+	int32 LastSerializedFrame = 0;
+
 };
 
 // Simulated: "non locally controlled" simulations. We support "Simulation Extrapolation" here (using the sim to fake inputs to advance the sim)
@@ -326,6 +425,7 @@ struct TRepController_Simulated : public TBase
 	const TSyncState& GetLastSerializedSyncState() const { return LastSerializedSyncState; }
 	const TAuxState& GetLastSerializedAuxState() const { return LastSerializedAuxState; }
 	const TInputCmd& GetLastSerializedInputCmd() const { return LastSerializedInputCmd; }
+	int32 GetLastSerializedFrame() const { return LastSerializedKeyframe; }
 	
 	ESimulatedUpdateMode GetSimulatedUpdateMode() const
 	{
@@ -374,6 +474,8 @@ struct TRepController_Simulated : public TBase
 			SyncState = Buffers.Sync.HeadElement();
 			AuxState = Buffers.Aux[Frame];
 			check(InputCmd && SyncState && AuxState	); // We should not be here if the buffer is empty. Want to avoid serializing an "empty" flag at the top here.
+
+			LastSerializedKeyframe = Frame;
 		}
 		else
 		{
@@ -387,11 +489,11 @@ struct TRepController_Simulated : public TBase
 			}
 
 			// Find out where this should go in the local buffer based on the serialized time
-			int32 DestinationFrame = INDEX_NONE;
+			LastSerializedKeyframe = INDEX_NONE;
 			if (LastSerializedSimulationTime > Ticker.GetTotalProcessedSimulationTime())
 			{
 				// We are getting new state that is ahead of what we have locally, so it can safety go right to head
-				DestinationFrame = Ticker.SimulationTimeBuffer.HeadFrame()+1;
+				LastSerializedKeyframe = Ticker.SimulationTimeBuffer.HeadFrame()+1;
 			}
 			else
 			{
@@ -400,12 +502,12 @@ struct TRepController_Simulated : public TBase
 				{
 					if (LastSerializedSimulationTime > *Ticker.SimulationTimeBuffer[Frame])
 					{
-						DestinationFrame = Frame+1;
+						LastSerializedKeyframe = Frame+1;
 						break;
 					}
 				}
 
-				if (DestinationFrame == INDEX_NONE)
+				if (LastSerializedKeyframe == INDEX_NONE)
 				{
 					FNetworkSimTime TotalTimeAhead = Ticker.GetTotalProcessedSimulationTime() - LastSerializedSimulationTime;
 					FNetworkSimTime SerializeDelta = LastSerializedSimulationTime - PrevLastSerializedSimulationTime;
@@ -413,23 +515,23 @@ struct TRepController_Simulated : public TBase
 					// We are way far ahead of the server... we will need to clear out sync buffers, take what they gave us, and catch up in reconcile
 					//UE_LOG(LogNetworkSim, Warning, TEXT("!!! TReplicator_Simulated. Large gap detected. SerializedTime: %s. Buffer time: [%s-%s]. %d Elements. DeltaFromHead: %s. DeltaSerialize: %s"), *LastSerializedSimulationTime.ToString(), 
 					//	*Ticker.SimulationTimeBuffer.GetElementFromTail(0)->ToString(), *Ticker.SimulationTimeBuffer.GetElementFromHead(0)->ToString(), Ticker.SimulationTimeBuffer.GetNumValidElements(), *TotalTimeAhead.ToString(), *SerializeDelta.ToString());
-					DestinationFrame = Ticker.SimulationTimeBuffer.HeadFrame()+2; // (Skip ahead 2 to force a break in continuity)
+					LastSerializedKeyframe = Ticker.SimulationTimeBuffer.HeadFrame()+2; // (Skip ahead 2 to force a break in continuity)
 				}
 			}
 
-			check(DestinationFrame != INDEX_NONE);
+			check(LastSerializedKeyframe != INDEX_NONE);
 
 			// "Finalize" our buffers and time keeping such that we serialize the latest state from the server in the right spot
-			InputCmd = Buffers.Input.WriteFrame(DestinationFrame);
-			SyncState = Buffers.Sync.WriteFrame(DestinationFrame);
-			AuxState = Buffers.Aux.WriteFrame(DestinationFrame);
+			InputCmd = Buffers.Input.WriteFrame(LastSerializedKeyframe);
+			SyncState = Buffers.Sync.WriteFrame(LastSerializedKeyframe);
+			AuxState = Buffers.Aux.WriteFrame(LastSerializedKeyframe);
 
-			ensure(Buffers.Input.HeadFrame() <= DestinationFrame);
-			ensure(Buffers.Sync.HeadFrame() <= DestinationFrame);
-			ensure(Buffers.Aux.HeadFrame() <= DestinationFrame);
+			ensure(Buffers.Input.HeadFrame() <= LastSerializedKeyframe);
+			ensure(Buffers.Sync.HeadFrame() <= LastSerializedKeyframe);
+			ensure(Buffers.Aux.HeadFrame() <= LastSerializedKeyframe);
 
 			// Update tick info
-			Ticker.SetTotalProcessedSimulationTime(LastSerializedSimulationTime, DestinationFrame);
+			Ticker.SetTotalProcessedSimulationTime(LastSerializedSimulationTime, LastSerializedKeyframe);
 			if (Ticker.GetTotalAllowedSimulationTime() < LastSerializedSimulationTime)
 			{
 				Ticker.SetTotalAllowedSimulationTime(LastSerializedSimulationTime);
@@ -437,8 +539,10 @@ struct TRepController_Simulated : public TBase
 
 			check(Ticker.GetTotalProcessedSimulationTime() <= Ticker.GetTotalAllowedSimulationTime());
 
-			Ticker.PendingFrame = DestinationFrame;	// We are about to serialize state to DestinationFrame which will be "unprocessed" (has not been used to generate a new frame)
-			Ticker.MaxAllowedFrame = DestinationFrame-1; // Do not process PendingFrame on our account. ::PreSimTick will advance this based on our interpolation/extrapolation settings
+			Ticker.PendingFrame = LastSerializedKeyframe;	// We are about to serialize state to LastSerializedKeyframe which will be "unprocessed" (has not been used to generate a new frame)
+			Ticker.MaxAllowedFrame = LastSerializedKeyframe-1; // Do not process PendingFrame on our account. ::PreSimTick will advance this based on our interpolation/extrapolation settings
+
+			UE_NP_TRACE_NET_SERIALIZE_RECV(LastSerializedSimulationTime, LastSerializedKeyframe);
 		}
 
 		check(SyncState && AuxState);
@@ -458,6 +562,11 @@ struct TRepController_Simulated : public TBase
 			LastSerializedInputCmd = *InputCmd;
 			LastSerializedSyncState = *SyncState;
 			LastSerializedAuxState = *AuxState;
+
+			UE_NP_TRACE_USER_STATE_INPUT(*InputCmd, LastSerializedKeyframe);
+			UE_NP_TRACE_USER_STATE_SYNC(*SyncState, LastSerializedKeyframe);
+			UE_NP_TRACE_USER_STATE_AUX(*AuxState, LastSerializedKeyframe);
+			
 		}
 
 		NETSIM_CHECKSUM(P.Ar);
@@ -470,6 +579,8 @@ struct TRepController_Simulated : public TBase
 			return;
 		}
 
+		UE_NP_TRACE_NET_SERIALIZE_COMMIT();
+
 		check(Ticker.GetTotalProcessedSimulationTime() <= Ticker.GetTotalAllowedSimulationTime());
 
 		if (GetSimulatedUpdateMode() == ESimulatedUpdateMode::Extrapolate && NetworkSimulationModelCVars::EnableSimulatedReconcile())
@@ -479,6 +590,7 @@ struct TRepController_Simulated : public TBase
 
 			// Simulated Reconcile requires the input buffer to be kept up to date with the Sync buffer
 			// Generate a new, fake, command since we just added a new sync state to head
+			UE_NP_TRACE_SYNTH_INPUT();
 			while (Buffers.Input.HeadFrame() < Buffers.Sync.HeadFrame())
 			{
 				TInputCmd* Next = Buffers.Input.WriteFrameInitializedFromHead(Buffers.Input.HeadFrame()+1);
@@ -588,9 +700,9 @@ private:
 		const int32 InputFrame = Buffers.Input.HeadFrame();
 		const int32 OutputFrame = InputFrame + 1;
 
-		// Create fake cmd				
-		TInputCmd* NewCmd = Buffers.Input.WriteFrameInitializedFromHead(OutputFrame);
-		NewCmd->SetFrameDeltaTime(DeltaSimTime);	
+		// Override the FrameDeltaTime on our InputCmd.
+		TInputCmd* InputCmd = Buffers.Input.HeadElement();
+		InputCmd->SetFrameDeltaTime(DeltaSimTime);
 		
 		// Create new sync state to write to
 		TSyncState* PrevSyncState = Buffers.Sync[InputFrame];
@@ -600,26 +712,22 @@ private:
 		if (NETSIM_MODEL_DEBUG)
 		{
 			FVisualLoggingParameters VLogParameters( EVisualLoggingContext::OtherPredicted, InputFrame, EVisualLoggingLifetime::Persistent, TEXT("Pre SimulationExtrapolation"));
-			Driver->InvokeVisualLog(NewCmd, PrevSyncState, AuxState, VLogParameters);
+			Driver->InvokeVisualLog(InputCmd, PrevSyncState, AuxState, VLogParameters);
 		}
 
 		// Do the actual update
-		{
-			TScopedSimulationTick<Model> UpdateScope(Ticker, Buffers.CueDispatcher, ESimulationTickContext::Resimulate, OutputFrame, NewCmd->GetFrameDeltaTime());
-			Simulation->SimulationTick( 
-				{ NewCmd->GetFrameDeltaTime(), Ticker },
-				{ *NewCmd, *PrevSyncState, *AuxState },
-				{ *NextSyncState, Buffers.Aux.LazyWriter(OutputFrame), Buffers.CueDispatcher } );
-		}
+		TSimulationDoTickFunc<Model>::DoTick(*Simulation, Ticker, Buffers, OutputFrame, ESimulationTickContext::Resimulate);
 
 		if (NETSIM_MODEL_DEBUG)
 		{
 			FVisualLoggingParameters VLogParameters( EVisualLoggingContext::OtherPredicted, InputFrame, EVisualLoggingLifetime::Persistent, TEXT("Post SimulationExtrapolation"));
-			Driver->InvokeVisualLog(NewCmd, NextSyncState, Buffers.Aux[OutputFrame], VLogParameters);
+			Driver->InvokeVisualLog(InputCmd, NextSyncState, Buffers.Aux[OutputFrame], VLogParameters);
 		}
 
 		Ticker.MaxAllowedFrame = OutputFrame;
 	}
+
+	int32 LastSerializedKeyframe;
 	
 	FNetworkSimTime ReconcileSimulationTime;
 	FNetworkSimTime LastSerializedSimulationTime;
@@ -682,12 +790,25 @@ struct TRepController_Autonomous: public TBase
 
 		if (Ar.IsLoading())
 		{
+			UE_NP_TRACE_NET_SERIALIZE_RECV(SerializedTime, SerializedFrame);
+			UE_NP_TRACE_USER_STATE_SYNC(SerializedSyncState, SerializedFrame);
+			UE_NP_TRACE_USER_STATE_AUX(SerializedAuxState, SerializedFrame);
+
 			bReconcileFaultDetected = false;
 			bPendingReconciliation = false;
 
 			TSyncState* ClientSync = Buffers.Sync[SerializedFrame];
 			TAuxState* ClientAux = Buffers.Aux[SerializedFrame];
 
+			if (Ticker.SimulationTimeBuffer.IsValidFrame(SerializedFrame))
+			{
+				FNetworkSimTime ClientSimTime = *Ticker.SimulationTimeBuffer[SerializedFrame];
+				if (SerializedTime != ClientSimTime)
+				{
+					bPendingReconciliation = true;
+				}
+			}
+
 			// The state the client predicted that corresponds to the state the server just serialized to us
 			if (ClientSync && ClientAux)
 			{
@@ -707,6 +828,7 @@ struct TRepController_Autonomous: public TBase
 					// Case 1: the serialized state is older than what we've kept in our buffer. A bigger buffer would solve this! (at the price of more resimulated frames to recover when this happens)
 					// This is a reconcile fault and we just need to chill. We'll stop sending user commands until the cmds in flight flush through the system and we catch back up.
 					bReconcileFaultDetected = true;
+					UE_NP_TRACE_NET_SERIALIZE_FAULT();
 				}
 				else
 				{
@@ -754,6 +876,8 @@ struct TRepController_Autonomous: public TBase
 			ClientSyncState = Buffers.Sync.WriteFrame( SerializedFrame );
 		}
 
+		UE_NP_TRACE_NET_SERIALIZE_COMMIT();
+
 		// Set client's sync state to the server version
 		check(ClientSyncState);
 		*ClientSyncState = SerializedSyncState;
@@ -808,13 +932,7 @@ struct TRepController_Autonomous: public TBase
 			}
 
 			// Do the actual update
-			{
-				TScopedSimulationTick<Model> UpdateScope(Ticker, Buffers.CueDispatcher, ESimulationTickContext::Resimulate, OutputFrame, ResimulateCmd->GetFrameDeltaTime());
-				Simulation->SimulationTick( 
-					{ ResimulateCmd->GetFrameDeltaTime(), Ticker },
-					{ *ResimulateCmd, *PrevSyncState, *AuxState },
-					{ *NextSyncState, Buffers.Aux.LazyWriter(OutputFrame), Buffers.CueDispatcher } );
-			}
+			TSimulationDoTickFunc<Model>::DoTick(*Simulation, Ticker, Buffers, OutputFrame, ESimulationTickContext::Resimulate);
 
 			// Log out the newly predicted state that we got.
 			if (NETSIM_MODEL_DEBUG)
diff --git a/Engine/Plugins/Runtime/NetworkPrediction/Source/NetworkPrediction/Public/NetworkedSimulationModelTypes.h b/Engine/Plugins/Runtime/NetworkPrediction/Source/NetworkPrediction/Public/NetworkedSimulationModelTypes.h
--- a/Engine/Plugins/Runtime/NetworkPrediction/Source/NetworkPrediction/Public/NetworkedSimulationModelTypes.h
+++ b/Engine/Plugins/Runtime/NetworkPrediction/Source/NetworkPrediction/Public/NetworkedSimulationModelTypes.h
@@ -6,6 +6,7 @@
 #include "NetworkedSimulationModelCues.h"
 #include "NetworkedSimulationModelTime.h"
 #include "NetworkedSimulationModelTraits.h"
+#include "Trace/NetworkPredictionTrace.h"
 
 // ---------------------------------------------------------------------------------------------------------------------
 //	TNetworkSimBufferContainer
@@ -205,7 +206,6 @@ private:
 	TFunction<void(bool bForWrite, TState*& OutState, bool& OutSafe)> GetStateFunc;
 };
 
-
 // ----------------------------------------------------------------------------------------------------------------------------------------------
 //	Simulation Input/Output parameters. These are the data structures passed into the simulation code each frame
 // ----------------------------------------------------------------------------------------------------------------------------------------------
@@ -239,29 +239,65 @@ struct TNetSimOutput
 		: Sync(InSync), Aux(InAux), CueDispatch(InCueDispatch) { }
 };
 
-// Scoped helper to be used right before entering a call to the sim's ::SimulationTick function.
+// Helper to actual tick a simulation. This does all the setup, runs ::SimulationTick, and finalizes the frame.
+// Caller is responsible to have the input buffers in the correct state: Input/Sync/Aux buffers must have valid data in OutputFrame-1.
+// DeltaTime comes from the InputCmd->GetFrameDeltaTime().
 // Important to note that this advances the PendingFrame to the output Frame. So that any writes that occur to the buffers during this scope will go to the output frame.
 template<typename Model>
-struct TScopedSimulationTick
+struct TSimulationDoTickFunc
 {
-	TScopedSimulationTick(FSimulationTickState& InTicker, TNetSimCueDispatcher<Model>& InDispatcher, ESimulationTickContext TickContext, const int32& InOutputFrame, const FNetworkSimTime& InDeltaSimTime)
-		: Ticker(InTicker), Dispatcher(InDispatcher), OutputFrame(InOutputFrame), DeltaSimTime(InDeltaSimTime)
+	using TSimulation = typename Model::Simulation;
+	using TTickSettings = typename Model::TickSettings;
+	using TDriver = typename TNetSimModelTraits<Model>::Driver;
+	using TBufferTypes = typename TNetSimModelTraits<Model>::InternalBufferTypes;
+	using TInputCmd = typename TBufferTypes::TInputCmd;
+	using TSyncState = typename TBufferTypes::TSyncState;
+	using TAuxState = typename TBufferTypes::TAuxState;
+
+	static void DoTick(TSimulation& Simulation, TSimulationTicker<TTickSettings>& Ticker, TNetworkSimBufferContainer<Model>& Buffers, int32 OutputFrame, ESimulationTickContext Context)
 	{
 		check(Ticker.bUpdateInProgress == false);
+		check(OutputFrame == Ticker.PendingFrame+1);
+
+		// Get inputs. These all need to be present to tick (caller's responsibility)
+		const TInputCmd* InputCmd = Buffers.Input[Ticker.PendingFrame];
+		const TSyncState* InSyncState = Buffers.Sync[Ticker.PendingFrame];
+		const TAuxState* InAuxState = Buffers.Aux[Ticker.PendingFrame];
+
+		check(InputCmd);
+		check(InSyncState);
+		check(InAuxState);
+
+		const FNetworkSimTime& DeltaSimTime = InputCmd->GetFrameDeltaTime();
+
+		// Output (only sync will always be generated
+		TSyncState* OutSyncState = Buffers.Sync.WriteFrame(OutputFrame);
+
+		// PreTick
 		Ticker.PendingFrame = OutputFrame;
 		Ticker.bUpdateInProgress = true;
 
-		Dispatcher.PushContext({Ticker.GetTotalProcessedSimulationTime() + DeltaSimTime, TickContext}); // Cues "take place" at the end of the frame
-	}
-	~TScopedSimulationTick()
-	{
+		Buffers.CueDispatcher.PushContext({Ticker.GetTotalProcessedSimulationTime() + DeltaSimTime, Context}); // Cues "take place" at the end of the frame
+
+		UE_NP_TRACE_SIM_TICK(OutputFrame, DeltaSimTime, Ticker);
+
+		// Simulation Tick
+		Simulation.SimulationTick(
+			{ DeltaSimTime, Ticker }, // How it ticks
+			{ *InputCmd, *InSyncState, *InAuxState }, // Input States
+			{ *OutSyncState, Buffers.Aux.LazyWriter(OutputFrame), Buffers.CueDispatcher } ); // Output
+
+		// Post
+		UE_NP_TRACE_USER_STATE_SYNC(*Buffers.Sync.HeadElement(), OutputFrame);
+		const TAuxState* AuxHead = Buffers.Aux.HeadElement();
+		if (AuxHead != InAuxState)
+		{
+			UE_NP_TRACE_USER_STATE_AUX(*AuxHead, OutputFrame);
+		}
+
 		Ticker.IncrementTotalProcessedSimulationTime(DeltaSimTime, OutputFrame);
 		Ticker.bUpdateInProgress = false;
 
-		Dispatcher.PopContext();
+		Buffers.CueDispatcher.PopContext();
 	}
-	FSimulationTickState& Ticker;
-	TNetSimCueDispatcher<Model>& Dispatcher;
-	const int32& OutputFrame;
-	const FNetworkSimTime& DeltaSimTime;
-};
+};
\ No newline at end of file
diff --git a/Engine/Plugins/Runtime/NetworkPrediction/Source/NetworkPrediction/Public/Trace/NetworkPredictionTrace.h b/Engine/Plugins/Runtime/NetworkPrediction/Source/NetworkPrediction/Public/Trace/NetworkPredictionTrace.h
new file mode 100644
--- /dev/null
+++ b/Engine/Plugins/Runtime/NetworkPrediction/Source/NetworkPrediction/Public/Trace/NetworkPredictionTrace.h
@@ -0,0 +1,158 @@
+// Copyright Epic Games, Inc. All Rights Reserved.
+
+#pragma once
+
+#include "NetworkedSimulationModelTick.h"
+#include "Trace/Trace.h"
+#include "Engine/EngineTypes.h" // For ENetRole urgh
+#include "Containers/UnrealString.h"
+#include "NetworkPredictionTypes.h"
+
+#ifndef UE_NP_TRACE_ENABLED
+#define UE_NP_TRACE_ENABLED !(UE_BUILD_SHIPPING || UE_BUILD_TEST)
+#endif
+
+// Tracing user state (content) can generate a lot of data. So this can be turned off here
+#ifndef UE_NP_TRACE_USER_STATES_ENABLED
+#define UE_NP_TRACE_USER_STATES_ENABLED 1
+#endif
+
+
+#if UE_NP_TRACE_ENABLED
+
+// Scope Push/Pop for active simulation
+#define UE_NP_TRACE_SET_SCOPE_SIM(SimulationId) FScopedSimulationTraceHelper _ScopedSimIdTraceHelper(SimulationId);
+
+// Called when simulation is created. (Note this also sets a Scope for tracing the initial user states next)
+#define UE_NP_TRACE_SIM_CREATED(OwningActor, GroupName) FNetworkPredictionTrace::TraceSimulationCreated(OwningActor, GroupName);
+
+#define UE_NP_TRACE_NETROLE(SimulationId, NetRole) FNetworkPredictionTrace::TraceSimulationNetRole(SimulationId, NetRole);
+
+// Called when simulation ticks in any context (includes resimulates, sim extrapolates, etc)
+#define UE_NP_TRACE_SIM_TICK(OutputFrame, FrameDeltaTime, TickState) FNetworkPredictionTrace::TraceSimulationTick(OutputFrame, FrameDeltaTime, TickState);
+
+// Called at the end of an engine frame (EOF) to update the current tick state of a simulation
+#define UE_NP_TRACE_SIM_EOF(TickState, LastSentKeyframe, LastReceivedKeyframe) FNetworkPredictionTrace::TraceSimulationEOF(TickState, LastSentKeyframe, LastReceivedKeyframe);
+
+// Called whenever we receive network data: note this only indicates that data was received, not necessarily that it was committed to the simulation buffers
+#define UE_NP_TRACE_NET_SERIALIZE_RECV(ReceivedTime, ReceivedFrame) FNetworkPredictionTrace::TraceNetSerializeRecv(ReceivedTime, ReceivedFrame);
+
+// Called whenever a new user state has been inserted into the buffers. Analysis will determine "how" it got there from previous trace events (SIM_TICK, NET_SERIALIZE_RECV)
+#define UE_NP_TRACE_USER_STATE_INPUT(UserState, Frame) FNetworkPredictionTrace::TraceUserState(UserState, Frame, FNetworkPredictionTrace::ETraceUserState::Input);
+#define UE_NP_TRACE_USER_STATE_SYNC(UserState, Frame) FNetworkPredictionTrace::TraceUserState(UserState, Frame, FNetworkPredictionTrace::ETraceUserState::Sync);
+#define UE_NP_TRACE_USER_STATE_AUX(UserState, Frame) FNetworkPredictionTrace::TraceUserState(UserState, Frame, FNetworkPredictionTrace::ETraceUserState::Aux);
+
+// Called to indicate that the previous received NetSerialized data was committed to the buffers. (Not all data is commited, but we want to trace all received data)
+#define UE_NP_TRACE_NET_SERIALIZE_COMMIT() FNetworkPredictionTrace::TraceNetSerializeCommit();
+
+// Called to indicate the previous received NetSerialized data caused a fault. This means the data was too far in the past and could not be used. (Does not mean "caused correction").
+#define UE_NP_TRACE_NET_SERIALIZE_FAULT() FNetworkPredictionTrace::TraceNetSerializeFault();
+
+// Called prior to sampling local input
+#define UE_NP_TRACE_PRODUCE_INPUT() FNetworkPredictionTrace::TraceProduceInput();
+
+// Called prior to generating synthesized input
+#define UE_NP_TRACE_SYNTH_INPUT() FNetworkPredictionTrace::TraceSynthInput();
+
+// Called to indicate we are about to write state to the buffers outside of the normal simulation tick/netrecive. TODO: add char* identifier to debug where the mod came from
+#define UE_NP_TRACE_OOB_STATE_MOD() FNetworkPredictionTrace::TraceOOBStateMod();
+
+// Called when a PIE session is started. This is so we can keep our *sets* of worlds/simulations seperate.
+#define UE_NP_TRACE_PIE_START() FNetworkPredictionTrace::TracePIEStart();
+
+
+#else
+
+// Compiled out
+#define UE_NP_TRACE_SET_SCOPE_SIM(SimulationId)
+#define UE_NP_TRACE_SIM_CREATED(...)
+#define UE_NP_TRACE_NETROLE(...)
+#define UE_NP_TRACE_SIM_TICK(...)
+#define UE_NP_TRACE_SIM_EOF(...)
+#define UE_NP_TRACE_NET_SERIALIZE_RECV(...)
+#define UE_NP_TRACE_USER_STATE_INPUT(...)
+#define UE_NP_TRACE_USER_STATE_SYNC(...)
+#define UE_NP_TRACE_USER_STATE_AUX(...)
+#define UE_NP_TRACE_ROLLBACK(...)
+#define UE_NP_TRACE_NET_SERIALIZE_COMMIT(...)
+#define UE_NP_TRACE_NET_SERIALIZE_FAULT(...)
+#define UE_NP_TRACE_PRODUCE_INPUT(...)
+#define UE_NP_TRACE_SYNTH_INPUT(...)
+#define UE_NP_TRACE_OOB_STATE_MOD(...)
+#define UE_NP_TRACE_PIE_START(...)
+
+#endif // UE_NP_TRACE_ENABLED
+
+NETWORKPREDICTION_API UE_TRACE_CHANNEL_EXTERN(NetworkPredictionChannel);
+
+class AActor;
+
+// Trace is what the system+user code calls to do traces (could still be wrapped around [my own] macros, etc. This will call into the generic UE_TRACE_LOG macros)
+class NETWORKPREDICTION_API FNetworkPredictionTrace
+{
+public:
+
+	enum ETraceVersion
+	{
+		Initial = 1,
+	};
+
+	static void TraceSimulationCreated(const AActor* OwningActor, const FName& GroupName);
+	static void TraceSimulationNetRole(uint32 SimulationId, ENetRole NetRole);
+	static void TraceSimulationTick(int32 OutputFrame, const FNetworkSimTime& FrameDeltaTime, const FSimulationTickState& TickState);
+	static void TraceSimulationEOF(const FSimulationTickState& TickState, int32 LastSentKeyframe, int32 LastReceivedKeyframe);
+	static void TraceNetSerializeRecv(const FNetworkSimTime& ReceivedTime, int32 ReceivedFrame);
+	static void TraceNetSerializeCommit();
+	static void TraceNetSerializeFault();
+	static void TraceProduceInput();
+	static void TraceSynthInput();
+	static void TraceOOBStateMod();
+	static void TracePIEStart();
+
+	enum ETraceUserState
+	{
+		Input,
+		Sync,
+		Aux
+	};
+	
+	template<typename TUserState>
+	static void TraceUserState(const TUserState& State, int32 Frame, ETraceUserState StateType)
+	{
+#if UE_NP_TRACE_USER_STATES_ENABLED
+		if (UE_TRACE_CHANNELEXPR_IS_ENABLED(NetworkPredictionChannel))
+		{
+			StrOut.SetAutoEmitLineTerminator(true);
+			StrOut.Reset();
+
+			FStandardLoggingParameters LoggingParameters(&StrOut, EStandardLoggingContext::Full, Frame);
+			State.Log(LoggingParameters);
+
+			TraceUserState_Internal(Frame, StateType);
+		}
+#endif
+	}
+
+private:
+
+	static void PushSimulationId(uint32 SimulationId);
+	static void PopSimulationId();
+
+	static void TraceUserState_Internal(int32 Frame, ETraceUserState StateType);
+
+	static FStringOutputDevice StrOut;
+
+	friend struct FScopedSimulationTraceHelper;
+};
+
+struct FScopedSimulationTraceHelper
+{
+	FScopedSimulationTraceHelper(uint32 SimulationId)
+	{
+		FNetworkPredictionTrace::PushSimulationId(SimulationId);
+	}
+	~FScopedSimulationTraceHelper()
+	{
+		FNetworkPredictionTrace::PopSimulationId();
+	}
+};
\ No newline at end of file
diff --git a/Engine/Plugins/Runtime/NetworkPrediction/readme.txt b/Engine/Plugins/Runtime/NetworkPrediction/readme.txt
--- a/Engine/Plugins/Runtime/NetworkPrediction/readme.txt
+++ b/Engine/Plugins/Runtime/NetworkPrediction/readme.txt
@@ -59,36 +59,20 @@ Simulated Proxy behavior / "modes"
 -Interpolate: Sync buffer is replicated but there is no ticking of the net sim. Literally do not have to Tick the simulation.
 	-Actual interpolation will happen in some class outside of TNetworkSimulationModel (doesn't exist yet). It will look at sync buffer and interpolate over it.
 -(Sim) Extrapolation: by default if the netsim is ticked, we synthesize command and extrapolate the simulation. With basic reconciliation to absorb fluctuations in latency.
--Forward Predict (not implemented yet): must be explicitly enabled by outside call, ties simulated proxy sim to an autonomous proxy sim. Sims will be reconciled together in step.
+-Forward Predict: Must be explicitly enabled by outside call, ties simulated proxy sim to an autonomous proxy sim. Sims will be reconciled together in step. Basic implementation is there but not complete. Still work to do here.
  
  
 TODO: Major missing elements
--"Events" into and out of the system are not in place. That is, something happens in your movement sim that needs to call outside code/generate an event. A system is needed to track this stuff rather than just calling directly into the handler.
 -No optimizations (bandwidth, cpu) have been done.
 -[Movement specific]: must lock down SetActorLocation/Rotation API so that movement state can be changed out from underneath the system.
  
-High level focus:
--The focus right now is on the generalized prediction system, rather than actual movement code.
--We are exploring "forward predicting non autonomous proxies" with the parametric mover. Unclear how deep we will go with this.
--Then Aux buffer and "events"
--Then mixing in ability system + other scripting options
- 
--Once everything feels solid we will go deeper on movement and build out a more extendable movement system.
--E.g, We don't intend all users of the engine to write their own network simulation model movement system.
--We hope to provide a generalized movement system that is easy to extend in itself without knowing all the details of the NetworkPrediction plugin.
--We intend to support RootMotion (Animation) / RootMotionSources (non anim) in the new movement system.
-
-Road Map:
--New Movement System / Ability System integration
- 
  
 // ----------------------------------------------------------------------------------------------------------
 // Glossary
 // ----------------------------------------------------------------------------------------------------------
 
-"Simulation": should always refer to use code. The simulation is what we are networking. It doesn't matter what the simulation is actually doing, we are building a system to network it like a black box.
-"System": should always refer to the TNetworkedSimulationModel code: the generic code that is doing the networking. It is distinct from the simulation.
-
+Simulation: should always refer to use code. The simulation is what we are networking. It doesn't matter what the simulation is actually doing, we are building a system to network it like a black box.
+System: should always refer to the TNetworkedSimulationModel code: the generic code that is doing the networking. It is distinct from the simulation.
 
 Reconcile: When a predicting client receives authoritative data from the server that conflicts with what they previously predicted: sorting this out is called "reconciling". Usually means you will resimulate.
 Synthesize: In general means "making shit up". For example, "synthesizing user commands" would mean "creating user commands that were not generated by a player, but instead guessing what it might be".
@@ -109,6 +93,15 @@ Smoothing: Taking the output of the simulation and applying an additional layer
 // Release notes
 // ----------------------------------------------------------------------------------------------------------
 
+Update (3-16-20)
+-Network Prediction Insights is the new Unreal Insights based tool debugging tool for Network rediction.
+-Enable the NetworkPredictionInsights plugin in your project. 
+-Run with -trace=NP or -trace=NetworkPrediction. You can also enable the NetworkPrediction channel via the TraceFilterChannels window.
+-Launch NPI through editor under Tools -> Network Prediction Insights
+-Alternatively, you can connect standalone UnrealInsights.exe directly to your game or open a saved .utrace file.
+-Existing canvas based debugging will be removed soon from Network Prediction.
+-Added TSimulationDoTickFunc<Model> to replace TScopedSimulationTick. This better encapsulates 'running a tick'. 
+
 Update (2-18-20)
 -Sorry for the lack of public updates. Development has been ongoing but mostly moved to a non public stream at the moment. The work being done revolves around:
 	-Unreal Insights. We are implementing an Insights tool for debugging Network Prediction. This will replace the canvas/visual logger debugging in the plugin. It is very powerful. Hopefully v.1 will make it into main branch soon.
diff --git a/Engine/Plugins/Runtime/NetworkPredictionExtras/Source/NetworkPredictionExtras/NetworkPredictionExtras.Build.cs b/Engine/Plugins/Runtime/NetworkPredictionExtras/Source/NetworkPredictionExtras/NetworkPredictionExtras.Build.cs
--- a/Engine/Plugins/Runtime/NetworkPredictionExtras/Source/NetworkPredictionExtras/NetworkPredictionExtras.Build.cs
+++ b/Engine/Plugins/Runtime/NetworkPredictionExtras/Source/NetworkPredictionExtras/NetworkPredictionExtras.Build.cs
@@ -33,7 +33,11 @@ namespace UnrealBuildTool.Rules
 			PrivateDependencyModuleNames.AddRange(
 				new string[]
 				{
-				}
+					"TraceLog",
+                    "TraceAnalysis",
+                    "TraceServices",
+                    "TraceInsights",
+                }
 				);
 
 			DynamicallyLoadedModuleNames.AddRange(
