CIS fixes

diff --git a/Engine/Plugins/Runtime/NetworkPrediction/Source/NetworkPrediction/Private/Trace/NetworkPredictionTrace.cpp b/Engine/Plugins/Runtime/NetworkPrediction/Source/NetworkPrediction/Private/Trace/NetworkPredictionTrace.cpp
--- a/Engine/Plugins/Runtime/NetworkPrediction/Source/NetworkPrediction/Private/Trace/NetworkPredictionTrace.cpp
+++ b/Engine/Plugins/Runtime/NetworkPrediction/Source/NetworkPrediction/Private/Trace/NetworkPredictionTrace.cpp
@@ -3,7 +3,6 @@
 
 #include "Trace/NetworkPredictionTrace.h"
 #include "CoreMinimal.h"
-#include "TraceServices/Model/AnalysisSession.h"
 #include "HAL/PlatformTime.h"
 #include "UObject/ObjectKey.h"
 #include "Engine/EngineBaseTypes.h"
diff --git a/Engine/Plugins/Runtime/NetworkPredictionInsights/Source/NetworkPredictionInsights/Private/UI/SNPSimFrameView.h b/Engine/Plugins/Runtime/NetworkPredictionInsights/Source/NetworkPredictionInsights/Private/UI/SNPSimFrameView.h
--- a/Engine/Plugins/Runtime/NetworkPredictionInsights/Source/NetworkPredictionInsights/Private/UI/SNPSimFrameView.h
+++ b/Engine/Plugins/Runtime/NetworkPredictionInsights/Source/NetworkPredictionInsights/Private/UI/SNPSimFrameView.h
@@ -4,6 +4,7 @@
 
 #include "NPAxisViewportInt32.h"
 #include "NPAxisViewportDouble.h"
+#include "INetworkPredictionProvider.h"
 
 enum class ESimFrameStatus : uint8
 {
