fix clang warning

diff --git a/Engine/Plugins/Runtime/NetworkPrediction/Source/NetworkPrediction/Private/Trace/NetworkPredictionTrace.cpp b/Engine/Plugins/Runtime/NetworkPrediction/Source/NetworkPrediction/Private/Trace/NetworkPredictionTrace.cpp
--- a/Engine/Plugins/Runtime/NetworkPrediction/Source/NetworkPrediction/Private/Trace/NetworkPredictionTrace.cpp
+++ b/Engine/Plugins/Runtime/NetworkPrediction/Source/NetworkPrediction/Private/Trace/NetworkPredictionTrace.cpp
@@ -104,10 +104,11 @@ UE_TRACE_EVENT_BEGIN(NetworkPrediction, ProduceInput)
 UE_TRACE_EVENT_END()
 
 UE_TRACE_EVENT_BEGIN(NetworkPrediction, SynthInput)
-UE_TRACE_EVENT_FIELD(uint32, SimulationId)
+	UE_TRACE_EVENT_FIELD(uint32, SimulationId)
 UE_TRACE_EVENT_END()
 
 UE_TRACE_EVENT_BEGIN(NetworkPrediction, PieBegin)
+	UE_TRACE_EVENT_FIELD(uint32, DummyData)
 UE_TRACE_EVENT_END()
 
 
@@ -341,5 +342,6 @@ void FNetworkPredictionTrace::TraceUserState_Internal(int32 Frame, ETraceUserSta
 
 void FNetworkPredictionTrace::TracePIEStart()
 {
-	UE_TRACE_LOG(NetworkPrediction, PieBegin, NetworkPredictionChannel);
+	UE_TRACE_LOG(NetworkPrediction, PieBegin, NetworkPredictionChannel)
+		<< PieBegin.DummyData(0); // temp to quiet clang
 }
\ No newline at end of file
