NetworkPrediction Plugin
-Refactor and simplifying replicated buffer+repController  implementation

#rnx

diff --git a/Engine/Plugins/Runtime/NetworkPrediction/Source/NetworkPrediction/Private/NetworkPredictionComponent.cpp b/Engine/Plugins/Runtime/NetworkPrediction/Source/NetworkPrediction/Private/NetworkPredictionComponent.cpp
--- a/Engine/Plugins/Runtime/NetworkPrediction/Source/NetworkPrediction/Private/NetworkPredictionComponent.cpp
+++ b/Engine/Plugins/Runtime/NetworkPrediction/Source/NetworkPrediction/Private/NetworkPredictionComponent.cpp
@@ -41,9 +41,6 @@ void UNetworkPredictionComponent::InitializeComponent()
 		ReplicationProxy_Autonomous.Init(NewNetworkSim, EReplicationProxyTarget::AutonomousProxy);
 		ReplicationProxy_Simulated.Init(NewNetworkSim, EReplicationProxyTarget::SimulatedProxy);
 		ReplicationProxy_Replay.Init(NewNetworkSim, EReplicationProxyTarget::Replay);
-	#if NETSIM_MODEL_DEBUG
-		ReplicationProxy_Debug.Init(NewNetworkSim, EReplicationProxyTarget::Debug);
-	#endif
 
 		// Register with GlobalManager			
 		NetworkSimGlobalManager->RegisterModel(NewNetworkSim, GetOwner());
@@ -81,9 +78,6 @@ void UNetworkPredictionComponent::PreReplication(IRepChangedPropertyTracker & Ch
 		ReplicationProxy_Autonomous.OnPreReplication();
 		ReplicationProxy_Simulated.OnPreReplication();
 		ReplicationProxy_Replay.OnPreReplication();
-#if NETSIM_MODEL_DEBUG
-		ReplicationProxy_Debug.OnPreReplication();
-#endif
 	}
 }
 
@@ -100,10 +94,6 @@ void UNetworkPredictionComponent::GetLifetimeReplicatedProps(TArray< FLifetimePr
 	DOREPLIFETIME_CONDITION( UNetworkPredictionComponent, ReplicationProxy_Autonomous, COND_AutonomousOnly);
 	DOREPLIFETIME_CONDITION( UNetworkPredictionComponent, ReplicationProxy_Simulated, COND_SimulatedOnly);
 	DOREPLIFETIME_CONDITION( UNetworkPredictionComponent, ReplicationProxy_Replay, COND_ReplayOnly);
-
-#if NETSIM_MODEL_DEBUG
-	DOREPLIFETIME_CONDITION( UNetworkPredictionComponent, ReplicationProxy_Debug, COND_ReplayOrOwner);
-#endif
 }
 
 void UNetworkPredictionComponent::InitializeForNetworkRole(ENetRole Role)
diff --git a/Engine/Plugins/Runtime/NetworkPrediction/Source/NetworkPrediction/Private/NetworkPredictionTests.cpp b/Engine/Plugins/Runtime/NetworkPrediction/Source/NetworkPrediction/Private/NetworkPredictionTests.cpp
--- a/Engine/Plugins/Runtime/NetworkPrediction/Source/NetworkPrediction/Private/NetworkPredictionTests.cpp
+++ b/Engine/Plugins/Runtime/NetworkPrediction/Source/NetworkPrediction/Private/NetworkPredictionTests.cpp
@@ -2,83 +2,4 @@
 
 #include "NetworkedSimulationModel.h"
 
-// ------------------------------------------------------------------------------------------------------------
-//	Testing Fixed vs Variable tick rate struct declration
-// ------------------------------------------------------------------------------------------------------------
-
 DEFINE_LOG_CATEGORY_STATIC(NetworkSimulationTests, Log, All);
-
-struct FNetSimDummy { };
-using FDummyBufferTypes = TNetworkSimBufferTypes<FNetSimDummy, FNetSimDummy, FNetSimDummy>;
-
-template<typename TTickSettings>
-void TickSettingsTest(TTickSettings& Settings)
-{
-	using TBufferTypes = TInternalBufferTypes<FDummyBufferTypes, TTickSettings>;
-
-	TSimulationTicker<TTickSettings> Ticker;
-
-	float RealTime = FMath::Rand();
-	Ticker.GiveSimulationTime(RealTime);
-
-}
-
-void NetSimulationTypesTest()
-{
-	TNetworkSimTickSettings<> DefaultTickSettings;
-	TickSettingsTest(DefaultTickSettings);
-
-
-	TNetworkSimTickSettings<0, 60> VariableTickSettingsWithMax;
-	TickSettingsTest(VariableTickSettingsWithMax);
-
-	TNetworkSimTickSettings<20> FixedTickSettings;
-	TickSettingsTest(FixedTickSettings);
-
-	//TNetworkSimTickSettings<20, 60> StaticAssertSettings;
-	//TickSettingsTest(StaticAssertSettings);
-}
-
-template<typename TBuffer>
-void NetSimulationTestBuffer(TBuffer& Buffer)
-{
-	for (int32 i=0; i < 10; ++i)
-	{
-		typename TBuffer::ElementType NewData;
-		*Buffer.WriteFrame(i) = NewData;
-		UE_LOG(NetworkSimulationTests, Display, TEXT("%s"), *Buffer.GetBasicDebugStr());
-
-		for (auto It = Buffer.CreateIterator(); It; ++It)
-		{
-			UE_LOG(NetworkSimulationTests, Display, TEXT("   [%d] %s"), It.Frame(), *It.Element()->ToString());
-		}
-
-		for (auto It = Buffer.CreateConstIterator(); It; ++It)
-		{
-			UE_LOG(NetworkSimulationTests, Display, TEXT("   [%d] %s"), It.Frame(), *It.Element()->ToString());
-		}
-	}
-}
-
-void NetSimulationBufferTest()
-{
-	struct FState
-	{
-		int32 x = 0;
-		int32 y = 0;
-		int32 z = 0;
-		float w = 1.f;
-
-		FString ToString() const { return FString::Printf(TEXT("%d %d %d %f"), x, y, z, w); }
-	};
-
-	TNetworkSimSparseBuffer<FState> SparseBuffer;
-	TNetworkSimContiguousBuffer<FState> ContiguousBuffer;
-
-	NetSimulationTestBuffer(SparseBuffer);
-	NetSimulationTestBuffer(ContiguousBuffer);
-
-
-	//TNetworkSimBufferIterator<TNetworkSimSparseBuffer<FState>> It(SparseBuffer);
-	//It.Frame();
-}
diff --git a/Engine/Plugins/Runtime/NetworkPrediction/Source/NetworkPrediction/Private/Trace/NetworkPredictionTrace.cpp b/Engine/Plugins/Runtime/NetworkPrediction/Source/NetworkPrediction/Private/Trace/NetworkPredictionTrace.cpp
--- a/Engine/Plugins/Runtime/NetworkPrediction/Source/NetworkPrediction/Private/Trace/NetworkPredictionTrace.cpp
+++ b/Engine/Plugins/Runtime/NetworkPrediction/Source/NetworkPrediction/Private/Trace/NetworkPredictionTrace.cpp
@@ -54,12 +54,13 @@ UE_TRACE_EVENT_END()
 UE_TRACE_EVENT_BEGIN(NetworkPrediction, SimulationEOF)
 	UE_TRACE_EVENT_FIELD(uint32, SimulationId)
 	UE_TRACE_EVENT_FIELD(uint64, EngineFrameNumber)
-	UE_TRACE_EVENT_FIELD(double, EngineFrameDeltaTime)	// FIXME: Need better tracking of engine frame vs simulation EOF updates
-	UE_TRACE_EVENT_FIELD(double, EngineCurrentTime)		// FIXME: ^^
-	UE_TRACE_EVENT_FIELD(int32, TotalProcessedMS)
-	UE_TRACE_EVENT_FIELD(int32, TotalAllowedMS)
-	UE_TRACE_EVENT_FIELD(int32, LastSentKeyframe)
-	UE_TRACE_EVENT_FIELD(int32, LastReceivedKeyframe)
+	UE_TRACE_EVENT_FIELD(float, EngineFrameDeltaTime)
+	UE_TRACE_EVENT_FIELD(int32, BufferSize)
+	UE_TRACE_EVENT_FIELD(int32, PendingTickFrame)
+	UE_TRACE_EVENT_FIELD(int32, LatestInputFrame)
+	UE_TRACE_EVENT_FIELD(int32, MaxTickFrame)
+	UE_TRACE_EVENT_FIELD(int32, TotalSimTime)
+	UE_TRACE_EVENT_FIELD(int32, AllowedSimTime)
 UE_TRACE_EVENT_END()
 
 UE_TRACE_EVENT_BEGIN(NetworkPrediction, NetSerializeRecv)
@@ -91,8 +92,10 @@ UE_TRACE_EVENT_BEGIN(NetworkPrediction, NetSerializeCommit)
 	UE_TRACE_EVENT_FIELD(uint32, SimulationId)
 UE_TRACE_EVENT_END()
 
-UE_TRACE_EVENT_BEGIN(NetworkPrediction,NetSerializeFault)
-	UE_TRACE_EVENT_FIELD(uint32,SimulationId)
+// General system fault. Log message is in attachment
+UE_TRACE_EVENT_BEGIN(NetworkPrediction, SystemFault)
+	UE_TRACE_EVENT_FIELD(uint32, SimulationId)
+	UE_TRACE_EVENT_FIELD(uint64, EngineFrameNumber)
 UE_TRACE_EVENT_END()
 
 UE_TRACE_EVENT_BEGIN(NetworkPrediction, OOBStateMod)
@@ -199,33 +202,33 @@ void FNetworkPredictionTrace::TraceSimulationNetRole(uint32 SimulationId, ENetRo
 		<< SimulationNetRole.NetRole((uint8)NetRole);
 }
 
-void FNetworkPredictionTrace::TraceSimulationTick(int32 OutputFrame, const FNetworkSimTime& FrameDeltaTime, const FSimulationTickState& TickState)
+void FNetworkPredictionTrace::TraceSimulationTick(int32 OutputFrame, const FNetworkSimTime& FrameDeltaTime, const FNetSimTimeStep& TimeStep)
 {
 	UE_TRACE_LOG(NetworkPrediction, SimulationTick, NetworkPredictionChannel)
 		<< SimulationTick.SimulationId(PeakSimulationIdChecked())
 		<< SimulationTick.EngineFrameNumber(GFrameNumber)
-		<< SimulationTick.StartMS(TickState.GetTotalProcessedSimulationTime())
-		<< SimulationTick.EndMS(TickState.GetTotalProcessedSimulationTime() + FrameDeltaTime)
+		<< SimulationTick.StartMS(TimeStep.TotalSimulationTime)
+		<< SimulationTick.EndMS(TimeStep.TotalSimulationTime + TimeStep.StepMS)
 		<< SimulationTick.OutputFrame(OutputFrame);
 }
 
-void FNetworkPredictionTrace::TraceSimulationEOF(const FSimulationTickState& TickState, int32 LastSentKeyframe, int32 LastReceivedKeyframe)
+void FNetworkPredictionTrace::TraceEOF_Internal(int32 BufferSize, int32 PendingTickFrame, int32 LatestInputFrame, int32 MaxTickFrame, FNetworkSimTime TotalSimTime, FNetworkSimTime AllowedSimTime)
 {
 	const uint32 SimulationId = PeakSimulationIdChecked();
 
 	UE_TRACE_LOG(NetworkPrediction, SimulationEOF, NetworkPredictionChannel)
 		<< SimulationEOF.SimulationId(SimulationId)
 		<< SimulationEOF.EngineFrameNumber(GFrameNumber)
-		<< SimulationEOF.EngineFrameDeltaTime(FApp::GetDeltaTime())	// FIXME
-		<< SimulationEOF.EngineCurrentTime(FApp::GetCurrentTime())	// FIXME
-		<< SimulationEOF.TotalProcessedMS(TickState.GetTotalProcessedSimulationTime())
-		<< SimulationEOF.TotalAllowedMS(TickState.GetTotalAllowedSimulationTime())
-		<< SimulationEOF.LastSentKeyframe(LastSentKeyframe)
-		<< SimulationEOF.LastReceivedKeyframe(LastReceivedKeyframe);
-	
-
-
-	// Do we need to trace the NetGUID? This stinks having to be here
+		<< SimulationEOF.EngineFrameDeltaTime(FApp::GetDeltaTime())
+		<< SimulationEOF.BufferSize(BufferSize)
+		<< SimulationEOF.PendingTickFrame(PendingTickFrame)
+		<< SimulationEOF.LatestInputFrame(LatestInputFrame)
+		<< SimulationEOF.MaxTickFrame(MaxTickFrame)
+		<< SimulationEOF.TotalSimTime(TotalSimTime)
+		<< SimulationEOF.AllowedSimTime(AllowedSimTime);
+
+	// Trace this simulation's NetGUID if we haven't already
+	// (This stinks having to be here: but we can't get an event for NetGUID assignment and it wont always be assigned when sim is created)
 	if (TracedSimulationNetGUIDs.Contains(SimulationId) == false)
 	{
 		uint32 NetGUID = 0;
@@ -271,14 +274,6 @@ void FNetworkPredictionTrace::TraceNetSerializeCommit()
 		<< NetSerializeCommit.SimulationId(SimulationId);
 }
 
-void FNetworkPredictionTrace::TraceNetSerializeFault()
-{
-	const uint32 SimulationId = PeakSimulationIdChecked();
-
-	UE_TRACE_LOG(NetworkPrediction, NetSerializeFault, NetworkPredictionChannel)
-		<< NetSerializeFault.SimulationId(SimulationId);
-}
-
 void FNetworkPredictionTrace::TraceProduceInput()
 {
 	const uint32 SimulationId = PeakSimulationIdChecked();
@@ -344,4 +339,60 @@ void FNetworkPredictionTrace::TracePIEStart()
 {
 	UE_TRACE_LOG(NetworkPrediction, PieBegin, NetworkPredictionChannel)
 		<< PieBegin.DummyData(0); // temp to quiet clang
+}
+
+#include "CoreTypes.h"
+#include "Misc/VarArgs.h"
+#include "HAL/UnrealMemory.h"
+#include "Templates/UnrealTemplate.h"
+
+// Copied from VarargsHelper.h
+#define GROWABLE_LOGF(SerializeFunc) \
+	int32	BufferSize	= 1024; \
+	TCHAR*	Buffer		= NULL; \
+	int32	Result		= -1; \
+	/* allocate some stack space to use on the first pass, which matches most strings */ \
+	TCHAR	StackBuffer[512]; \
+	TCHAR*	AllocatedBuffer = NULL; \
+\
+	/* first, try using the stack buffer */ \
+	Buffer = StackBuffer; \
+	GET_VARARGS_RESULT( Buffer, UE_ARRAY_COUNT(StackBuffer), UE_ARRAY_COUNT(StackBuffer) - 1, Fmt, Fmt, Result ); \
+\
+	/* if that fails, then use heap allocation to make enough space */ \
+	while(Result == -1) \
+	{ \
+		FMemory::SystemFree(AllocatedBuffer); \
+		/* We need to use malloc here directly as GMalloc might not be safe. */ \
+		Buffer = AllocatedBuffer = (TCHAR*) FMemory::SystemMalloc( BufferSize * sizeof(TCHAR) ); \
+		if (Buffer == NULL) \
+		{ \
+			return; \
+		} \
+		GET_VARARGS_RESULT( Buffer, BufferSize, BufferSize-1, Fmt, Fmt, Result ); \
+		BufferSize *= 2; \
+	}; \
+	Buffer[Result] = 0; \
+	; \
+\
+	SerializeFunc; \
+	FMemory::SystemFree(AllocatedBuffer);
+
+
+void FNetworkPredictionTrace::TraceSystemFault(const TCHAR* Fmt, ...)
+{
+	const uint32 SimulationId = PeakSimulationIdChecked();
+
+	GROWABLE_LOGF( 
+		
+		check(Result >= 0 );
+		const uint16 AttachmentSize = (Result+1) * sizeof(TCHAR);
+
+		UE_TRACE_LOG(NetworkPrediction, SystemFault, NetworkPredictionChannel, AttachmentSize)
+			<< SystemFault.SimulationId(SimulationId)
+			<< SystemFault.EngineFrameNumber(GFrameNumber)
+			<< SystemFault.Attachment(Buffer, AttachmentSize);
+
+		UE_LOG(LogNetworkSim, Warning, TEXT("SystemFault: %s"), Buffer);
+	);
 }
\ No newline at end of file
diff --git a/Engine/Plugins/Runtime/NetworkPrediction/Source/NetworkPrediction/Public/NetworkPredictionCheck.h b/Engine/Plugins/Runtime/NetworkPrediction/Source/NetworkPrediction/Public/NetworkPredictionCheck.h
new file mode 100644
--- /dev/null
+++ b/Engine/Plugins/Runtime/NetworkPrediction/Source/NetworkPrediction/Public/NetworkPredictionCheck.h
@@ -0,0 +1,27 @@
+// Copyright Epic Games, Inc. All Rights Reserved.
+
+#pragma once
+
+#define NP_CHECKS_AND_ENSURES 1
+#if NP_CHECKS_AND_ENSURES
+	#define npCheckf(Condition, ...) checkf(Condition, ##__VA_ARGS__)
+	#define npEnsure(Condition) ensureAlways(Condition)
+	#define npEnsureMsgf(Condition, ...) ensureAlwaysMsgf(Condition, ##__VA_ARGS__)
+#else
+	#define npCheckf(...)
+	#define npEnsure(Condition)
+	#define npEnsureMsgf(...)
+#endif
+
+#define NP_CHECKS_AND_ENSURES_SLOW !(UE_BUILD_SHIPPING || UE_BUILD_TEST)
+#if NP_CHECKS_AND_ENSURES_SLOW
+	#define npCheckSlow(Condition) check(Condition)
+	#define npCheckfSlow(Condition, ...) checkf(Condition, ##__VA_ARGS__)
+	#define npEnsureSlow(Condition) ensureAlways(Condition)
+	#define npEnsureMsgfSlow(Condition, ...) ensureAlwaysMsgf(Condition, ##__VA_ARGS__)
+#else
+	#define npCheckSlow(Condition)
+	#define npCheckfSlow(...)
+	#define npEnsureSlow(...)
+	#define npEnsureMsgfSlow(...)
+#endif
\ No newline at end of file
diff --git a/Engine/Plugins/Runtime/NetworkPrediction/Source/NetworkPrediction/Public/NetworkPredictionComponent.h b/Engine/Plugins/Runtime/NetworkPrediction/Source/NetworkPrediction/Public/NetworkPredictionComponent.h
--- a/Engine/Plugins/Runtime/NetworkPrediction/Source/NetworkPrediction/Public/NetworkPredictionComponent.h
+++ b/Engine/Plugins/Runtime/NetworkPrediction/Source/NetworkPrediction/Public/NetworkPredictionComponent.h
@@ -75,9 +75,6 @@ private:
 	UPROPERTY(Replicated, transient)
 	FReplicationProxy ReplicationProxy_Replay;
 
-	UPROPERTY(Replicated, transient)
-	FReplicationProxy ReplicationProxy_Debug;
-
 protected:
 
 	// Called via ServerRPCDelegate, ticks serverRPC timing and decides whether to send the RPC or not
diff --git a/Engine/Plugins/Runtime/NetworkPrediction/Source/NetworkPrediction/Public/NetworkPredictionTypes.h b/Engine/Plugins/Runtime/NetworkPrediction/Source/NetworkPrediction/Public/NetworkPredictionTypes.h
--- a/Engine/Plugins/Runtime/NetworkPrediction/Source/NetworkPrediction/Public/NetworkPredictionTypes.h
+++ b/Engine/Plugins/Runtime/NetworkPrediction/Source/NetworkPrediction/Public/NetworkPredictionTypes.h
@@ -26,7 +26,6 @@ enum class EReplicationProxyTarget: uint8
 	AutonomousProxy,	// Owning/Controlling client
 	SimulatedProxy,		// Non owning client
 	Replay,				// Replay net driver
-	Debug,				// Replication target that is disabled in shipping
 };
 
 inline FString LexToString(EReplicationProxyTarget A)
@@ -65,54 +64,6 @@ struct FNetworkSimulationModelInitParameters
 	int32 HistoricBufferSize=0;
 };
 
-struct FNetSimProcessedFrameDebugInfo
-{
-	int32 LocalGFrameNumber; // Local GFrame number
-	float LocalDeltaTimeSeconds; // Local frame time
-	TArray<int32> ProcessedFrames; // Which frames were processed this frame
-	int32 PendingFrame; // What PendingFrame was at the end of the frame. Does NOT mean we processed it this frame!
-	int32 HeadFrame; // Head frame of the inputbuffer when the frame ended
-
-	float RemainingAllowedSimulationTimeSeconds;
-	
-	int32 LastSentInputFrame;
-	int32 LastReceivedInputFrame;
-
-	void NetSerialize(const FNetSerializeParams& P)
-	{	
-		P.Ar << LocalGFrameNumber;
-		P.Ar << LocalDeltaTimeSeconds;
-		P.Ar << PendingFrame;
-		P.Ar << HeadFrame;
-		P.Ar << RemainingAllowedSimulationTimeSeconds;
-
-		P.Ar << LastSentInputFrame;
-		P.Ar << LastReceivedInputFrame;
-
-		SafeNetSerializeTArray_Default<31>(P.Ar, ProcessedFrames);
-	}
-	void Log(FStandardLoggingParameters& P) const
-	{
-		P.Ar->Logf(TEXT("LocalGFrameNumber: %d"), LocalGFrameNumber);
-		P.Ar->Logf(TEXT("LocalDeltaTimeSeconds: %.4f"), LocalDeltaTimeSeconds);
-		P.Ar->Logf(TEXT("PendingFrame: %d"), PendingFrame);
-		P.Ar->Logf(TEXT("HeadFrame: %d"), HeadFrame);
-		P.Ar->Logf(TEXT("RemainingAllowedSimulationTimeSeconds: %.4f"), RemainingAllowedSimulationTimeSeconds);
-
-		P.Ar->Logf(TEXT("LastSentInputFrame: %d"), LastSentInputFrame);
-		P.Ar->Logf(TEXT("LastReceivedInputFrame: %d"), LastReceivedInputFrame);
-
-		FString ProcessedFramesStr;
-		for (int32 k : ProcessedFrames)
-		{
-			ProcessedFramesStr += LexToString(k);
-			ProcessedFramesStr += TEXT(" ");
-		}
-		ProcessedFramesStr.TrimEndInline();
-		P.Ar->Logf(TEXT("ProcessedFrames: [%s]"), *ProcessedFramesStr);
-	}
-};
-
 UENUM()
 enum class EVisualLoggingContext: uint8
 {
diff --git a/Engine/Plugins/Runtime/NetworkPrediction/Source/NetworkPrediction/Public/NetworkedSimulationModel.h b/Engine/Plugins/Runtime/NetworkPrediction/Source/NetworkPrediction/Public/NetworkedSimulationModel.h
--- a/Engine/Plugins/Runtime/NetworkPrediction/Source/NetworkPrediction/Public/NetworkedSimulationModel.h
+++ b/Engine/Plugins/Runtime/NetworkPrediction/Source/NetworkPrediction/Public/NetworkedSimulationModel.h
@@ -11,17 +11,15 @@
 // ----------------------------------------------------------------------------------------------------------------------------------------------------------------------
 //	TNetworkedSimulationModel
 //
-//	This is the object that ties all user-defined NetworkPrediction things together:
+//	This is the core object that ties all user-defined NetworkPrediction things together:
 //		1. The Simulation: "pure" gameplay simulation: takes input -> produces output. Knows nothing about networking or the general UE4 side of things.
 //		2. The Model Definition: defines user types, tick settings, the simulation class and functions such as ShouldReconcile and Interpolate. E.g, "How do you want the networking model to work".
 //		3. The Driver: required interface for this NetSimModel to function: FinalizeFrame, ProduceInput, Debugging/VisualLogging, etc. Implemented by the Actor/Component/UE4 side.
 //
-//
-//	* Has 4 buffers:
+//	* Contains all state in TNetworkedSimulationState
 //		-Input: Generated by a client / not the authority.
 //		-Sync: What we are trying to keep in sync. The state that evolves frame to frame with an Update function.
 //		-Aux: State that is also an input into the simulation but does not intrinsically evolve from to frame. Changes to this state can be trapped/tracked/predicted.
-//		-Debug: Replicated buffer from server->client with server-frame centered debug information. Compiled out of shipping builds.
 //
 //	* How other code interacts with this:
 //		-Network updates will come in through UE4 networking -> FReplicationProxy (on actor/component) -> TNetworkedSimulationModel::RepProxy_* -> Buffers.*
@@ -79,10 +77,7 @@ struct TNetSimModelRepControllerTraitsBase
 	using RepController_Autonomous = TRepController_Autonomous<Model>;
 
 	template<typename Model>
-	using RepController_Debug = TRepController_Debug<Model>;
-
-	template<typename Model>
-	using RepController_Replay = TRepController_Empty<Model>;
+	using RepController_Replay = TRepController_Simulated<Model>; // Replay writes simulated data stream.
 };
 
 // Template struct to specialize per-model
@@ -100,14 +95,15 @@ public:
 	// "Core" model pieces
 	using TSimulation = typename Model::Simulation;
 	using TTickSettings = typename Model::TickSettings;
-	using TDriver = typename TNetSimModelTraits<Model>::Driver;
-	using TBufferTypes = typename TNetSimModelTraits<Model>::InternalBufferTypes;
+	using TBufferTypes = typename Model::BufferTypes;
+
+	using TDriver = TNetworkedSimulationModelDriver<TBufferTypes>;
+	using TFrameState = TSimulationFrameState<Model>;
 
 	// Buffer types
 	using TInputCmd = typename TBufferTypes::TInputCmd;
 	using TSyncState = typename TBufferTypes::TSyncState;
 	using TAuxState = typename TBufferTypes::TAuxState;
-	using TDebugState = typename TBufferTypes::TDebugState;
 
 	// Time
 	using TSimTime = FNetworkSimTime;
@@ -118,8 +114,6 @@ public:
 	using TRepProxySimulated =	typename TNetSimModelRepControllerTraits<Model>::template RepController_Simulated<Model>;
 	using TRepProxyAutonomous =	typename TNetSimModelRepControllerTraits<Model>::template RepController_Autonomous<Model>;
 
-	// Defines how replication happens on these special channels, but doesn't dictate how simulation evolves
-	using TRepProxyDebug =	typename TNetSimModelRepControllerTraits<Model>::template RepController_Debug<Model>;
 	using TRepProxyReplay =	typename TNetSimModelRepControllerTraits<Model>::template RepController_Replay<Model>;
 
 	// ----------------------------------------------------------------------------------------------------------
@@ -133,17 +127,16 @@ public:
 	{
 		check(Simulation.IsValid() && InDriver);
 		Driver = InDriver;
-		*Buffers.Sync.WriteFrame(0) = InitialSyncState;
-		*Buffers.Aux.WriteFrame(0) = InitialAuxState;
-		Ticker.SetTotalProcessedSimulationTime(FNetworkSimTime(), 0);
+		State.WriteFrameInitial()->SyncState = InitialSyncState;
+		*State.WriteAuxState(0) = InitialAuxState;
 		DO_NETSIM_MODEL_DEBUG(FNetworkSimulationModelDebuggerManager::Get().RegisterNetworkSimulationModel(this, Driver->GetVLogOwner()));
 
 		this->ProcessPendingNetSimCuesFunc = [this, InDriver]()
 		{ 
-			Buffers.CueDispatcher.template DispatchCueRecord<TInDriver>(*InDriver, Ticker.GetTotalProcessedSimulationTime()); 
+			State.CueDispatcher.template DispatchCueRecord<TInDriver>(*InDriver, State.GetTotalProcessedSimulationTime()); 
 		};
 
-		this->Buffers.CueDispatcher.GetDebugName = [this]() { return Driver->GetDebugName(); };
+		this->State.CueDispatcher.GetDebugName = [this]() { return Driver->GetDebugName(); };
 		
 		SimulationId = this->GetNextSimulationId();
 		UE_NP_TRACE_SET_SCOPE_SIM(SimulationId);
@@ -162,162 +155,23 @@ public:
 	void Tick(const FNetSimTickParameters& Parameters) final override
 	{
 		UE_NP_TRACE_SET_SCOPE_SIM(SimulationId);
-
-		// Update previous DebugState based on what we (might) have sent *after* our last Tick 
-		// (property replication and ServerRPC get sent after the tick, rather than forcing awkward callback into the NetSim post replication, we can check it here)
-		if (auto* DebugBuffer = GetLocalDebugBuffer())
-		{
-			if (TDebugState* const PrevDebugState = DebugBuffer->Get(DebugBuffer->HeadFrame()))
-			{
-				if (Parameters.Role == ROLE_AutonomousProxy)
-				{
-					PrevDebugState->LastSentInputFrame = RepProxy_ServerRPC.GetLastSerializedFrame();
-				}
-				else if (Parameters.Role == ROLE_Authority)
-				{
-					PrevDebugState->LastSentInputFrame = RepProxy_Autonomous.GetLastSerializedFrame();
-				}
-			}
-		}
-
-		// Current frame debug state
-		TDebugState* const DebugState = GetNextLocalDebugStateWrite();
-		if (DebugState)
-		{
-			*DebugState = TDebugState();
-			DebugState->LocalDeltaTimeSeconds = Parameters.LocalDeltaTimeSeconds;
-			DebugState->LocalGFrameNumber = GFrameNumber;
-			DebugState->ProcessedFrames.Reset();
-			
-			if (Parameters.Role == ROLE_AutonomousProxy)
-			{
-				DebugState->LastReceivedInputFrame = RepProxy_Autonomous.GetLastSerializedFrame();
-			}
-			else if (Parameters.Role == ROLE_Authority)
-			{
-				DebugState->LastReceivedInputFrame = RepProxy_ServerRPC.GetLastSerializedFrame();
-			}
-		}
-
-		// ----------------------------------------------------------------------------------------------------------------
-		//	PreSimTick
-		//	This is the beginning of a new frame. PreSimTick will decide if we should take Parameters.LocalDeltaTimeSeconds
-		//	and advance the simulation or not. It will also generate new local input if necessary.
-		// ----------------------------------------------------------------------------------------------------------------
-		switch (Parameters.Role)
-		{
-			case ROLE_Authority:
-				RepProxy_ServerRPC.PreSimTick(Driver, Buffers, Ticker, Parameters);
-			break;
-
-			case ROLE_AutonomousProxy:
-				RepProxy_Autonomous.PreSimTick(Driver, Buffers, Ticker, Parameters);
-			break;
-
-			case ROLE_SimulatedProxy:
-				RepProxy_Simulated.PreSimTick(Driver, Buffers, Ticker, Parameters);
-			break;
-		}
-
-		// -------------------------------------------------------------------------------------------------------------------------------------------------
-		//												Input Processing & Simulation Update
-		// -------------------------------------------------------------------------------------------------------------------------------------------------
-		while (Ticker.PendingFrame <= Ticker.MaxAllowedFrame)
-		{
-			const int32 InputFrame = Ticker.PendingFrame;
-			const int32 OutputFrame = Ticker.PendingFrame + 1;
-
-			const TInputCmd* InputCmd = Buffers.Input[InputFrame];
-			if (!ensureMsgf(InputCmd, TEXT("No InputCmd available for Frame %d. PendingFrame: %d. MaxAllowedFrame: %d."), InputFrame, Ticker.PendingFrame, Ticker.MaxAllowedFrame))
-			{
-				break;
-			}
-
-			// We have an unprocessed command, do we have enough allotted simulation time to process it?
-			if (Ticker.GetRemainingAllowedSimulationTime() >= InputCmd->GetFrameDeltaTime())
-			{
-				TSyncState* InSyncState = Buffers.Sync[InputFrame];
-				if (InSyncState == nullptr)
-				{
-					// We don't have valid sync state for this frame. This could mean we skipped ahead somehow, such as heavy packet loss where the server doesn't receive all input cmds
-					// We will use the last known sync state, but first write it out to the current InputFrame
-					UE_LOG(LogNetworkSim, Log, TEXT("%s: No sync state found @ frame %d. Using previous head element @ frame %d."), *Driver->GetDebugName(), InputFrame, Buffers.Sync.HeadFrame());
-					InSyncState = Buffers.Sync.WriteFrameInitializedFromHead(InputFrame);
-				}
-
-				ESimulationTickContext Context = Parameters.Role == ROLE_Authority ? ESimulationTickContext::Authority : ESimulationTickContext::Predict;
-				TSimulationDoTickFunc<Model>::DoTick(*Simulation.Get(), Ticker, Buffers, OutputFrame, Context);
-				
-				if (DebugState)
-				{
-					DebugState->ProcessedFrames.Add(InputFrame);
-				}
-			}
-			else
-			{
-				break;
-			}
-		}
-
-		// -------------------------------------------------------------------------------------------------------------------------------------------------
-		//												Post Sim Tick: finalize the frame
-		// -------------------------------------------------------------------------------------------------------------------------------------------------
-
-		switch (Parameters.Role)
-		{
-			case ROLE_Authority:
-				RepProxy_ServerRPC.PostSimTick(Driver, Buffers, Ticker, Parameters);
-			break;
-
-			case ROLE_AutonomousProxy:
-				RepProxy_Autonomous.PostSimTick(Driver, Buffers, Ticker, Parameters);
-			break;
-
-			case ROLE_SimulatedProxy:
-				RepProxy_Simulated.PostSimTick(Driver, Buffers, Ticker, Parameters);
-			break;
-		}
-
-		// -------------------------------------------------------------------------------------------------------------------------------------------------
-		//														Debug
-		// -------------------------------------------------------------------------------------------------------------------------------------------------
-
-		// Finish debug state buffer recording (what the server processed each frame)
-		if (DebugState)
-		{
-			DebugState->PendingFrame = Ticker.PendingFrame;
-			DebugState->HeadFrame = Buffers.Input.HeadFrame();
-			DebugState->RemainingAllowedSimulationTimeSeconds = (float)Ticker.GetRemainingAllowedSimulationTime().ToRealTimeSeconds();
-		}
-
-		// Historical data recording (longer buffers for historical reference)
-		if (auto* HistoricData = GetHistoricBuffers())
-		{
-			HistoricData->Input.CopyAndMerge(Buffers.Input);
-			HistoricData->Sync.CopyAndMerge(Buffers.Sync);
-			HistoricData->Aux.CopyAndMerge(Buffers.Aux);
-		}
 		
-		int32 LastSentKeyframe = 0;
-		int32 LastReceivedKeyframe = 0;
 		switch (Parameters.Role)
 		{
 			case ROLE_Authority:
-				LastSentKeyframe = RepProxy_Autonomous.GetLastSerializedFrame();
-				LastReceivedKeyframe = RepProxy_ServerRPC.GetLastSerializedFrame();
+				RepProxy_ServerRPC.Tick(Driver, Simulation.Get(), State, Parameters);
 			break;
 
 			case ROLE_AutonomousProxy:
-				LastSentKeyframe = RepProxy_ServerRPC.GetLastSerializedFrame();
-				LastReceivedKeyframe = RepProxy_Autonomous.GetLastSerializedFrame();
+				RepProxy_Autonomous.Tick(Driver, Simulation.Get(), State, Parameters);
 			break;
 
 			case ROLE_SimulatedProxy:
-				LastReceivedKeyframe = RepProxy_Simulated.GetLastSerializedFrame();
+				RepProxy_Simulated.Tick(Driver, Simulation.Get(), State, Parameters);
 			break;
-		}
+		}		
 
-		UE_NP_TRACE_SIM_EOF(Ticker, LastSentKeyframe, LastReceivedKeyframe);
+		UE_NP_TRACE_SIM_EOF(State);
 	}
 
 	virtual void Reconcile(const ENetRole Role) final override
@@ -332,15 +186,15 @@ public:
 		switch (Role)
 		{
 			case ROLE_Authority:
-				RepProxy_ServerRPC.Reconcile(Simulation.Get(), Driver, Buffers, Ticker);
+				RepProxy_ServerRPC.Reconcile(Simulation.Get(), Driver, State);
 			break;
 
 			case ROLE_AutonomousProxy:
-				RepProxy_Autonomous.Reconcile(Simulation.Get(), Driver, Buffers, Ticker);
+				RepProxy_Autonomous.Reconcile(Simulation.Get(), Driver, State);
 			break;
 
 			case ROLE_SimulatedProxy:
-				RepProxy_Simulated.Reconcile(Simulation.Get(), Driver, Buffers, Ticker);
+				RepProxy_Simulated.Reconcile(Simulation.Get(), Driver, State);
 			break;
 		}
 	}
@@ -353,21 +207,7 @@ public:
 
 		//Buffers.Input.SetBufferSize(Parameters.InputBufferSize);
 		//Buffers.Sync.SetBufferSize(Parameters.SyncedBufferSize);
-		//Buffers.Aux.SetBufferSize(Parameters.AuxBufferSize); AUXFIXME
-
-		if (GetLocalDebugBuffer())
-		{
-			//GetLocalDebugBuffer()->SetBufferSize(Parameters.DebugBufferSize);
-		}
-
-		if (auto* MyHistoricBuffers = GetHistoricBuffers(true))
-		{
-			//MyHistoricBuffers->Input.SetBufferSize(Parameters.HistoricBufferSize);
-			//MyHistoricBuffers->Sync.SetBufferSize(Parameters.HistoricBufferSize);
-			//MyHistoricBuffers->Aux.SetBufferSize(Parameters.HistoricBufferSize); AUXFIXME
-		}
-
-		//Ticker.InitSimulationTimeBuffer(Parameters.SyncedBufferSize);
+		//Buffers.Aux.SetBufferSize(Parameters.AuxBufferSize);
 	}
 
 	void NetSerializeProxy(EReplicationProxyTarget Target, const FNetSerializeParams& Params) final override
@@ -377,22 +217,17 @@ public:
 		switch(Target)
 		{
 		case EReplicationProxyTarget::ServerRPC:
-			RepProxy_ServerRPC.NetSerialize(Params, Buffers, Ticker);
+			RepProxy_ServerRPC.NetSerialize(Params, State);
 			break;
 		case EReplicationProxyTarget::AutonomousProxy:
-			RepProxy_Autonomous.NetSerialize(Params, Buffers, Ticker);
+			RepProxy_Autonomous.NetSerialize(Params, State);
 			break;
 		case EReplicationProxyTarget::SimulatedProxy:
-			RepProxy_Simulated.NetSerialize(Params, Buffers, Ticker);
+			RepProxy_Simulated.NetSerialize(Params, State);
 			break;
 		case EReplicationProxyTarget::Replay:
-			RepProxy_Replay.NetSerialize(Params, Buffers, Ticker);
-			break;
-		case EReplicationProxyTarget::Debug:
-#if NETSIM_MODEL_DEBUG
-			RepProxy_Debug.NetSerialize(Params, Buffers, Ticker);
+			RepProxy_Replay.NetSerialize(Params, State);
 			break;
-#endif
 		default:
 			checkf(false, TEXT("Unknown: %d"), (int32)Target);
 		};
@@ -400,71 +235,16 @@ public:
 
 	int32 GetProxyDirtyCount(EReplicationProxyTarget Target) final override
 	{
-		switch(Target)
-		{
-		case EReplicationProxyTarget::ServerRPC:
-			return RepProxy_ServerRPC.GetProxyDirtyCount(Buffers);
-		case EReplicationProxyTarget::AutonomousProxy:
-			return RepProxy_Autonomous.GetProxyDirtyCount(Buffers);
-		case EReplicationProxyTarget::SimulatedProxy:
-			return RepProxy_Simulated.GetProxyDirtyCount(Buffers);
-		case EReplicationProxyTarget::Replay:
-			return RepProxy_Replay.GetProxyDirtyCount(Buffers);
-		case EReplicationProxyTarget::Debug:
-#if NETSIM_MODEL_DEBUG
-			return RepProxy_Debug.GetProxyDirtyCount(Buffers);
-#endif
-		default:
-			checkf(false, TEXT("Unknown: %d"), (int32)Target);
-			return 0;
-		};
+		// Originally we were allowing RepControllers to define this but it seems overkill/unnecessary.
+		// Simpler to assume if the system is changing, then replication is dirty. 
+		return FMath::Max(State.GetPendingTickFrame(), State.GetLatestInputFrame());
 	}
 
 	ESimulatedUpdateMode GetSimulatedUpdateMode() const
 	{
 		return RepProxy_Simulated.GetSimulatedUpdateMode();
 	}
-
-	// ------------------------------------------------------------------------------------------------------
-	// State Accessors. This allows game code to read or conditionally write to the pending frame values.
-	// This mirrors what TNetSimStateAccessor do, but is callable directly on the NetworkedSimulationModel.
-	// Accessor conditionally gives access to the current (pending) Sync/Aux state to outside code.
-	// Reads are always allowed, Writes are conditional and null will sometimes be returned!
-	//
-	// Authority can always write to the pending frame. Non authority requires the netsim to be currently processing an ::SimulationTick.
-	// If you aren't inside an ::SimulationTick call, it is really not safe to predict state changes. It is safest and simplest to just not predict these changes.
-	//
-	// Explanation: During the scope of an ::SimulationTick call, we know exactly 'when' we are relative to what the server is processing. If the predicting client wants
-	// to predict a change to sync/aux state during an update, the server will do it at the exact same time (assuming not a mis prediction). When a state change
-	// happens "out of band" (outside an ::SimulationTick call) - we really have no way to correlate when the server will do it. While its tempting to think "we will get
-	// a correction anyways, might as well guess at it and maybe get a smaller correction" - but this opens us up to other problems. The server may actually not 
-	// change the state at all and you may not get an update that corrects you. You could add a timeout and track the state change somewhere but that really complicates
-	// things and could leave you open to "double" problems: if the state change is additive, you may stack the authority change on top of the local predicted change, or
-	// you may roll back the predicted change to then later receive the authority change.
-	//	
-	// What still may make sense to do is allow the "In Update" bool to be temporarily disabled if we enter code that we know is not rollback friendly.
-	// ------------------------------------------------------------------------------------------------------
-
-	template<typename TState>
-	const TState* GetPendingStateRead() const
-	{
-		// Reading is always allowed
-		auto& Buffer = NetSimBufferSelect::Get<TNetworkSimBufferContainer<Model>, TState>(Buffers);
-		return Buffer[Ticker.PendingFrame];
-	}
-
-	template<typename TState>
-	TState* GetPendingStateWrite(bool bHasAuthority)
-	{
-		// Writes are always allowed on the authority or during a simulation update for predicting clients
-		if (bHasAuthority || Ticker.bUpdateInProgress)
-		{
-			auto& Buffer = NetSimBufferSelect::Get<TNetworkSimBufferContainer<Model>, TState>(Buffers);
-			return Buffer.WriteFrameInitializedFromHead(Ticker.PendingFrame);
-		}
-		return nullptr;
-	}
-
+	
 	void SetEnableSimulationExtrapolation(bool bNewValue) final override
 	{
 		RepProxy_Simulated.bAllowSimulatedExtrapolation = bNewValue;
@@ -513,13 +293,13 @@ public:
 	void BeginRollback(const FNetworkSimTime& RollbackDeltaTime, const int32 ParentFrame) final override
 	{
 		UE_NP_TRACE_SET_SCOPE_SIM(SimulationId);
-		RepProxy_Simulated.DependentRollbackBegin(Simulation.Get(), Driver, Buffers, Ticker, RollbackDeltaTime, ParentFrame);
+		RepProxy_Simulated.DependentRollbackBegin(Simulation.Get(), Driver, State, RollbackDeltaTime, ParentFrame);
 	}
 
 	void StepRollback(const FNetworkSimTime& Step, const int32 ParentFrame, const bool bFinalStep) final override
 	{
 		UE_NP_TRACE_SET_SCOPE_SIM(SimulationId);
-		RepProxy_Simulated.DependentRollbackStep(Simulation.Get(), Driver, Buffers, Ticker, Step, ParentFrame, bFinalStep);
+		RepProxy_Simulated.DependentRollbackStep(Simulation.Get(), Driver, State, Step, ParentFrame, bFinalStep);
 	}
 
 	void ClearAllDependentSimulations()
@@ -534,22 +314,15 @@ public:
 	// -------------------------------------------------------------------------------------------------------
 
 	TUniquePtr<TSimulation> Simulation;
-	TDriver* Driver = nullptr;
+	TDriver* Driver = nullptr;	
 
-	TSimulationTicker<TTickSettings> Ticker;
-
-	TNetworkSimBufferContainer<Model> Buffers;
+	TNetworkedSimulationState<Model> State;
 
 	TRepProxyServerRPC RepProxy_ServerRPC;
 	TRepProxyAutonomous RepProxy_Autonomous;
 	TRepProxySimulated RepProxy_Simulated;
 	TRepProxyReplay RepProxy_Replay;
 
-	using TInputBuffer = typename TNetworkSimBufferContainer<Model>::TInputBuffer;
-	using TSyncBuffer = typename TNetworkSimBufferContainer<Model>::TSyncBuffer;
-	using TAuxBuffer = typename TNetworkSimBufferContainer<Model>::TAuxBuffer;
-	using TDebugBuffer = typename TNetworkSimBufferContainer<Model>::TDebugBuffer;
-
 	// ------------------------------------------------------------------
 	// RPC Sending helper: provides basic send frequency settings for tracking when the Server RPC can be invoked.
 	// Note that the Driver is the one that must call the RPC, that cannot be rolled into this templated structure.
@@ -583,34 +356,4 @@ private:
 	float ServerRPCThresholdTimeSeconds = 1.f / 999.f; // Default is to send at a max of 999hz. This part of the system needs to be build out more (better handling of super high FPS clients and fixed rate servers)
 
 	uint32 SimulationId = 0;
-
-	// ------------------------------------------------------------------
-	//	Debugging
-	// ------------------------------------------------------------------
-public:
-
-#if NETSIM_MODEL_DEBUG
-	TDebugBuffer* GetLocalDebugBuffer() {	return &Buffers.Debug; }
-	TDebugState* GetNextLocalDebugStateWrite() { return Buffers.Debug.WriteFrame( Buffers.Debug.HeadFrame() + 1 ); }
-	TNetworkSimBufferContainer<Model>* GetHistoricBuffers(bool bCreate=false)
-	{
-		if (HistoricBuffers.IsValid() == false && bCreate) { HistoricBuffers.Reset(new TNetworkSimBufferContainer<Model>()); }
-		return HistoricBuffers.Get();
-	}
-
-	TDebugBuffer* GetRemoteDebugBuffer() {	return &RepProxy_Debug.ReceivedBuffer; }
-#else
-	TDebugBuffer* GetLocalDebugBuffer() {	return nullptr; }
-	TDebugState* GetNextLocalDebugStateWrite() { return nullptr; }
-	TNetworkSimBufferContainer<Model>* GetHistoricBuffers(bool bCreate=false) { return nullptr; }
-	TDebugBuffer* GetRemoteDebugBuffer() {	return nullptr; }
-#endif
-
-private:
-
-#if NETSIM_MODEL_DEBUG
-	TRepProxyDebug RepProxy_Debug;
-	TUniquePtr<TNetworkSimBufferContainer<Model>> HistoricBuffers;
-#endif
-
 };
\ No newline at end of file
diff --git a/Engine/Plugins/Runtime/NetworkPrediction/Source/NetworkPrediction/Public/NetworkedSimulationModelBuffer.h b/Engine/Plugins/Runtime/NetworkPrediction/Source/NetworkPrediction/Public/NetworkedSimulationModelBuffer.h
--- a/Engine/Plugins/Runtime/NetworkPrediction/Source/NetworkPrediction/Public/NetworkedSimulationModelBuffer.h
+++ b/Engine/Plugins/Runtime/NetworkPrediction/Source/NetworkPrediction/Public/NetworkedSimulationModelBuffer.h
@@ -2,22 +2,7 @@
 
 #pragma once
 #include "Containers/UnrealString.h"
-
-// Generic iterator for iterating over each frame from tail to head
-template<typename TBuffer, typename ElementType>
-struct TNetworkSimBufferIterator
-{
-	TNetworkSimBufferIterator(TBuffer& InBuffer) : Buffer(InBuffer)	{ CurrentPos = Buffer.IterStartPos(); }
-	int32 Frame() const { return Buffer.IterFrame(CurrentPos); }
-	ElementType* Element() const { return Buffer.IterElement(CurrentPos); }
-
-	void operator++() { CurrentPos++; }
-	operator bool()	{ return CurrentPos != INDEX_NONE && CurrentPos <= Buffer.IterEndPos(); }
-
-private:
-	TBuffer& Buffer;
-	int32 CurrentPos = -1;
-};
+#include "Containers/ArrayView.h"
 
 // Helper struct to encapsulate optional, delayed writing of new element to buffer
 template<typename ElementType>
@@ -26,7 +11,7 @@ struct TNetSimLazyWriterFunc
 	template<typename TBuffer>
 	TNetSimLazyWriterFunc(TBuffer& Buffer, int32 PendingFrame)
 	{
-		GetFunc = [&Buffer, PendingFrame]() { return Buffer.WriteFrameInitializedFromHead(PendingFrame); };
+		GetFunc = [&Buffer, PendingFrame]() { return Buffer.WriteAtFrame(PendingFrame); };
 	}
 
 	ElementType* Get() const
@@ -54,115 +39,51 @@ struct TNetSimLazyWriter
 	TFunctionRef<void*()> GetFunc;
 };
 
-// Common implementations shared by TNetworkSimContiguousBuffer/TNetworkSimSparseBuffer
-template<typename T>
-struct TNetworkSimBufferBase
+
+// ---------------------------------------------------------------------------------------
+// ---------------------------------------------------------------------------------------
+template<typename ElementType> 
+struct TNetworkSimAuxBuffer
 {
-	FString GetBasicDebugStr() const
+	TNetworkSimAuxBuffer(int32 Capacity)
 	{
-		return FString::Printf(TEXT("Elements: [%d/%d]. Frames: [%d-%d]"), ((T*)this)->Num(), ((T*)this)->Max(), ((T*)this)->TailFrame(), ((T*)this)->HeadFrame());
+		Init(Capacity);
 	}
 
-	bool IsValidFrame(int32 Frame) const 
-	{ 
-		return Frame >= ((T*)this)->TailFrame() && Frame <= ((T*)this)->HeadFrame(); 
+	ElementType* operator[](int32 Frame)
+	{
+		return const_cast<ElementType*>(GetImpl(Frame)); 
 	}
-
-	// Copies each element of SourceBuffer into this buffer. This could break continuity so for TNetworkSimContiguousBuffer, it is possible to lose elements by doing this.
-	template<typename TSource>
-	void CopyAndMerge(const TSource& SourceBuffer)
+	
+	const ElementType* operator[](int32 Frame) const
 	{
-		for (auto It = SourceBuffer.CreateConstIterator(); It; ++It)
-		{
-			*((T*)this)->WriteFrame( It.Frame() ) = *It.Element();
-		}
+		return GetImpl(Frame); 
 	}
 
-protected:
-
-	// Creates a new frame, but
-	//	-If frame already exists, returns existing
-	//	-If frame > head, contents of head are copied into new frame
-	//	-If frame < tail, contents are cleared to default value
-	// (this function is sketchy and poorly named but is proving to be useful in a few places)
-	template<typename ElementType>
-	ElementType* WriteFrameInitializedFromHeadImpl(int32 Frame)
+	// Writes a new element at Frame. Copies previous frames contents into new element.
+	ElementType* WriteAtFrame(int32 Frame)
 	{
-		if (Frame > ((T*)this)->HeadFrame())
+		const int32 TailPos = FMath::Max(0, HeadPosition - Elements.Num() + 1);
+		int32 Pos = HeadPosition;
+		do
 		{
-			ElementType* HeadElementPtr = ((T*)this)->HeadElement();
-			ElementType* NewElement = ((T*)this)->WriteFrame(Frame);
-			if (HeadElementPtr)
+			TInternal& Data = Elements[Pos & IndexMask];
+			if (Data.Frame <= Frame)
 			{
-				*NewElement = *HeadElementPtr;
+				HeadPosition = Pos+1;
+				TInternal& NewData = Elements[HeadPosition & IndexMask];
+				NewData.Frame = Frame;
+				NewData.Element = Data.Element;
+				return &NewData.Element;
 			}
-			return NewElement;
-		}
-		else if (Frame < ((T*)this)->TailFrame())
-		{
-			ElementType* NewElement = ((T*)this)->WriteFrame(Frame);
-			*NewElement = ElementType();
-			return NewElement;
-		}
-		
-		return ((T*)this)->WriteFrame(Frame);
-	}
-};
+			Pos--;
+		} while(--Pos >= TailPos);
 
-template<typename T, int32 NumElements=32>
-struct TNetworkSimContiguousBuffer : public TNetworkSimBufferBase<TNetworkSimContiguousBuffer<T, NumElements>>
-{
-	using ElementType = T;
-	TNetworkSimContiguousBuffer()
-	{
-		Data.SetNum(NumElements);
-	}
-
-	T* operator[](int32 Frame) {	return const_cast<T*>(GetImpl(Frame)); }
-	const T* operator[](int32 Frame) const { return GetImpl(Frame); }
-
-	T* Get(int32 Frame) { return const_cast<T*>(GetImpl(Frame)); }
-	const T* Get(int32 Frame) const { return GetImpl(Frame); }
-
-	const T* HeadElement() const { return HeadElementImpl(); }
-	T* HeadElement() { return const_cast<T*>(HeadElementImpl()); }
-
-	const T* TailElement() const { return TailElementImpl(); }
-	T* TailElement() { return const_cast<T*>(TailElementImpl()); }
-
-	int32 HeadFrame() const { return Head; }
-	int32 TailFrame() const { return GetTail(); }
-
-	int32 Num() const { return Head == INDEX_NONE ? 0 : (Head - GetTail() + 1); }
-	int32 Max() const { return Data.Max(); }
-	int32 GetDirtyCount() const { return DirtyCount; }
-
-	// Returns the element @ frame for writing. The contents of this element are unknown (could be stale content!). Note the element returned is immediately considered "valid" by Num(), Iterators, etc!
-	// The written frame becomes the Head element. Essentially invalidates any existing frames >= Frame
-	T* WriteFrame(int32 Frame)
-	{
-		check(Frame >= 0);
-
-		const int32 Tail = GetTail();
-		if (Head == INDEX_NONE || Frame < Tail || Frame > Head+1)
-		{
-			// Writing outside the current range (+1) results in a full wipe of valid contents
-			NumValidElements = 1;
-		}
-		else
-		{
-			// Writing inside the current range (+1) preserves valid elements up to the size of our buffer
-			NumValidElements = FMath::Min<int32>(Frame - Tail + 1, Data.Num());
-		}
-		
-		Head = Frame;
-		++DirtyCount;
-		return &Data[Frame % Data.Num()];
-	}
-	
-	T* WriteFrameInitializedFromHead(int32 Frame)
-	{
-		return this->template WriteFrameInitializedFromHeadImpl<ElementType>(Frame);
+		HeadPosition = 0;
+		TInternal& NewData = Elements[0];
+		NewData.Element = ElementType();
+		NewData.Frame = Frame;
+		return &NewData.Element;
 	}
 
 	TNetSimLazyWriterFunc<ElementType> LazyWriter(int32 Frame)
@@ -170,205 +91,104 @@ struct TNetworkSimContiguousBuffer : public TNetworkSimBufferBase<TNetworkSimCon
 		return TNetSimLazyWriterFunc<ElementType>(*this, Frame);
 	}
 
-	using IteratorType = TNetworkSimBufferIterator<TNetworkSimContiguousBuffer<ElementType, NumElements>, ElementType>;
-	using ConstIteratorType = TNetworkSimBufferIterator<const TNetworkSimContiguousBuffer<ElementType, NumElements>, const ElementType>;
-
-	IteratorType CreateIterator() { return IteratorType(*this); }
-	ConstIteratorType CreateConstIterator() const { return ConstIteratorType(*this); }
-
 private:
 
-	const T* GetImpl(int32 Frame) const
+	struct TInternal
 	{
-		const int32 RelativeToHead = Frame - Head;
-		if (RelativeToHead > 0 || RelativeToHead <= -NumValidElements)
-		{
-			return nullptr;
-		}
-
-		return &Data[Frame % Data.Num()];
-	}
+		int32 Frame = -1;
+		ElementType Element;
+	};
 
-	const T* HeadElementImpl() const
+	void Init(int32 NewSize)
 	{
-		if (Head != INDEX_NONE)
-		{
-			return &Data[Head % Data.Num()];
-		}
-		return nullptr;
+		Elements.Reset();
+		Elements.AddDefaulted(FMath::RoundUpToPowerOfTwo(NewSize));
+		IndexMask = Elements.Num() - 1;	
 	}
 
-	const T* TailElementImpl() const
+	const ElementType* GetImpl(int32 Frame) const
 	{
-		const int32 Tail = GetTail();
-		if (Tail != INDEX_NONE)
+		const int32 TailPos = FMath::Max(0, HeadPosition - Elements.Num() + 1);
+		int32 Pos = HeadPosition;
+		do
 		{
-			return &Data[Tail % Data.Num()];
-		}
+			const TInternal& Data = Elements[Pos & IndexMask];
+			if (Data.Frame <= Frame)
+			{
+				return &Data.Element;
+			}
+			Pos--;
+		} while(--Pos >= TailPos);
+
 		return nullptr;
 	}
 
-	int32 GetTail() const { return Head == INDEX_NONE ? INDEX_NONE : (Head - NumValidElements + 1); };
+	// Latest element to be written
+	int32 HeadPosition = 0;
 
-	int32 DirtyCount = 0;
-	int32 Head = INDEX_NONE;
-	int32 NumValidElements = 0;
-	TArray<ElementType, TInlineAllocator<NumElements>> Data;
+	// Holds the mask for indexing the buffer's elements
+	int32 IndexMask = 0;
 
-	// Efficient position based accessing through TNetworkSimBufferIterator.
-	template<typename, typename>
-	friend struct TNetworkSimBufferIterator;
-	int32 IterStartPos() const { return GetTail(); }
-	int32 IterEndPos() const { return Head; }
-	int32 IterFrame(int32 Pos) const { return Pos; }
-	T* IterElement(int32 Pos) { return &Data[Pos % Data.Num()]; }
-	const T* IterElement(int32 Pos) const { return &Data[Pos % Data.Num()]; }
+	// Holds the buffer's elements
+	TArray<TInternal> Elements;
 };
 
-// Frame: arbitrary identifier for data. Not contiguous or controlled by us (always passed in)
-// Position/Pos: increasing counter for position in array. Mod with Data.Num to get Index.
-// Index/Idx: actual index into Data Array
-template<typename T, int32 NumElements=32>
-struct TNetworkSimSparseBuffer : public TNetworkSimBufferBase<TNetworkSimSparseBuffer<T, NumElements>>
+// ---------------------------------------------------------------------------------------
+// ---------------------------------------------------------------------------------------
+
+template<typename ElementType> 
+struct TNetworkSimFrameBuffer
 {
-	using ElementType = T;
-	TNetworkSimSparseBuffer()
+	TNetworkSimFrameBuffer(int32 Capacity)
 	{
-		Data.SetNum(NumElements);
+		check(Capacity > 0);
+		Init(Capacity);
 	}
-
-	T* operator[](int32 Frame) {	return const_cast<T*>(GetImpl(Frame)); }
-	const T* operator[](int32 Frame) const { return GetImpl(Frame); }
-
-	T* Get(int32 Frame) { return const_cast<T*>(GetImpl(Frame)); }
-	const T* Get(int32 Frame) const { return GetImpl(Frame); }
-
-	const T* HeadElement() const { return HeadElementImpl(); }
-	T* HeadElement() { return const_cast<T*>(HeadElementImpl()); }
-
-	const T* TailElement() const { return TailElementImpl(); }
-	T* TailElement() { return const_cast<T*>(TailElementImpl()); }
-
-	int32 HeadFrame() const { return Data[HeadPos % Data.Num()].Frame; }
-	int32 TailFrame() const { return Data[GetTailPos() % Data.Num()].Frame; }
-
-	int32 Num() const { return HeadPos == INDEX_NONE ? 0 : (HeadPos - GetTailPos() + 1); }
-	int32 Max() const { return Data.Max(); }
-	int32 GetDirtyCount() const { return DirtyCount; }
-
-	// Returns the element @ frame for writing. The contents of this element are unknown (could be stale content!). Note the element returned is immediately considered "valid" by Num(), Iterators, etc!
-	// The written frame becomes the Head element. Essentially invalidates any existing frames >= Frame
-	T* WriteFrame(int32 Frame)
+	
+	ElementType& operator[](int32 Index)
 	{
-		check(Frame >= 0);
-		++DirtyCount;
-
-		int32 Pos = HeadPos;
-		if (HeadPos != INDEX_NONE)
-		{
-			// If we have elements, find where this Frame would go
-			const int32 Tail = GetTailPos();
-			for (; Pos >= Tail; --Pos)
-			{
-				TInternal& InternalData = Data[Pos % Data.Num()];
-				if (InternalData.Frame == Frame)
-				{
-					HeadPos = Pos;
-					return &InternalData.Element;
-				}
-				if (InternalData.Frame <= Frame)
-				{
-					break;
-				}
-			}
-		}
-
-		// Write frame to next pos
-		HeadPos = Pos+1;
-		TInternal& WriteData = Data[HeadPos % Data.Num()];
-		WriteData.Frame = Frame;
-		return &WriteData.Element;
+		return Elements[Index & IndexMask];
 	}
-
-	ElementType* WriteFrameInitializedFromHead(int32 Frame)
+	
+	const ElementType& operator[](int32 Index) const
 	{
-		return this->template WriteFrameInitializedFromHeadImpl<ElementType>(Frame);
+		return Elements[Index & IndexMask];
 	}
-
-	TNetSimLazyWriterFunc<ElementType> LazyWriter(int32 Frame)
+	
+	int32 Capacity() const
 	{
-		return TNetSimLazyWriterFunc<ElementType>(*this, Frame);
+		return Elements.Num();
 	}
 
-	using IteratorType = TNetworkSimBufferIterator<TNetworkSimSparseBuffer<ElementType, NumElements>, ElementType>;
-	using ConstIteratorType = TNetworkSimBufferIterator<const TNetworkSimSparseBuffer<ElementType, NumElements>, const ElementType>;
-
-	IteratorType CreateIterator() { return IteratorType(*this); }
-	ConstIteratorType CreateConstIterator() const { return ConstIteratorType(*this); }
-
-private:
-
-	struct TInternal
+	// Resizes while preserving contents from head position
+	// Could be faster but doesn't really matter, this will be rare
+	void Resize(int32 NewSize, int32 PrevHead)
 	{
-		int32 Frame = -1;
-		T Element;
-	};
+		check(NewSize > 0);
+		TArray<ElementType> PrevElements = MoveTemp(Elements);
+		int32 PrevMask = IndexMask;
 
-	const T* GetImpl(int32 Frame) const
-	{
-		const int32 Pos = GetPosForFrame(Frame);
-		if (Pos != INDEX_NONE)
-		{		
-			return &Data[Pos % Data.Num()].Element;
-		}
+		Init(NewSize);
 
-		return nullptr;
-	}
-
-	const T* HeadElementImpl() const
-	{
-		if (HeadPos != INDEX_NONE)
+		int32 StartCopyFrom = PrevHead - FMath::Min(PrevElements.Num(), NewSize);
+		for (int32 i=StartCopyFrom; i <= PrevHead; ++i)
 		{
-			return &Data[HeadPos % Data.Num()].Element;
+			Elements[i & IndexMask] = PrevElements[i & PrevMask];
 		}
-		return nullptr;
 	}
 
-	const T* TailElementImpl() const
-	{
-		const int32 TailPos = GetTailPos();
-		if (TailPos != INDEX_NONE)
-		{
-			return &Data[TailPos % Data.Num()].Element;
-		}
-		return nullptr;
-	}
-	
-	int32 GetPosForFrame(int32 Frame) const
+private:
+
+	void Init(int32 NewSize)
 	{
-		const int32 TailPos = GetTailPos();
-		for (int32 Pos=HeadPos; Pos >= TailPos; --Pos)
-		{
-			if (Data[Pos % Data.Num()].Frame <= Frame)
-			{
-				return Pos;
-			}
-		}
-		return INDEX_NONE;
+		Elements.Reset();
+		Elements.AddDefaulted(FMath::RoundUpToPowerOfTwo(NewSize));
+		IndexMask = Elements.Num() - 1;	
 	}
 
-	int32 GetTailPos() const { return HeadPos == INDEX_NONE ? INDEX_NONE : FMath::Max<int32>(0, HeadPos - Data.Num()); }
-
-	int32 DirtyCount=0;
-	int32 HeadPos=INDEX_NONE;
-	TArray<TInternal, TInlineAllocator<NumElements>> Data;
+	// Holds the mask for indexing the buffer's elements.
+	int32 IndexMask = 0;
 
-	// Efficient position based accessing through TNetworkSimBufferIterator.
-	template<typename, typename>
-	friend struct TNetworkSimBufferIterator;
-	int32 IterStartPos() const { return GetTailPos(); }
-	int32 IterEndPos() const { return HeadPos; }
-	int32 IterFrame(int32 Pos) const { return Data[Pos % Data.Num()].Frame; }
-	T* IterElement(int32 Pos) { return &Data[Pos % Data.Num()].Element; }
-	const T* IterElement(int32 Pos) const { return &Data[Pos % Data.Num()].Element; }
-};
\ No newline at end of file
+	// Holds the buffer's elements.
+	TArray<ElementType> Elements;
+};
diff --git a/Engine/Plugins/Runtime/NetworkPrediction/Source/NetworkPrediction/Public/NetworkedSimulationModelBufferTypes.h b/Engine/Plugins/Runtime/NetworkPrediction/Source/NetworkPrediction/Public/NetworkedSimulationModelBufferTypes.h
--- a/Engine/Plugins/Runtime/NetworkPrediction/Source/NetworkPrediction/Public/NetworkedSimulationModelBufferTypes.h
+++ b/Engine/Plugins/Runtime/NetworkPrediction/Source/NetworkPrediction/Public/NetworkedSimulationModelBufferTypes.h
@@ -1,6 +1,7 @@
 // Copyright Epic Games, Inc. All Rights Reserved.
 #pragma once
 #include "NetworkedSimulationModelTick.h"
+#include "Misc/EnumClassFlags.h"
 
 // ---------------------------------------------------------------------------------------------------------------------
 //	BufferTypes: template helpers for addressing the different buffer types of the system.
@@ -11,8 +12,7 @@ enum class ENetworkSimBufferTypeId : uint8
 {
 	Input,
 	Sync,
-	Aux,
-	Debug
+	Aux
 };
 
 inline FString LexToString(ENetworkSimBufferTypeId A)
@@ -22,7 +22,6 @@ inline FString LexToString(ENetworkSimBufferTypeId A)
 		case ENetworkSimBufferTypeId::Input: return TEXT("Input");
 		case ENetworkSimBufferTypeId::Sync: return TEXT("Sync");
 		case ENetworkSimBufferTypeId::Aux: return TEXT("Aux");
-		case ENetworkSimBufferTypeId::Debug: return TEXT("Debug");
 	};
 	return TEXT("Unknown");
 }
@@ -52,12 +51,6 @@ struct TSelectTypeHelper<TBufferTypes, ENetworkSimBufferTypeId::Aux>
 	using type = typename TBufferTypes::TAuxState;
 };
 
-template<typename TBufferTypes>
-struct TSelectTypeHelper<TBufferTypes, ENetworkSimBufferTypeId::Debug>
-{
-	using type = typename TBufferTypes::TDebugState;
-};
-
 template<typename T>
 struct HasNetSerialize
 {
@@ -77,80 +70,87 @@ struct HasLog
 };
 
 // A collection of the system's buffer types. This allows us to collapse the 4 types into a single type to use a template argument elsewhere.
-template<typename InInputCmd, typename InSyncState, typename InAuxState, typename InDebugState = FNetSimProcessedFrameDebugInfo>
+template<typename InInputCmd, typename InSyncState, typename InAuxState>
 struct TNetworkSimBufferTypes
 {
 	// Quick access to types when you know what you want
 	using TInputCmd = InInputCmd;
 	using TSyncState = InSyncState;
 	using TAuxState = InAuxState;
-	using TDebugState = InDebugState;
 
 	// Template access via ENetworkSimBufferTypeId when "which buffer" is parameterized
 	template<ENetworkSimBufferTypeId Id>
 	struct select_type
 	{
-		using type = typename TSelectTypeHelper< TNetworkSimBufferTypes<TInputCmd, TSyncState, TAuxState, TDebugState>, Id >::type;
+		using type = typename TSelectTypeHelper< TNetworkSimBufferTypes<TInputCmd, TSyncState, TAuxState>, Id >::type;
 	};
 
 	// Must implement NetSerialize and Log functions. This purposefully does not pass on inherited methods
 	static_assert(HasNetSerialize<InInputCmd>::Has == true, "InputCmd Must implement NetSerialize");
 	static_assert(HasNetSerialize<InSyncState>::Has == true, "SyncState Must implement NetSerialize");
 	static_assert(HasNetSerialize<InAuxState>::Has == true, "AuxState Must implement NetSerialize");
-	static_assert(HasNetSerialize<InDebugState>::Has == true, "DebugState Must implement NetSerialize");
 
 	static_assert(HasLog<InInputCmd>::Has == true, "InputCmd Must implement Log");
 	static_assert(HasLog<InSyncState>::Has == true, "SyncState Must implement Log");
 	static_assert(HasLog<InAuxState>::Has == true, "AuxState Must implement Log");
-	static_assert(HasLog<InDebugState>::Has == true, "DebugState Must implement Log");
 };
 
-// ----------------------------------------------------------------------------------------------------------------------------------------------
-//	FrameCmd - in variable tick simulations we store the timestep of each frame with the input. TFrameCmd wraps the user struct to do this.
-// ----------------------------------------------------------------------------------------------------------------------------------------------
+// ------------------------------------------------
 
-// Wraps an input command (BaseType) in a NetworkSimulation time. 
-template<typename BaseCmdType, typename TickSettings, bool IsFixedTick=(TickSettings::FixedStepMS!=0)>
-struct TFrameCmd : public BaseCmdType
+// Wraps FrameDelta time. In fixed tick mode, this is specialized to be empty and return the known-at-compile-time step
+template<typename TickSettings, bool IsFixedTick=(TickSettings::FixedStepMS!=0)>
+struct TFrameDeltaTimeWrapper
 {
-	FNetworkSimTime GetFrameDeltaTime() const { return FrameDeltaTime; }
-	void SetFrameDeltaTime(const FNetworkSimTime& InTime) { FrameDeltaTime = InTime; }
-	void NetSerialize(const FNetSerializeParams& P)
-	{
-		FrameDeltaTime.NetSerialize(P.Ar);
-		BaseCmdType::NetSerialize(P); 
-	}
-
-	void Log(FStandardLoggingParameters& P) const { BaseCmdType::Log(P); }
+	FNetworkSimTime Get() const { return FrameDeltaTime; }
+	void Set(const FNetworkSimTime& InTime) { FrameDeltaTime = InTime; }
+	void NetSerialize(const FNetSerializeParams& P)	{ FrameDeltaTime.NetSerialize(P.Ar); }
 
 private:
 	FNetworkSimTime FrameDeltaTime;
 };
 
-// Fixed tick specialization
-template<typename BaseCmdType, typename TickSettings>
-struct TFrameCmd<BaseCmdType, TickSettings, true> : public BaseCmdType
+template<typename TickSettings>
+struct TFrameDeltaTimeWrapper<TickSettings, true>
 {
-	FNetworkSimTime GetFrameDeltaTime() const { return FNetworkSimTime::FromMSec(TickSettings::GetFixedStepMS()); }
-	void SetFrameDeltaTime(const FNetworkSimTime& InTime) { }
-	void NetSerialize(const FNetSerializeParams& P) { BaseCmdType::NetSerialize(P); }
-	void Log(FStandardLoggingParameters& P) const { BaseCmdType::Log(P); }
+	FNetworkSimTime Get() const { return FNetworkSimTime::FromMSec(TickSettings::GetFixedStepMS()); }
+	void Set(const FNetworkSimTime& InTime) { }
+	void NetSerialize(const FNetSerializeParams& P) { }
 };
 
-// ----------------------------------------------------------------------------------------------------------------------------------------------
-// "Internal" Buffers: Helper to turn user supplied buffer types into the "real" buffer types: the InputCmd struct is wrapped in TFrameCmd
-//		-Outside/User code will always use the user supplied buffer types in callbacks/APIs etc
-//		-Internal code to the TNetworkedSimulationModel will use the internal buffer types
-// ----------------------------------------------------------------------------------------------------------------------------------------------
-template<typename TUserBufferTypes, typename TTickSettings>
-struct TInternalBufferTypes : TNetworkSimBufferTypes< 
-	
-	// InputCmds are wrapped in TFrameCmd, which will store an explicit sim delta time if we are not a fixed tick sim
-	TFrameCmd< typename TUserBufferTypes::TInputCmd , TTickSettings>,
-
-	typename TUserBufferTypes::TSyncState,	// SyncState passes through
-	typename TUserBufferTypes::TAuxState,	// Auxstate passes through
-	typename TUserBufferTypes::TDebugState	// Debugstate passes through
->
+// The actual per-frame data the system keeps track of
+//template<typename TUserBufferTypes, typename TTickSettings>
+
+enum class ESimulationFrameStateFlags : uint8
 {
-};
\ No newline at end of file
+	InputWritten	= 1 << 0,	// Input has been written
+};
+
+ENUM_CLASS_FLAGS(ESimulationFrameStateFlags)
+
+// This is what we actually store in contiguous memory, per frame.
+template<typename Model>
+struct TSimulationFrameStateBase
+{
+	using BufferTypes = typename Model::BufferTypes;
+	using TickSettings = typename Model::TickSettings;
+
+	using TInputCmd	= typename BufferTypes::TInputCmd;
+	using TSyncState = typename BufferTypes::TSyncState;
+
+	uint8 Flags = 0;
+	FNetworkSimTime	TotalSimTime;
+	TFrameDeltaTimeWrapper<TickSettings> FrameDeltaTime;
+
+	TInputCmd	InputCmd;
+	TSyncState	SyncState;
+
+	bool HasFlag(const ESimulationFrameStateFlags& F) const { return (Flags & (uint8)F) > 0; }
+	void SetFlag(const ESimulationFrameStateFlags& F) { Flags |= (uint8)F; }
+	void ClearFlags() { Flags = 0; };
+	void ResetFlagsTo(const ESimulationFrameStateFlags& F) { Flags = (uint8)F; }
+};
+
+// Outward facing struct. This allows user simulations to extend the internal FrameState struct if needed.
+// Use case would be custom RepController logic that needed to do additional per-frame state tracking.
+template<typename Model>
+struct TSimulationFrameState : TSimulationFrameStateBase<Model> { };
\ No newline at end of file
diff --git a/Engine/Plugins/Runtime/NetworkPrediction/Source/NetworkPrediction/Public/NetworkedSimulationModelCVars.h b/Engine/Plugins/Runtime/NetworkPrediction/Source/NetworkPrediction/Public/NetworkedSimulationModelCVars.h
--- a/Engine/Plugins/Runtime/NetworkPrediction/Source/NetworkPrediction/Public/NetworkedSimulationModelCVars.h
+++ b/Engine/Plugins/Runtime/NetworkPrediction/Source/NetworkPrediction/Public/NetworkedSimulationModelCVars.h
@@ -2,6 +2,7 @@
 
 #pragma once
 #include "HAL/IConsoleManager.h"
+#include "Logging/LogMacros.h"
 
 // ------------------------------------------------------------------------------------------------------------
 //	"Shipping const" cvars: cvars that should compile out to const functions in shipping/test builds
@@ -10,6 +11,12 @@
 //		a bit slower but shouldn't matter since it is compiled out of shipping/test.
 // ------------------------------------------------------------------------------------------------------------
 
+inline IConsoleVariable* FindConsoleVarHelper(const TCHAR* VarName)
+{
+	UE_LOG(LogTemp, Warning, TEXT("RATTI Finding... %s"), VarName);
+	return IConsoleManager::Get().FindConsoleVariable(VarName, false);
+}
+
 // Whether to treat these cvars as consts
 #define NETSIM_CONST_CVARS (UE_BUILD_SHIPPING || UE_BUILD_TEST)
 
@@ -19,16 +26,20 @@ struct FConditionalAutoConsoleRegister
 {
 	FConditionalAutoConsoleRegister(const TCHAR* VarName, int32 Value, const TCHAR* Help)
 	{
+		UE_LOG(LogTemp, Warning, TEXT("RATTI FConditionalAutoConsoleRegister() %s"), VarName);
 		if (!IConsoleManager::Get().FindConsoleVariable(VarName, false))
 		{
+			UE_LOG(LogTemp, Warning, TEXT("RATTI RegisterConsoleVariable() %s"), VarName);
 			IConsoleManager::Get().RegisterConsoleVariable(VarName, Value, Help, ECVF_Cheat);
 		}
 	}
 
 	FConditionalAutoConsoleRegister(const TCHAR* VarName, float Value, const TCHAR* Help)
 	{
+		UE_LOG(LogTemp, Warning, TEXT("RATTI FConditionalAutoConsoleRegister() %s"), VarName);
 		if (!IConsoleManager::Get().FindConsoleVariable(VarName, false))
 		{
+			UE_LOG(LogTemp, Warning, TEXT("RATTI RegisterConsoleVariable() %s"), VarName);
 			IConsoleManager::Get().RegisterConsoleVariable(VarName, Value, Help, ECVF_Cheat);
 		}
 	}
@@ -44,13 +55,13 @@ struct FConditionalAutoConsoleRegister
 	static FConditionalAutoConsoleRegister Var##Auto(TEXT(VarName),(int32)Value,TEXT(Help)); \
 	inline int32 Var() \
 	{ \
-		static const auto* Existing = IConsoleManager::Get().FindConsoleVariable(TEXT(VarName), false); \
+		static const auto* Existing = FindConsoleVarHelper(TEXT(VarName)); \
 		check(Existing); \
 		return Existing->GetInt(); \
 	} \
 	inline void Set##Var(int32 V) \
 	{ \
-		static auto* Existing = IConsoleManager::Get().FindConsoleVariable(TEXT(VarName), false); \
+		static auto* Existing = FindConsoleVarHelper(TEXT(VarName)); \
 		check(Existing); \
 		Existing->Set(V, ECVF_SetByConsole); \
 	}
@@ -65,13 +76,13 @@ struct FConditionalAutoConsoleRegister
 	static FConditionalAutoConsoleRegister Var##Auto(TEXT(VarName),(float)Value,TEXT(Help)); \
 	inline float Var() \
 	{ \
-		static const auto* Existing = IConsoleManager::Get().FindConsoleVariable(TEXT(VarName), false); \
+		static const auto* Existing = FindConsoleVarHelper(TEXT(VarName)); \
 		check(Existing); \
 		return Existing->GetFloat(); \
 	} \
 	inline void Set##Var(float V) \
 	{ \
-		static auto* Existing = IConsoleManager::Get().FindConsoleVariable(TEXT(VarName), false); \
+		static auto* Existing = FindConsoleVarHelper(TEXT(VarName)); \
 		check(Existing); \
 		Existing->Set(V, ECVF_SetByConsole); \
 	}
diff --git a/Engine/Plugins/Runtime/NetworkPrediction/Source/NetworkPrediction/Public/NetworkedSimulationModelDebugger.h b/Engine/Plugins/Runtime/NetworkPrediction/Source/NetworkPrediction/Public/NetworkedSimulationModelDebugger.h
--- a/Engine/Plugins/Runtime/NetworkPrediction/Source/NetworkPrediction/Public/NetworkedSimulationModelDebugger.h
+++ b/Engine/Plugins/Runtime/NetworkPrediction/Source/NetworkPrediction/Public/NetworkedSimulationModelDebugger.h
@@ -154,7 +154,6 @@ template <typename TNetSimModel>
 struct TNetworkSimulationModelDebugger : public INetworkSimulationModelDebugger
 {
 	using TSimTime = typename TNetSimModel::TSimTime;
-	using TDebugState = typename TNetSimModel::TDebugState;
 
 	TNetworkSimulationModelDebugger(TNetSimModel* InNetSim, const AActor* OwningActor)
 	{
@@ -187,7 +186,7 @@ struct TNetworkSimulationModelDebugger : public INetworkSimulationModelDebugger
 		static float Pad = 2.f;
 		static float BaseLineYPCT = 0.8f;
 
-		auto& InputBuffer = NetworkSim->GetHistoricBuffers() ? NetworkSim->GetHistoricBuffers()->Input : NetworkSim->Buffers.Input;
+		auto& InputBuffer = NetworkSim->GetHistoricBuffers() ? NetworkSim->GetHistoricBuffers()->Input : NetworkSim->State.Input;
 
 		if (Canvas && DebugBuffer && DebugBuffer->Num() > 0)
 		{
@@ -310,7 +309,7 @@ struct TNetworkSimulationModelDebugger : public INetworkSimulationModelDebugger
 	template<typename TBuffer>
 	void CalcMaxColumnFrameTime(TBuffer* DebugBuffer, float& MaxInputTime, float& MaxLocalFrameTime)
 	{
-		auto& InputBuffer = NetworkSim->GetHistoricBuffers() ? NetworkSim->GetHistoricBuffers()->Input : NetworkSim->Buffers.Input;
+		auto& InputBuffer = NetworkSim->GetHistoricBuffers() ? NetworkSim->GetHistoricBuffers()->Input : NetworkSim->State.Input;
 
 		for (auto It = DebugBuffer->CreateIterator(); It; ++It)
 		{
@@ -351,7 +350,7 @@ struct TNetworkSimulationModelDebugger : public INetworkSimulationModelDebugger
 		// ------------------------------------------------------------------------------------------------------------------------------------------------
 
 		Out.Emit(FString::Printf(TEXT("%s - %s"), *Owner->GetName(), *UEnum::GetValueAsString(TEXT("Engine.ENetRole"), Owner->GetLocalRole())), FColor::Yellow);
-		Out.Emit(FString::Printf(TEXT("PendingFrame: %d (%d Buffered)"), NetworkSim->Ticker.PendingFrame, NetworkSim->Buffers.Input.HeadFrame() - NetworkSim->Ticker.PendingFrame));
+		Out.Emit(FString::Printf(TEXT("PendingFrame: %d (%d Buffered)"), NetworkSim->State.GetPendingTickFrame(), NetworkSim->State.GetLatestInputFrame() - NetworkSim->State.GetPendingTickFrame()));
 				
 		if (Owner->GetLocalRole() == ROLE_AutonomousProxy)
 		{			
@@ -360,18 +359,18 @@ struct TNetworkSimulationModelDebugger : public INetworkSimulationModelDebugger
 
 			const int32 LastSerializedFrame = NetworkSim->RepProxy_Autonomous.GetLastSerializedFrame();
 
-			// Calc how much predicted time we have processed. Note that we use the motionstate buffer to iterate but the MS is on the input cmd. (if we are buffering cmds, don't want to count them)
-			
+			// Calc how much predicted time we have processed.				
 			TSimTime PredictedMS;
-			for (int32 PredFrame = LastSerializedFrame+1; PredFrame <= NetworkSim->Buffers.Sync.HeadFrame(); ++PredFrame)
+
+			for (int32 PredFrame = LastSerializedFrame+1; PredFrame <= NetworkSim->State.GetPendingTickFrame(); ++PredFrame)
 			{
-				if (auto* Cmd = NetworkSim->Buffers.Input[PredFrame])
+				if (auto* Cmd = NetworkSim->State.GetFrameState(PredFrame))
 				{
-					PredictedMS += Cmd->GetFrameDeltaTime();
+					PredictedMS += Cmd->FrameDeltaTime.Get();
 				}
 			}
 
-			FString ConfirmedFrameStr = FString::Printf(TEXT("LastConfirmedFrame: %d. Prediction: %d Frames, %s MS"), LastSerializedFrame, NetworkSim->Buffers.Sync.HeadFrame() - LastSerializedFrame, *PredictedMS.ToString());
+			FString ConfirmedFrameStr = FString::Printf(TEXT("LastConfirmedFrame: %d. Prediction: %d Frames, %s MS"), LastSerializedFrame, NetworkSim->State.GetPendingTickFrame() - LastSerializedFrame, *PredictedMS.ToString());
 			if (FaultDetected)
 			{
 				ConfirmedFrameStr += TEXT(" RECONCILE FAULT DETECTED!");
@@ -380,92 +379,19 @@ struct TNetworkSimulationModelDebugger : public INetworkSimulationModelDebugger
 
 			Out.Emit(MoveTemp(ConfirmedFrameStr), Color);
 
-			FString SimulationTimeString = FString::Printf(TEXT("Local SimulationTime: %s. SerializedSimulationTime: %s. Difference MS: %s"), *NetworkSim->Ticker.GetTotalProcessedSimulationTime().ToString(),
-				*NetworkSim->RepProxy_Autonomous.GetLastSerializedSimTime().ToString(), *(NetworkSim->Ticker.GetTotalProcessedSimulationTime() - NetworkSim->RepProxy_Autonomous.GetLastSerializedSimTime()).ToString());
+			FString SimulationTimeString = FString::Printf(TEXT("Local SimulationTime: %s. SerializedSimulationTime: %s. Difference MS: %s"), *NetworkSim->State.GetTotalProcessedSimulationTime().ToString(),
+				*NetworkSim->RepProxy_Autonomous.GetLastSerializedSimTime().ToString(), *(NetworkSim->State.GetTotalProcessedSimulationTime() - NetworkSim->RepProxy_Autonomous.GetLastSerializedSimTime()).ToString());
 			Out.Emit(MoveTemp(SimulationTimeString), Color);
 
-			FString AllowedSimulationTimeString = FString::Printf(TEXT("Allowed Simulation Time: %s. Frame: %d/%d/%d"), *NetworkSim->Ticker.GetRemainingAllowedSimulationTime().ToString(), NetworkSim->Ticker.MaxAllowedFrame, NetworkSim->Ticker.PendingFrame, NetworkSim->Buffers.Input.HeadFrame());
+			FString AllowedSimulationTimeString = FString::Printf(TEXT("Frame: %d/%d/%d"), NetworkSim->State.GetMaxTickFrame(), NetworkSim->State.GetPendingTickFrame(), NetworkSim->State.GetLatestInputFrame());
 			Out.Emit(MoveTemp(AllowedSimulationTimeString), Color);
 		}
 		else if (Owner->GetLocalRole() == ROLE_SimulatedProxy)
 		{
 			FColor Color = FColor::White;
-			FString TimeString = FString::Printf(TEXT("Total Processed Simulation Time: %s. Last Serialized Simulation Time: %s. Delta: %s"), *NetworkSim->Ticker.GetTotalProcessedSimulationTime().ToString(), *NetworkSim->RepProxy_Simulated.GetLastSerializedSimulationTime().ToString(), *(NetworkSim->RepProxy_Simulated.GetLastSerializedSimulationTime() - NetworkSim->Ticker.GetTotalProcessedSimulationTime()).ToString());
+			FString TimeString = FString::Printf(TEXT("Total Processed Simulation Time: %s. Last Serialized Simulation Time: %s. Delta: %s"), *NetworkSim->State.GetTotalProcessedSimulationTime().ToString(), *NetworkSim->RepProxy_Simulated.GetLastSerializedSimulationTime().ToString(), *(NetworkSim->RepProxy_Simulated.GetLastSerializedSimulationTime() - NetworkSim->State.GetTotalProcessedSimulationTime()).ToString());
 			Out.Emit(MoveTemp(TimeString), Color);
 		}
-
-		auto EmitBuffer = [&Out](const FString& BufferName, auto& Buffer)
-		{
-			Out.Emit();
-			Out.Emit(FString::Printf(TEXT("//////////////// %s ///////////////"), *BufferName), FColor::Yellow);
-			Out.Emit(FString::Printf(TEXT("%s"), *Buffer.GetBasicDebugStr()));		
-			Out.Emit();
-			Out.EmitElement(Buffer, FStandardLoggingParameters(nullptr, EStandardLoggingContext::Full, Buffer.HeadFrame()));
-		};
-
-		EmitBuffer(TEXT("InputBuffer"), NetworkSim->Buffers.Input);
-		EmitBuffer(TEXT("SyncBuffer"), NetworkSim->Buffers.Sync);
-		EmitBuffer(TEXT("AuxBuffer"), NetworkSim->Buffers.Aux);
-
-		// ------------------------------------------------------------------------------------------------------------------------------------------------
-		//	Canvas
-		// ------------------------------------------------------------------------------------------------------------------------------------------------
-
-		if (Canvas)
-		{
-			FRect ServerRect;
-			ServerRect.Min.X = 0.30f * (float)Canvas->SizeX;
-			ServerRect.Max.X = 0.95f * (float)Canvas->SizeX;
-			ServerRect.Min.Y = 0.05f * (float)Canvas->SizeY;
-			ServerRect.Max.Y = 0.45f * (float)Canvas->SizeY;
-
-			FRect ClientRect;
-			ClientRect.Min.X = 0.30f * (float)Canvas->SizeX;
-			ClientRect.Max.X = 0.95f * (float)Canvas->SizeX;
-			ClientRect.Min.Y = 0.55f * (float)Canvas->SizeY;
-			ClientRect.Max.Y = 0.95f * (float)Canvas->SizeY;
-
-			float MaxColumnTime = 1.f/60.f;
-			float MaxLocalFrameTime = 1.f/60.f;
-			CalcMaxColumnFrameTime(NetworkSim->GetRemoteDebugBuffer(), MaxColumnTime, MaxLocalFrameTime);
-			CalcMaxColumnFrameTime(NetworkSim->GetLocalDebugBuffer(), MaxColumnTime, MaxLocalFrameTime);
-
-			FCachedScreenPositionMap ServerToClientCache;
-			FCachedScreenPositionMap ClientToServerCache;
-
-			GatherDebugGraph(Out, Canvas, NetworkSim->GetRemoteDebugBuffer(), ServerRect, MaxColumnTime, MaxLocalFrameTime, TEXT("Server"), ServerToClientCache, ClientToServerCache);
-			GatherDebugGraph(Out, Canvas, NetworkSim->GetLocalDebugBuffer(), ClientRect, MaxColumnTime, MaxLocalFrameTime, TEXT("Client"), ClientToServerCache, ServerToClientCache);
-
-			// Network Send/Recv lines
-			if (NetworkSimulationModelDebugCVars::DrawNetworkSendLines() > 0)
-			{
-				for (auto& It : ServerToClientCache.Frames)
-				{
-					const int32 Frame = It.Key;
-					auto& Positions = It.Value;
-					if (Frame != 0 && Positions.RecvPosition != FVector2D::ZeroVector && Positions.SentPosition != FVector2D::ZeroVector)
-					{
-						Out.EmitLine(Positions.SentPosition, Positions.RecvPosition, FColor::Purple);
-
-						FVector2D TextPos = Positions.SentPosition + (0.25f * (Positions.RecvPosition - Positions.SentPosition));
-						Out.EmitText(TextPos, FColor::Purple, LexToString(Frame));
-					}
-				}
-
-				for (auto& It : ClientToServerCache.Frames)
-				{
-					const int32 Frame = It.Key;
-					auto& Positions = It.Value;
-					if (Frame != 0 && Positions.RecvPosition != FVector2D::ZeroVector && Positions.SentPosition != FVector2D::ZeroVector)
-					{
-						Out.EmitLine(Positions.SentPosition, Positions.RecvPosition, FColor::Orange);
-
-						FVector2D TextPos = Positions.SentPosition + (0.25f * (Positions.RecvPosition - Positions.SentPosition));
-						Out.EmitText(TextPos, FColor::Orange, LexToString(Frame));
-					}
-				}
-			}
-		}
 	}
 
 	virtual void Tick( float DeltaTime ) override
@@ -478,46 +404,52 @@ struct TNetworkSimulationModelDebugger : public INetworkSimulationModelDebugger
 
 		UWorld* World = Owner->GetWorld();
 
-		if (auto* LatestSync = NetworkSim->Buffers.Sync.HeadElement())
 		{
-			FVisualLoggingParameters VLogParams(EVisualLoggingContext::LastPredicted, NetworkSim->Buffers.Sync.HeadFrame(), EVisualLoggingLifetime::Transient);
-			int32 SyncFrame = NetworkSim->Buffers.Sync.HeadFrame();
-			NetworkSim->Driver->InvokeVisualLog(NetworkSim->Buffers.Input[SyncFrame], LatestSync, NetworkSim->Buffers.Aux[SyncFrame], VLogParams);
+			const int32 HeadFrame = NetworkSim->State.GetPendingTickFrame();
+			if (auto* LatestFrameState = NetworkSim->State.GetFrameState(HeadFrame))
+			{
+				auto* LatestSync = &LatestFrameState->SyncState;
+				FVisualLoggingParameters VLogParams(EVisualLoggingContext::LastPredicted, HeadFrame, EVisualLoggingLifetime::Transient);
+				NetworkSim->Driver->InvokeVisualLog(&LatestFrameState->InputCmd, LatestSync, NetworkSim->State.GetAuxState(HeadFrame), VLogParams);
+			}
 		}
 
-		FStuff ServerPIEStuff = GetServerPIEStuff();
-		if (ServerPIEStuff.NetworkSim)
 		{
-			if (auto* ServerLatestSync = ServerPIEStuff.NetworkSim->Buffers.Sync.HeadElement())
+			FStuff ServerPIEStuff = GetServerPIEStuff();
+			if (ServerPIEStuff.NetworkSim)
 			{
-				FVisualLoggingParameters VLogParams(EVisualLoggingContext::CurrentServerPIE, ServerPIEStuff.NetworkSim->Buffers.Sync.HeadFrame(), EVisualLoggingLifetime::Transient);
-				int32 SyncFrame = ServerPIEStuff.NetworkSim->Buffers.Sync.HeadFrame();
-				NetworkSim->Driver->InvokeVisualLog(ServerPIEStuff.NetworkSim->Buffers.Input[SyncFrame], ServerLatestSync, ServerPIEStuff.NetworkSim->Buffers.Aux[SyncFrame], VLogParams);
+				const int32 HeadFrame = ServerPIEStuff.NetworkSim->State.GetPendingTickFrame();
+				if (auto* LatestFrameState = NetworkSim->State.GetFrameState(HeadFrame))
+				{
+					FVisualLoggingParameters VLogParams(EVisualLoggingContext::CurrentServerPIE, HeadFrame, EVisualLoggingLifetime::Transient);
+					NetworkSim->Driver->InvokeVisualLog(&LatestFrameState->InputCmd, &LatestFrameState->SyncState, ServerPIEStuff.NetworkSim->State.GetAuxState(HeadFrame), VLogParams);
+				}
 			}
 		}
 
 		if (Owner->GetLocalRole() == ROLE_AutonomousProxy)
 		{
-			for (int32 Frame = NetworkSim->RepProxy_Autonomous.GetLastSerializedFrame(); Frame < NetworkSim->Buffers.Sync.HeadFrame(); ++Frame)
+			for (int32 Frame = NetworkSim->RepProxy_Autonomous.GetLastSerializedFrame(); Frame < NetworkSim->State.GetPendingTickFrame(); ++Frame)
 			{
-				if (auto* SyncState = NetworkSim->Buffers.Sync[Frame])
+				if (auto* FrameState = NetworkSim->State.GetFrameState(Frame))
 				{
 					const EVisualLoggingContext Context = (Frame == NetworkSim->RepProxy_Autonomous.GetLastSerializedFrame()) ? EVisualLoggingContext::LastConfirmed : EVisualLoggingContext::OtherPredicted;
-					FVisualLoggingParameters VLogParams(Context, NetworkSim->Buffers.Sync.HeadFrame(), EVisualLoggingLifetime::Transient);
-					NetworkSim->Driver->InvokeVisualLog(NetworkSim->Buffers.Input[Frame], SyncState, NetworkSim->Buffers.Aux[Frame], VLogParams);
+					FVisualLoggingParameters VLogParams(Context, Frame, EVisualLoggingLifetime::Transient);
+					NetworkSim->Driver->InvokeVisualLog(&FrameState->InputCmd, &FrameState->SyncState, NetworkSim->State.GetAuxState(Frame), VLogParams);
 				}
 			}
 		}
 		else if (Owner->GetLocalRole() == ROLE_SimulatedProxy)
 		{
 			{
-				FVisualLoggingParameters VLogParams(EVisualLoggingContext::LastPredicted, NetworkSim->Buffers.Sync.HeadFrame(), EVisualLoggingLifetime::Transient);
-				int32 HeadFrame = NetworkSim->Buffers.Sync.HeadFrame();
-				NetworkSim->Driver->InvokeVisualLog(NetworkSim->Buffers.Input[HeadFrame], NetworkSim->Buffers.Sync[HeadFrame], NetworkSim->Buffers.Aux[HeadFrame], VLogParams);
+				int32 HeadFrame = NetworkSim->State.GetPendingTickFrame();
+				FVisualLoggingParameters VLogParams(EVisualLoggingContext::LastPredicted, HeadFrame, EVisualLoggingLifetime::Transient);
+				auto* FrameState = NetworkSim->State.GetFrameState(HeadFrame);
+				NetworkSim->Driver->InvokeVisualLog(&FrameState->InputCmd, &FrameState->SyncState, NetworkSim->State.GetAuxState(HeadFrame), VLogParams);
 			}
 
 			{
-				FVisualLoggingParameters VLogParams(EVisualLoggingContext::LastConfirmed, NetworkSim->Buffers.Sync.HeadFrame(), EVisualLoggingLifetime::Transient);
+				FVisualLoggingParameters VLogParams(EVisualLoggingContext::LastConfirmed, NetworkSim->RepProxy_Simulated.GetLastSerializedFrame(), EVisualLoggingLifetime::Transient);
 				NetworkSim->Driver->InvokeVisualLog(&NetworkSim->RepProxy_Simulated.GetLastSerializedInputCmd(), &NetworkSim->RepProxy_Simulated.GetLastSerializedSyncState(), &NetworkSim->RepProxy_Simulated.GetLastSerializedAuxState(), VLogParams);
 			}
 			
@@ -526,7 +458,7 @@ struct TNetworkSimulationModelDebugger : public INetworkSimulationModelDebugger
 				FVector2D ServerSimulationTimeData(Owner->GetWorld()->GetTimeSeconds(), NetworkSim->RepProxy_Simulated.GetLastSerializedSimulationTime().ToRealTimeMS());
 				UE_VLOG_HISTOGRAM(Owner, LogNetworkSimDebug, Log, "Simulated Time Graph", "Serialized Simulation Time", ServerSimulationTimeData);
 
-				FVector2D LocalSimulationTimeData(Owner->GetWorld()->GetTimeSeconds(), NetworkSim->Ticker.GetTotalProcessedSimulationTime().ToRealTimeMS());
+				FVector2D LocalSimulationTimeData(Owner->GetWorld()->GetTimeSeconds(), NetworkSim->State.GetTotalProcessedSimulationTime().ToRealTimeMS());
 				UE_VLOG_HISTOGRAM(Owner, LogNetworkSimDebug, Log, "Simulated Time Graph", "Local Simulation Time", LocalSimulationTimeData);
 			}
 		}
diff --git a/Engine/Plugins/Runtime/NetworkPrediction/Source/NetworkPrediction/Public/NetworkedSimulationModelInterpolator.h b/Engine/Plugins/Runtime/NetworkPrediction/Source/NetworkPrediction/Public/NetworkedSimulationModelInterpolator.h
--- a/Engine/Plugins/Runtime/NetworkPrediction/Source/NetworkPrediction/Public/NetworkedSimulationModelInterpolator.h
+++ b/Engine/Plugins/Runtime/NetworkPrediction/Source/NetworkPrediction/Public/NetworkedSimulationModelInterpolator.h
@@ -55,7 +55,7 @@ struct TInterpolatorParameters
 template<typename Model>
 struct TNetSimInterpolator
 {
-	using TBufferTypes = typename TNetSimModelTraits<Model>::InternalBufferTypes;
+	using TBufferTypes = typename Model::BufferTypes;
 	using TTickSettings = typename Model::TickSettings;
 	using TSyncState = typename TBufferTypes::TSyncState;
 	using TAuxState = typename TBufferTypes::TAuxState;
@@ -63,6 +63,8 @@ struct TNetSimInterpolator
 	using FSimTime = FNetworkSimTime;
 	using FRealTime = FNetworkSimTime::FRealTime;
 
+	using TFrameState = TSimulationFrameState<Model>;
+
 	struct FStatePair
 	{
 		TSyncState Sync;
@@ -72,7 +74,7 @@ struct TNetSimInterpolator
 	bool bEnableVisualLog = true;
 
 	template<typename TSystemDriver>
-	FRealTime PostSimTick(TSystemDriver* Driver, const TNetworkSimBufferContainer<Model>& Buffers, const FSimulationTickState& TickInfo, const FNetSimTickParameters& TickParameters)
+	FRealTime PostSimTick(TSystemDriver* Driver, const TNetworkedSimulationState<Model>& State, const FNetSimTickParameters& TickParameters)
 	{
 		const bool bDoVLog = NetworkInterpolationDebugCVars::VLog() && bEnableVisualLog;
 		const float DeltaSeconds = TickParameters.LocalDeltaTimeSeconds;
@@ -80,31 +82,36 @@ struct TNetSimInterpolator
 		const AActor* LogOwner = Driver->GetVLogOwner();
 
 		// Interpolation disabled or we don't have 2 elements yet
-		if (NetworkInterpolationDebugCVars::Disable() > 0 || TickInfo.SimulationTimeBuffer.Num() <= 1)
+		if (NetworkInterpolationDebugCVars::Disable() > 0 || State.GetPendingTickFrame() <= 1)
 		{
 			// We are still responsible for calling FinalizeFrame though
-			if (const TSyncState* HeadState = Buffers.Sync.HeadElement())
+			const int32 HeadFrame = State.GetPendingTickFrame();
+			if (HeadFrame >= 0)
 			{
-				const TAuxState* AuxState = Buffers.Aux.HeadElement();
-				check(AuxState);
-
-				Driver->FinalizeFrame(*HeadState, *AuxState);
+				const TFrameState* FrameState = State.GetFrameState(HeadFrame);
+				const TAuxState* AuxState = State.GetAuxState(HeadFrame);
+				check(FrameState && AuxState);
+				Driver->FinalizeFrame(FrameState->SyncState, *AuxState);
 			}
-			
-			return TickInfo.GetTotalProcessedSimulationTime().ToRealTimeSeconds();
+
+			return State.GetTotalProcessedSimulationTime().ToRealTimeSeconds();
 		}
 
-		auto& SimulationTimeBuffer = TickInfo.SimulationTimeBuffer;
+		const TFrameState* HeadFrameState = State.GetFrameState(State.GetPendingTickFrame());
+		const TFrameState* TailFrameState = State.GetFrameState(State.GetConfirmedFrame());
+
+		FRealTime HeadRealTimeSeconds = HeadFrameState->TotalSimTime.ToRealTimeSeconds();
+		FRealTime TailRealTimeSeconds = TailFrameState->TotalSimTime.ToRealTimeSeconds();
 
 		// Starting off: start at the tail end
 		if (InterpolationTime <= 0.f)
 		{
-			InterpolationTime = SimulationTimeBuffer.TailElement()->ToRealTimeSeconds();
-			InterpolationFrame = SimulationTimeBuffer.TailFrame();
+			InterpolationTime = TailFrameState->TotalSimTime.ToRealTimeSeconds();
+			InterpolationFrame = State.GetConfirmedFrame();
 
 			auto& FromState = GetFromInterpolationState();
-			FromState.Sync = *Buffers.Sync[InterpolationFrame];
-			FromState.Aux = *Buffers.Aux[InterpolationFrame];
+			FromState.Sync = State.GetFrameState(InterpolationFrame)->SyncState;
+			FromState.Aux = *State.GetAuxState(InterpolationFrame);
 		}
 
 		EVisualLoggingContext LoggingContext = EVisualLoggingContext::InterpolationLatest;
@@ -113,18 +120,18 @@ struct TNetSimInterpolator
 		// Wait if we were too far ahead
 		if (WaitUntilTime > 0.f)
 		{
-			if (WaitUntilTime <= SimulationTimeBuffer.HeadElement()->ToRealTimeSeconds())
+			if (WaitUntilTime <= HeadRealTimeSeconds)
 			{
 				// done waiting, we can continue
 				WaitUntilTime = 0.f;
-				UE_VLOG(LogOwner, LogNetInterpolation, Log, TEXT("Done Waiting! Head: %s"), *LexToString(SimulationTimeBuffer.HeadElement()->ToRealTimeSeconds()));
+				UE_VLOG(LogOwner, LogNetInterpolation, Log, TEXT("Done Waiting! Head: %s"), *LexToString(HeadRealTimeSeconds));
 			}
 			else
 			{
 				if (bDoVLog)
 				{
 					// Still waiting, return
-					UE_VLOG(LogOwner, LogNetInterpolation, Log, TEXT("Still Waiting! %s < %s"), *LexToString(WaitUntilTime), *LexToString(SimulationTimeBuffer.HeadElement()->ToRealTimeSeconds()));
+					UE_VLOG(LogOwner, LogNetInterpolation, Log, TEXT("Still Waiting! %s < %s"), *LexToString(WaitUntilTime), *LexToString(HeadRealTimeSeconds));
 
 					const FName LocalInterpolationTimeName("Local Interpolation Time");
 					FVector2D LocalTimeVsInterpolationTime(LogOwner->GetWorld()->GetTimeSeconds(), (int32)(InterpolationTime * 1000.f));
@@ -138,18 +145,18 @@ struct TNetSimInterpolator
 			FRealTime Step = DeltaSeconds;
 
 			// Snap if way far behind
-			FRealTime SnapThreshold = SimulationTimeBuffer.HeadElement()->ToRealTimeSeconds() - NetworkInterpolationDebugCVars::SnapThreshold();
+			FRealTime SnapThreshold = HeadRealTimeSeconds - NetworkInterpolationDebugCVars::SnapThreshold();
 			if (NewInterpolationTime < SnapThreshold)
 			{
-				NewInterpolationTime = SimulationTimeBuffer.HeadElement()->ToRealTimeSeconds() - NetworkInterpolationDebugCVars::SnapGoal();
+				NewInterpolationTime = HeadRealTimeSeconds - NetworkInterpolationDebugCVars::SnapGoal();
 				UE_VLOG(LogOwner, LogNetInterpolation, Log, TEXT("Snapping to %s"), *LexToString(NewInterpolationTime));
 			}
 			
 			// Speed up if we are too far behind
-			FRealTime CatchUpThreshold = SimulationTimeBuffer.HeadElement()->ToRealTimeSeconds() - NetworkInterpolationDebugCVars::CatchUpThreshold();
+			FRealTime CatchUpThreshold = HeadRealTimeSeconds - NetworkInterpolationDebugCVars::CatchUpThreshold();
 			if (CatchUpUntilTime <= 0.f && NewInterpolationTime < CatchUpThreshold)
 			{
-				CatchUpUntilTime = SimulationTimeBuffer.HeadElement()->ToRealTimeSeconds() - NetworkInterpolationDebugCVars::CatchUpGoal();
+				CatchUpUntilTime = HeadRealTimeSeconds - NetworkInterpolationDebugCVars::CatchUpGoal();
 			}
 
 			if (CatchUpUntilTime > 0.f)
@@ -170,11 +177,11 @@ struct TNetSimInterpolator
 			NewInterpolationTime += Step;
 
 			// Did this put us too far ahead, and now we need to start waiting?
-			if (NewInterpolationTime > SimulationTimeBuffer.HeadElement()->ToRealTimeSeconds())
+			if (NewInterpolationTime > HeadRealTimeSeconds)
 			{
-				UE_VLOG(LogOwner, LogNetInterpolation, Log, TEXT("Too far ahead! Starting to wait! Head: %s"), *LexToString(SimulationTimeBuffer.HeadElement()->ToRealTimeSeconds()));
-				WaitUntilTime = SimulationTimeBuffer.HeadElement()->ToRealTimeSeconds() + NetworkInterpolationDebugCVars::WaitSlack();
-				NewInterpolationTime = SimulationTimeBuffer.HeadElement()->ToRealTimeSeconds();
+				UE_VLOG(LogOwner, LogNetInterpolation, Log, TEXT("Too far ahead! Starting to wait! Head: %s"), *LexToString(HeadRealTimeSeconds));
+				WaitUntilTime = HeadRealTimeSeconds + NetworkInterpolationDebugCVars::WaitSlack();
+				NewInterpolationTime = HeadRealTimeSeconds;
 			}
 		}
 
@@ -183,15 +190,15 @@ struct TNetSimInterpolator
 		const TAuxState* ToAuxState = nullptr;
 		FRealTime ToTime = 0.f;
 
-		for (auto It = SimulationTimeBuffer.CreateConstIterator(); It; ++It)
+		for (int32 Frame = State.GetConfirmedFrame(); Frame <= State.GetPendingTickFrame(); ++Frame)
 		{
-			FSimTime ElementSimTime = *It.Element();
-			if (NewInterpolationTime <= ElementSimTime.ToRealTimeSeconds())
+			FRealTime FrameRealTime = State.GetFrameState(Frame)->TotalSimTime.ToRealTimeSeconds();
+			if (NewInterpolationTime <= FrameRealTime)
 			{
-				InterpolationFrame = It.Frame();
-				ToTime = ElementSimTime.ToRealTimeSeconds();
-				ToState = Buffers.Sync[It.Frame()];
-				ToAuxState = Buffers.Aux[It.Frame()];
+				InterpolationFrame = Frame;
+				ToTime = FrameRealTime;
+				ToState = &State.GetFrameState(Frame)->SyncState;
+				ToAuxState = State.GetAuxState(Frame);
 				break;
 			}
 		}
@@ -220,49 +227,51 @@ struct TNetSimInterpolator
 				// Graph Interpolation Time vs Buffer Head/Tail times
 				const FName ServerSimulationGraphName("ServerSimulationTimeGraph");
 				const FName ServerSimTimeName("Server Simulation Time");
-				FVector2D LocalTimeVsServerSimTime(LogOwner->GetWorld()->GetTimeSeconds(), (int32)TickInfo.SimulationTimeBuffer.HeadElement()->ToRealTimeMS());
+				FVector2D LocalTimeVsServerSimTime(LogOwner->GetWorld()->GetTimeSeconds(), (int32)HeadFrameState->TotalSimTime.ToRealTimeMS());
 				UE_VLOG_HISTOGRAM(LogOwner, LogNetInterpolation, Log, ServerSimulationGraphName, ServerSimTimeName, LocalTimeVsServerSimTime);
 
 				const FName BufferTailSimTimeName("Buffer Tail Simulation Time");
-				FVector2D LocalTimeVsBufferTailSim(LogOwner->GetWorld()->GetTimeSeconds(), (int32)TickInfo.SimulationTimeBuffer.TailElement()->ToRealTimeMS());
+				FVector2D LocalTimeVsBufferTailSim(LogOwner->GetWorld()->GetTimeSeconds(), (int32)TailFrameState->TotalSimTime.ToRealTimeMS());
 				UE_VLOG_HISTOGRAM(LogOwner, LogNetInterpolation, Log, ServerSimulationGraphName, BufferTailSimTimeName, LocalTimeVsBufferTailSim);
 
 				const FName LocalInterpolationTimeName("Local Interpolation Time");
 				FVector2D LocalTimeVsInterpolationTime(LogOwner->GetWorld()->GetTimeSeconds(), (int32)(NewInterpolationTime * 1000.f));
 				UE_VLOG_HISTOGRAM(LogOwner, LogNetInterpolation, Log, ServerSimulationGraphName, LocalInterpolationTimeName, LocalTimeVsInterpolationTime);
 					
-				FVector2D LocalTimeVsCatchUpThreshold(LogOwner->GetWorld()->GetTimeSeconds(), (SimulationTimeBuffer.HeadElement()->ToRealTimeSeconds() - NetworkInterpolationDebugCVars::CatchUpThreshold()) * 1000.f);
+				FVector2D LocalTimeVsCatchUpThreshold(LogOwner->GetWorld()->GetTimeSeconds(), (HeadRealTimeSeconds - NetworkInterpolationDebugCVars::CatchUpThreshold()) * 1000.f);
 				UE_VLOG_HISTOGRAM(LogOwner, LogNetInterpolation, Log, ServerSimulationGraphName, "Catch Up Threshold", LocalTimeVsCatchUpThreshold);
 
-				FVector2D LocalTimeVsCatchUpGoal(LogOwner->GetWorld()->GetTimeSeconds(), (SimulationTimeBuffer.HeadElement()->ToRealTimeSeconds() - NetworkInterpolationDebugCVars::CatchUpGoal()) * 1000.f);
+				FVector2D LocalTimeVsCatchUpGoal(LogOwner->GetWorld()->GetTimeSeconds(), (TailRealTimeSeconds - NetworkInterpolationDebugCVars::CatchUpGoal()) * 1000.f);
 				UE_VLOG_HISTOGRAM(LogOwner, LogNetInterpolation, Log, ServerSimulationGraphName, "Catch Up Goal", LocalTimeVsCatchUpGoal);
 
 				// VLog the actual motion states
-				const TSyncState* DebugTail = Buffers.Sync.TailElement();
-				const TSyncState* DebugHead = Buffers.Sync.HeadElement();
+				const TSyncState* DebugTail = &TailFrameState->SyncState;
+				const TSyncState* DebugHead = &HeadFrameState->SyncState;
 
 				auto VLogHelper = [&](int32 Frame, EVisualLoggingContext Context, const FString& DebugStr)
 				{
 					FVisualLoggingParameters VLogParams(Context, Frame, EVisualLoggingLifetime::Transient, DebugStr);
-					Driver->InvokeVisualLog(Buffers.Input[Frame], Buffers.Sync[Frame], Buffers.Aux[Frame], VLogParams);
+					const TFrameState* FrameState = State.GetFrameState(Frame);
+
+					Driver->InvokeVisualLog(&FrameState->InputCmd, &FrameState->SyncState, State.GetAuxState(Frame), VLogParams);
 				};
 
-				VLogHelper(Buffers.Sync.TailFrame(), EVisualLoggingContext::InterpolationBufferTail, LexToString(TickInfo.SimulationTimeBuffer.HeadElement()->ToRealTimeMS()));
-				VLogHelper(Buffers.Sync.HeadFrame(), EVisualLoggingContext::InterpolationBufferHead, LexToString(TickInfo.SimulationTimeBuffer.TailElement()->ToRealTimeMS()));
+				VLogHelper(State.GetConfirmedFrame(), EVisualLoggingContext::InterpolationBufferTail, LexToString(TailFrameState->TotalSimTime.ToRealTimeMS()));
+				VLogHelper(State.GetPendingTickFrame(), EVisualLoggingContext::InterpolationBufferHead, LexToString(HeadFrameState->TotalSimTime.ToRealTimeMS()));
 
 				{
 					FVisualLoggingParameters VLogParams(EVisualLoggingContext::InterpolationFrom, InterpolationFrame-1, EVisualLoggingLifetime::Transient, LexToString(FromRealTime));
-					Driver->InvokeVisualLog(Buffers.Input[InterpolationFrame-1], &FromState.Sync, &FromState.Aux, VLogParams);
+					Driver->InvokeVisualLog(&State.GetFrameState(InterpolationFrame-1)->InputCmd, &FromState.Sync, &FromState.Aux, VLogParams);
 				}
 
 				{
 					FVisualLoggingParameters VLogParams(EVisualLoggingContext::InterpolationTo, InterpolationFrame, EVisualLoggingLifetime::Transient, LexToString(ToRealTime));
-					Driver->InvokeVisualLog(Buffers.Input[InterpolationFrame], ToState, ToAuxState, VLogParams);
+					Driver->InvokeVisualLog(&State.GetFrameState(InterpolationFrame)->InputCmd, ToState, ToAuxState, VLogParams);
 				}
 
 				{
 					FVisualLoggingParameters VLogParams(LoggingContext, InterpolationFrame, EVisualLoggingLifetime::Transient, LexToString(NewInterpolationTime));
-					Driver->InvokeVisualLog(Buffers.Input[InterpolationFrame], &OutputState.Sync, &OutputState.Aux, VLogParams);
+					Driver->InvokeVisualLog(&State.GetFrameState(InterpolationFrame)->InputCmd, &OutputState.Sync, &OutputState.Aux, VLogParams);
 				}
 			}
 				
diff --git a/Engine/Plugins/Runtime/NetworkPrediction/Source/NetworkPrediction/Public/NetworkedSimulationModelRepControllers.h b/Engine/Plugins/Runtime/NetworkPrediction/Source/NetworkPrediction/Public/NetworkedSimulationModelRepControllers.h
--- a/Engine/Plugins/Runtime/NetworkPrediction/Source/NetworkPrediction/Public/NetworkedSimulationModelRepControllers.h
+++ b/Engine/Plugins/Runtime/NetworkPrediction/Source/NetworkPrediction/Public/NetworkedSimulationModelRepControllers.h
@@ -5,9 +5,27 @@
 #include "Net/UnrealNetwork.h" // For MakeRelative
 #include "NetworkedSimulationModelCVars.h"
 #include "NetworkedSimulationModelInterpolator.h"
+#include "ProfilingDebugging/ScopedTimers.h"
+#include "NetworkPredictionCheck.h"
 
 // ----------------------------------------------------------------------------------------------------------------------------------------------------------------------
-//	CVars and compile time constants
+//	Replication Controllers
+//	
+//	The RepControllers are the pieces of the TNetworkedSimulationModel that make up the role-specific functionality (Server, Autonomous Client, and Simulated Client).
+//	Mainly they NetSerialize, Reconcile, and Tick. Essentially, they are what drives the Networked Simulation frame to frame based on networking role.
+//
+//	Multiple controllers can be used on a single TNetworkedSimulationModel instance. There is a 1:1 mapping for Tick, Reconcile, and Receiving. But to send you must 
+//	use the RepController of the receiving side.
+//
+//	Tick, Reconcile, NetSerialize: Receive
+//		Server    --> TRepController_Server
+//		AP Client --> TRepController_Autonomous
+//		SP Client --> TRepController_Simulated
+//
+//	NetSerialize: Send
+//		Server    --> TRepController_Autonomous, TRepController_Simulated
+//		AP Client --> TRepController_Server
+//
 // ----------------------------------------------------------------------------------------------------------------------------------------------------------------------
 
 namespace NetworkSimulationModelCVars
@@ -17,7 +35,7 @@ namespace NetworkSimulationModelCVars
 	NETSIM_DEVCVAR_SHIPCONST_INT(EnableSimulatedExtrapolation, 1, "ns.EnableSimulatedExtrapolation", "Toggle simulated proxy extrapolation.");
 	NETSIM_DEVCVAR_SHIPCONST_INT(ForceReconcile, 0, "ns.ForceReconcile", "Forces reconcile even if state does not differ. E.g, force resimulation after every netupdate.");
 	NETSIM_DEVCVAR_SHIPCONST_INT(ForceReconcileSingle, 0, "ns.ForceReconcileSingle", "Forces a since reconcile to happen on the next frame");
-	NETSIM_DEVCVAR_SHIPCONST_INT(EnableDebugBufferReplication, 1, "ns.DebugBufferReplication", "Replicate debug buffer (takes lots of bandwidth. Always compiled out of ship/test)");
+	NETSIM_DEVCVAR_SHIPCONST_INT(ForceAutonomousFrameJump, 0, "ns.ForceAutonomousFrameJump", "Causes AP client to jump 1000 frames ahead in its local frame #");
 }
 
 #define NETSIM_ENABLE_CHECKSUMS !(UE_BUILD_SHIPPING || UE_BUILD_TEST)
@@ -33,10 +51,11 @@ namespace NetworkSimulationModelCVars
 
 struct FNetworkSimulationSerialization
 {
-	// How many bits we use to encode the key frame number for buffers.
-	// Client Frames are stored locally as 32 bit integers, but we use a smaller # of bits to NetSerialize.
-	// Frames are only relatively relevant: the absolute value doesn't really matter. We just need to detect newer/older.
-	enum { NUM_BITS_FRAME = NETSIM_NETCONSTANT_NUM_BITS_FRAME };		
+	// How many bits we use to NetSerialize Frame numbers. This is only relevant AP Client <--> Server communication.
+	// Frames are stored locally as 32 bit integers, but we use a smaller number of bits to NetSerialize.
+	// The system internally guards from Frame numbers diverging. E.g, the client will not generate new frames if the
+	// last serialization frame would be pushed out of the buffer. Server does not generate frames without input from client.
+	enum { NUM_BITS_FRAME = NETSIM_NETCONSTANT_NUM_BITS_FRAME };
 
 	// Abs max value we encode into the bit writer
 	enum { MAX_FRAME_WRITE = 1 << NUM_BITS_FRAME };
@@ -47,15 +66,15 @@ struct FNetworkSimulationSerialization
 	enum { FRAME_ERROR_THRESHOLD = MAX_FRAME_WRITE / 2};
 
 	// Helper to serialize the int32 HeadFrame. Returns the unpacked value (this will be same as input in the save path)
-	static int32 SerializeFrame(FArchive& Ar, int32 LocalHeadFrame)
+	static int32 SerializeFrame(FArchive& Ar, int32 Frame, int32 LastSerializedFrame)
 	{
 		if (Ar.IsSaving())
 		{
-			((FNetBitWriter&)Ar).WriteIntWrapped( LocalHeadFrame, MAX_FRAME_WRITE );
-			return LocalHeadFrame;
+			((FNetBitWriter&)Ar).WriteIntWrapped( Frame, MAX_FRAME_WRITE );
+			return Frame;
 		}
 		
-		return MakeRelative(((FNetBitReader&)Ar).ReadInt( MAX_FRAME_WRITE ), LocalHeadFrame, MAX_FRAME_WRITE );
+		return MakeRelative(((FNetBitReader&)Ar).ReadInt( MAX_FRAME_WRITE ), LastSerializedFrame, MAX_FRAME_WRITE );
 	}
 };
 
@@ -65,241 +84,179 @@ struct FNetworkSimulationSerialization
 
 /** Helper that writes a new input cmd to the input buffer, at given Frame (usually the sim's PendingFrame). If frame doesn't exist, ProduceInput is called on the driver if bProduceInputViaDriver=true, otherwise the input cmd will be initialized from the previous input cmd. */
 template<typename Model>
-void GenerateLocalInputCmdAtFrame(typename TNetSimModelTraits<Model>::Driver* Driver, TNetworkSimBufferContainer<Model>& Buffers, const FNetworkSimTime& DeltaSimTime, int32 Frame, bool bProduceInputViaDriver=true)
+void GenerateLocalInputCmdAtFrame(TNetworkedSimulationModelDriver<typename Model::BufferTypes>* Driver, TNetworkedSimulationState<Model>& State, const FNetworkSimTime& DeltaSimTime, bool bProduceInputViaDriver)
 {
-	using TInputCmd = typename TNetworkSimBufferContainer<Model>::TInputCmd;
-	if (TInputCmd* InputCmd = Buffers.Input[Frame])
+	// Write the new frame
+	const int32 InputFrame = State.GetNextInputFrame();
+	TSimulationFrameState<Model>* FrameState = State.WriteFrameSequential(InputFrame);
+	
+	npEnsure(FrameState->HasFlag(ESimulationFrameStateFlags::InputWritten) == false);
+	FrameState->FrameDeltaTime.Set(DeltaSimTime);
+	FrameState->SetFlag(ESimulationFrameStateFlags::InputWritten);
+
+	if (bProduceInputViaDriver)
 	{
-		InputCmd->SetFrameDeltaTime(DeltaSimTime);
+		UE_NP_TRACE_PRODUCE_INPUT();
+		Driver->ProduceInput(DeltaSimTime, FrameState->InputCmd);
 	}
 	else
 	{
-		InputCmd = Buffers.Input.WriteFrameInitializedFromHead(Frame);
-		InputCmd->SetFrameDeltaTime(DeltaSimTime);
+		UE_NP_TRACE_SYNTH_INPUT();
+	}
+
+	State.SetLatestInputFrame(InputFrame);
+	UE_NP_TRACE_USER_STATE_INPUT(FrameState->InputCmd, InputFrame);
+}
+
+/** Helper to generate a local input cmd if we have simulation time to spend and advance the simulation's MaxAllowedFrame so that it can be processed. */
+template<typename Model>
+void TryGenerateLocalInput(TNetworkedSimulationModelDriver<typename Model::BufferTypes>* Driver, TNetworkedSimulationState<Model>& State, bool bProduceInputViaDriver, int32 MaxInputFrame=INDEX_NONE)
+{	
+	const FNetworkSimTime AllowedSimTime = State.GetTotalAllowedSimulationTime();
+	
+	// Submit new input if necessary
+	FNetworkSimTime MissingInputTime = AllowedSimTime - State.GetTotalInputSimulationTime();
+	while (MissingInputTime.IsPositive())
+	{
+		const int32 InputFrame = State.GetNextInputFrame();
+		if (MaxInputFrame != INDEX_NONE && InputFrame >= MaxInputFrame)
+		{
+			break;
+		}
+		
+		TSimulationFrameState<Model>* FrameState = State.WriteFrameSequential(InputFrame);
+		npEnsure(FrameState->HasFlag(ESimulationFrameStateFlags::InputWritten) == false);
+		FrameState->FrameDeltaTime.Set(MissingInputTime);
+		FrameState->SetFlag(ESimulationFrameStateFlags::InputWritten);
+		State.SetLatestInputFrame(InputFrame);
+
+		npEnsureMsgfSlow(State.GetTotalAllowedSimulationTime() >= State.GetTotalInputSimulationTime(), TEXT("Overflowed input %s"), *State.GetBasicDebugStr());
+		MissingInputTime -= FrameState->FrameDeltaTime.Get();
+
 		if (bProduceInputViaDriver)
 		{
 			UE_NP_TRACE_PRODUCE_INPUT();
-			Driver->ProduceInput(DeltaSimTime, *InputCmd);
+			Driver->ProduceInput(FrameState->FrameDeltaTime.Get(), FrameState->InputCmd);
 		}
 		else
 		{
 			UE_NP_TRACE_SYNTH_INPUT();
 		}
-		UE_NP_TRACE_USER_STATE_INPUT(*InputCmd, Frame);
+		
+		UE_NP_TRACE_USER_STATE_INPUT(FrameState->InputCmd, InputFrame);
 	}
-}
 
-/** Helper to generate a local input cmd if we have simulation time to spend and advance the simulation's MaxAllowedFrame so that it can be processed. */
-template<typename Model>
-void TryGenerateLocalInput(typename TNetSimModelTraits<Model>::Driver* Driver, TNetworkSimBufferContainer<Model>& Buffers, TSimulationTicker<typename Model::TickSettings>& Ticker, bool bProduceInputViaDriver=true)
-{
-	using TBufferTypes = typename TNetSimModelTraits<Model>::InternalBufferTypes;
-	using TInputCmd = typename TBufferTypes::TInputCmd;
-
-	FNetworkSimTime DeltaSimTime = Ticker.GetRemainingAllowedSimulationTime();
-	if (DeltaSimTime.IsPositive())
+	// Advance MaxTickFrame
+	for (int32 Frame = State.GetLatestInputFrame(); Frame > State.GetMaxTickFrame(); --Frame)
 	{
-		GenerateLocalInputCmdAtFrame<Model>(Driver, Buffers, DeltaSimTime, Ticker.PendingFrame, bProduceInputViaDriver);
-		Ticker.MaxAllowedFrame = Ticker.PendingFrame;
+		auto* FrameState = State.GetFrameState(Frame);
+		if (npEnsure(FrameState) && FrameState->HasFlag(ESimulationFrameStateFlags::InputWritten) && FrameState->TotalSimTime <= AllowedSimTime)
+		{
+			State.SetMaxTickFrame(Frame);
+			break;
+		}
 	}
 }
 
-template<typename T>
-struct TNetSimModelTraits;
-
-// ----------------------------------------------------------------------------------------------------------------------------------------------------------------------
-//	Replication Controllers
-//	
-//	The RepControllers are the pieces of the TNetworkedSimulationModel that make up the role-specific functionality (Server, Autonomous Client, and Simulated Client).
-//	Mainly they NetSerialize, Reconcile, and have Pre/Post sim tick functions. The TNetworkedSimulationModel is still the core piece that ticks the sim, but the Replicators
-//	do everything else in a role-specific way.
-// ----------------------------------------------------------------------------------------------------------------------------------------------------------------------
-
-// Empty RepController: zero functionality, just stubbed out functions and using definitions for convenience
 template<typename Model>
-struct TRepController_Empty
-{ 
-	using TSimulation = typename Model::Simulation;
-	using TTickSettings = typename Model::TickSettings;
-	using TDriver = typename TNetSimModelTraits<Model>::Driver;
-	using TBufferTypes = typename TNetSimModelTraits<Model>::InternalBufferTypes;
-	using TInputCmd = typename TBufferTypes::TInputCmd;
-	using TSyncState = typename TBufferTypes::TSyncState;
-	using TAuxState = typename TBufferTypes::TAuxState;
-
-	// Used for client shadowstate compares. Should just call GetDirtyCount() on the buffer you are replicating
-	int32 GetProxyDirtyCount(TNetworkSimBufferContainer<Model>& Buffers) const { return 0; }
-
-	// NetSerialize: just serialize the network data. Don't run simulation steps. Every replicator will be NetSerialized before moving on to Reconcile phase.
-	void NetSerialize(const FNetSerializeParams& P, TNetworkSimBufferContainer<Model>& Buffers, const TSimulationTicker<TTickSettings>& Ticker) { }
-
-	// Reconcile: called after everyone has NetSerialized. "Get right with the server": this function is about reconciling what the server told you vs what you have extrapolated or forward predicted locally
-	void Reconcile(TSimulation* Sim, TDriver* Driver, TNetworkSimBufferContainer<Model>& Buffers, TSimulationTicker<TTickSettings>& Ticker) { }
-
-	// Called prior to input processing. This function must update Ticker to allow simulation time (from TickParameters) and to possibly get new input.
-	void PreSimTick(TDriver* Driver, TNetworkSimBufferContainer<Model>& Buffers, TSimulationTicker<TTickSettings>& Ticker, const FNetSimTickParameters& TickParameters) { }
-
-	// Called after input processing. Should finalize the frame and do any smoothing/interpolation. This function is not allowed to modify the buffers or tick state, or even call the simulation/Update function.
-	void PostSimTick(TDriver* Driver, const TNetworkSimBufferContainer<Model>& Buffers, const TSimulationTicker<TTickSettings>& Ticker, const FNetSimTickParameters& TickParameters) { }
-};
-
-// Base implementation with minimal functionality: will grant simulation time and call finalize frame
-template<typename Model, typename TBase=TRepController_Empty<Model>>
-struct TRepController_Base : public TBase
+void DoSimulationTicks(typename Model::Simulation& Simulation, TNetworkedSimulationState<Model>& State, ESimulationTickContext Context)
 {
-	using typename TBase::TSimulation;
-	using typename TBase::TTickSettings;
-	using typename TBase::TDriver;
-	using typename TBase::TBufferTypes;
-	using typename TBase::TInputCmd;
-
-	// Called prior to input processing. This function must updated Ticker to allow simulation time (from TickParameters) and to possibly get new input.
-	void PreSimTick(TDriver* Driver, TNetworkSimBufferContainer<Model>& Buffers, TSimulationTicker<TTickSettings>& Ticker, const FNetSimTickParameters& TickParameters)
-	{
-		// Accumulate local delta time into Ticker
-		Ticker.GiveSimulationTime(TickParameters.LocalDeltaTimeSeconds);
-	}
+	using TFrameState = TSimulationFrameState<Model>;
 
-	// Called after input processing. Should finalize the frame and do any smoothing/interpolation. This function is not allowed to modify the buffers or tick state, or even call the simulation/Update function.
-	void PostSimTick(TDriver* Driver, const TNetworkSimBufferContainer<Model>& Buffers, const TSimulationTicker<TTickSettings>& Ticker, const FNetSimTickParameters& TickParameters)
+	while (State.GetPendingTickFrame() <= State.GetMaxTickFrame())
 	{
-		// Sync to latest frame if there is any
-		if (Buffers.Sync.Num() > 0)
+		const int32 InputSyncFrame = State.GetPendingTickFrame(); // Where Sync/Aux state have to come from. PendingFrame is where OOB mods are made, so we cannot skip ahead.
+			
+		// Find InputFrame. We want it to be PendingFrame but if no input was written at this frame, it can be skipped if there are valid inputs at later frames (gaps in network recvs)
+		int32 InputFrame = InputSyncFrame;
+		TFrameState* InputFrameState = State.GetFrameState(InputFrame);
+		while (InputFrameState && !InputFrameState->HasFlag(ESimulationFrameStateFlags::InputWritten))
 		{
-			Driver->FinalizeFrame(*Buffers.Sync.HeadElement(), *Buffers.Aux.HeadElement());
+			if (++InputFrame > State.GetMaxTickFrame())
+			{
+				InputFrameState = nullptr;
+				break;
+			}
+			InputFrameState = State.GetFrameState(InputFrame);
 		}
-	}
-};
-
-// -------------------------------------------------------------------------------------------------------
-//	Simulation Time replicator
-//	Just serializes simulation time. Can be disabled by Enabled templated parameter. Used as base class in other classes below.
-// -------------------------------------------------------------------------------------------------------
-template<typename Model, typename TBase=TRepController_Base<Model>, bool Enabled=true>
-struct TRepController_SimTime : public TBase
-{
-	using typename TBase::TSimulation;
-	using typename TBase::TTickSettings;
-	using typename TBase::TDriver;
-	using typename TBase::TBufferTypes;
-	using typename TBase::TInputCmd;
-
-	void NetSerialize(const FNetSerializeParams& P, TNetworkSimBufferContainer<Model>& Buffers, TSimulationTicker<TTickSettings>& Ticker)
-	{
-		TBase::NetSerialize(P, Buffers, Ticker);
-		SerializedTime = Ticker.GetTotalProcessedSimulationTime();
-		SerializedTime.NetSerialize(P.Ar);
-
-		NETSIM_CHECKSUM(P.Ar);
-	}
-
-	FNetworkSimTime SerializedTime;
-};
-
-// Enabled=false specialization: do nothing
-template<typename Model, typename TBase>
-struct TRepController_SimTime<Model, TBase, false> : public TBase { };
-
-// -------------------------------------------------------------------------------------------------------
-//	Sequence replicator
-//
-//	Replicates a sequence of elements. i.e., "the last MaxNumElements".
-//	On the receiving side, we merge the sequence into whatever we have locally
-//	Frames are synchronized. SimTime is also serialized by default (change by changing TBase)
-// -------------------------------------------------------------------------------------------------------
-
-template<typename Model, ENetworkSimBufferTypeId BufferId, int32 MaxNumElements=3, typename TBase=TRepController_SimTime<Model>>
-struct TRepController_Sequence : public TBase
-{
-	using typename TBase::TSimulation;
-	using typename TBase::TTickSettings;
-	using typename TBase::TDriver;
-	using typename TBase::TBufferTypes;
-	using typename TBase::TInputCmd;
-
-	int32 GetProxyDirtyCount(TNetworkSimBufferContainer<Model>& Buffers) const 
-	{
-		return Buffers.template Get<BufferId>().GetDirtyCount() ^ (TBase::GetProxyDirtyCount(Buffers) << 2); 
-	}
-
-	void NetSerialize(const FNetSerializeParams& P, TNetworkSimBufferContainer<Model>& Buffers, TSimulationTicker<TTickSettings>& Ticker)
-	{
-		NETSIM_CHECKSUM(P.Ar);
-
-		TBase::NetSerialize(P, Buffers, Ticker);
-
-		auto& Buffer = Buffers.template Get<BufferId>();
-		FArchive& Ar = P.Ar;
-		uint8 SerializedNumElements = FMath::Min<uint8>(MaxNumElements, Buffer.Num());
-		Ar << SerializedNumElements;
 
-		const int32 HeadFrame = FNetworkSimulationSerialization::SerializeFrame(Ar, Buffer.HeadFrame());
-		LastSerializedFrameStart = FMath::Max(0, HeadFrame - SerializedNumElements + 1);
+		// No valid input to be processed between PendingFrame -> MaxAllowedFrame
+		if (!InputFrameState)
+		{
+			break;
+		}
 
-		if (Ar.IsLoading())
+		// If we skipped frames due to lack of input, we need to copy the sync state over and set PendingTickFrame
+		if (InputFrame > InputSyncFrame)
 		{
-			const int32 PrevHead = Buffer.HeadFrame();
-			if (PrevHead < LastSerializedFrameStart && PrevHead >= 0)
-			{
-				// There is a gap in the stream. In some cases, we want this to be a "fault" and bubble up. We may want to synthesize state or maybe we just skip ahead.
-				UE_LOG(LogNetworkSim, Warning, TEXT("Fault: gap in received %s buffer. PrevHead: %d. Received: %d-%d. Reseting previous buffer contents"), *LexToString(BufferId), PrevHead, LastSerializedFrameStart, HeadFrame);
-			}
+			TSimulationFrameState<Model>* InputSyncState = State.GetFrameState(InputSyncFrame);
+			InputFrameState->SyncState = InputSyncState->SyncState;
+			InputFrameState->TotalSimTime = InputSyncState->TotalSimTime;
+			State.SetPendingTickFrame(InputFrame);
 		}
 
-		for (int32 Frame = LastSerializedFrameStart; Frame <= HeadFrame; ++Frame)
+		const FNetworkSimTime EndSimTime = InputFrameState->TotalSimTime + InputFrameState->FrameDeltaTime.Get();
+		if (EndSimTime <= State.GetTotalAllowedSimulationTime())
+		{
+			TSimulationDoTickFunc<Model>::DoTick(Simulation, State, Context);
+		}
+		else
 		{
-			// This, as is, is bad. The intention is that these functions serialize multiple items in some delta compressed fashion.
-			// As is, we are just serializing the elements individually.
-			auto* Cmd = Ar.IsLoading() ? Buffer.WriteFrame(Frame) : Buffer[Frame];
-			Cmd->NetSerialize(P);
+			break;
 		}
+	}
+}
 
-		LastSerializedFrame = HeadFrame;
+// Empty RepController: zero functionality, just stubbed out functions and using definitions for convenience
+template<typename Model>
+struct TRepController_Empty
+{ 
+	using TSimulation = typename Model::Simulation;
+	using TTickSettings = typename Model::TickSettings;
+	using TBufferTypes = typename Model::BufferTypes;
 
-		NETSIM_CHECKSUM(P.Ar);
-	}
+	using TDriver = TNetworkedSimulationModelDriver<TBufferTypes>;
+	using TFrameState = TSimulationFrameState<Model>;
+	
+	using TInputCmd = typename TBufferTypes::TInputCmd;
+	using TSyncState = typename TBufferTypes::TSyncState;
+	using TAuxState = typename TBufferTypes::TAuxState;
 
-	int32 GetLastSerializedFrame() const { return LastSerializedFrame; }
+	// NetSerialize: just serialize the network data. Don't run simulation steps. Every replicator will be NetSerialized before moving on to Reconcile phase.
+	void NetSerialize(const FNetSerializeParams& P, TNetworkedSimulationState<Model>& State) { }
 
-protected:
+	// Reconcile: called after everyone has NetSerialized. "Get right with the server": this function is about reconciling what the server told you vs what you have extrapolated or forward predicted locally
+	void Reconcile(TSimulation* Sim, TDriver* Driver, TNetworkedSimulationState<Model>& State) { }
 
-	int32 LastSerializedFrameStart = 0;
-	int32 LastSerializedFrame = 0;
+	// Called prior to input processing. This function must update tick state to allow simulation time (from TickParameters) and to possibly get new input.
+	void Tick(TDriver* Driver, TNetworkedSimulationState<Model>& Buffers, const FNetSimTickParameters& TickParameters) { }
 };
 
-// -------------------------------------------------------------------------------------------------------
-//	Server Replication Controller. FIXME: should this be called something like "client autonomous proxy -> Server RPC" ?
-//
-// -------------------------------------------------------------------------------------------------------
+// ----------------------------------------------------------------------------------------------------------------------------------------
+//	TRepController_Server
+//	Server Replication Controller. Server receives and ticks on this controller. AutonomousProxy sends on this controller.
+// ----------------------------------------------------------------------------------------------------------------------------------------
 template<typename Model, int32 NumSendPerUpdate=3, typename TBase=TRepController_Empty<Model>>
 struct TRepController_Server : public TBase
 {
 	using typename TBase::TSimulation;
 	using typename TBase::TTickSettings;
 	using typename TBase::TDriver;
-	using typename TBase::TBufferTypes;
 	using typename TBase::TInputCmd;
-
-	int32 GetProxyDirtyCount(TNetworkSimBufferContainer<Model>& Buffers) const 
-	{
-		return Buffers.Input.GetDirtyCount();
-	}
+	using typename TBase::TFrameState;
+	
 	// -------------------------------------------------------------------------------------------------
 	// TRepController_Server::NetSerialize
-	//	This is the client sending input data to the server
+	//	This is the client sending input data (Ar.IsSaving) to the server receiving input data (Ar.IsLoading)
 	// -------------------------------------------------------------------------------------------------
-	void NetSerialize(const FNetSerializeParams& P, TNetworkSimBufferContainer<Model>& Buffers, TSimulationTicker<TTickSettings>& Ticker)
+	void NetSerialize(const FNetSerializeParams& P, TNetworkedSimulationState<Model>& State)
 	{
 		NETSIM_CHECKSUM(P.Ar);
 
-		auto& InputBuffer = Buffers.Input;
+		TSimulationFrameState<Model> Test;
+		
 
 		FArchive& Ar = P.Ar;
-		uint8 SerializedNumElements = FMath::Min<uint8>(NumSendPerUpdate, InputBuffer.Num());
-		Ar << SerializedNumElements;
-
-		const int32 HeadFrame = FNetworkSimulationSerialization::SerializeFrame(Ar, InputBuffer.HeadFrame());
-		const int32 StartFrame = FMath::Max(0, HeadFrame - SerializedNumElements + 1);
 
 		// Serialize Input Cmds redundantly. This could be better:
 		//	-Sending: We could delta compress the data that doesn't change in subsequent cmds
@@ -307,101 +264,156 @@ struct TRepController_Server : public TBase
 		
 		if (Ar.IsSaving())
 		{
-			for (int32 Frame = StartFrame; Frame <= HeadFrame; ++Frame)
+			FNetBitWriter& BitWriter = (FNetBitWriter&)P.Ar;
+
+			const int32 HeadInputFrame = State.GetLatestInputFrame();
+			const int32 StartInputFrame = FMath::Max(0, (HeadInputFrame - NumSendPerUpdate + 1));
+			
+			FBitWriterMark NumSentMark(BitWriter);
+
+			uint8 SerializedNumElements = 0;
+			Ar << SerializedNumElements;
+
+			FNetworkSimulationSerialization::SerializeFrame(Ar, HeadInputFrame, LastSerializedFrame);
+
+			for (int32 Frame = StartInputFrame; Frame <= HeadInputFrame; ++Frame)
 			{
-				InputBuffer[Frame]->NetSerialize(P);
+				TFrameState* FrameState = State.GetFrameState(Frame);
+				if (FrameState && FrameState->HasFlag(ESimulationFrameStateFlags::InputWritten))
+				{
+					SerializedNumElements++;
+
+					FrameState->FrameDeltaTime.NetSerialize(P.Ar);
+					FrameState->InputCmd.NetSerialize(P);
+				}
 			}
+
+			FBitWriterMark Restore(BitWriter);
+			NumSentMark.PopWithoutClear(BitWriter);
+			Ar << SerializedNumElements;
+			Restore.PopWithoutClear(BitWriter);
+
+			LastSerializedFrame = HeadInputFrame;
 		}
 		else
 		{
-			const int32 ExistingInputHeadFrame = InputBuffer.HeadFrame(); // Anything we receive before this frame is redundant
-			
-			// SimulationTimeBuffer[N] is the sim time after running frame N.
-			// Sync[N]/Aux[N] is the output for frame N and input to frame N+1.
-			// Input[N] is the input cmd used in the sim tick that will generate frame N+1.			
-			// Therefor, Input buffer can be no more than 1 behind the simulation SImulationTimeBuffer.
-			ensureMsgf(ExistingInputHeadFrame+1 >= Ticker.SimulationTimeBuffer.HeadFrame(), TEXT("Simulation Ahead of Input Buffer: %d >= %d"), ExistingInputHeadFrame, Ticker.SimulationTimeBuffer.HeadFrame());
-
-			// Figure out what the sim time will be at input head (local input buffer may be ahead of processed simulation times)
-			FNetworkSimTime SimTime = Ticker.GetTotalProcessedSimulationTime();
-			for (int32 TimeFrame = Ticker.SimulationTimeBuffer.HeadFrame(); TimeFrame <= ExistingInputHeadFrame; ++TimeFrame)
+			auto TraceNewInputCmd = [](TInputCmd& Cmd, const FNetworkSimTime& TotalSimTime, const int32& Frame)
+			{
+				UE_NP_TRACE_NET_SERIALIZE_RECV(TotalSimTime, Frame);
+				UE_NP_TRACE_USER_STATE_INPUT(Cmd, Frame);
+				UE_NP_TRACE_NET_SERIALIZE_COMMIT();
+			};
+
+			TInputCmd ScratchInputCmd;
+			FNetworkSimTime ScratchNetworkSimTime;
+
+			uint8 SerializedNumElements = 0;
+			Ar << SerializedNumElements;
+
+			SerializedNumElements = FMath::Clamp<uint8>(SerializedNumElements, 0, NumSendPerUpdate); // Don't let client send us garbage
+
+			const int32 ReceivedHeadInputFrame = FMath::Max(FNetworkSimulationSerialization::SerializeFrame(Ar, State.GetLatestInputFrame(), LastSerializedFrame), 0);
+			const int32 StartInputFrame = FMath::Max(ReceivedHeadInputFrame - SerializedNumElements + 1, 0);
+
+			const int32 MinAcceptableFrame = State.GetNextInputFrame();
+			const int32 MaxAcceptableFrame = State.GetPendingTickFrame() + State.Max() - 1;
+
+			for (int32 Frame = StartInputFrame; Frame <= ReceivedHeadInputFrame; ++Frame)
 			{
-				//if (ensure(InputBuffer.IsValidFrame(TimeFrame))) // ensure
-				if (InputBuffer.IsValidFrame(TimeFrame)) // ensure
+				if (Frame < MinAcceptableFrame || Frame > MaxAcceptableFrame)
 				{
-					SimTime += InputBuffer[TimeFrame]->GetFrameDeltaTime();
+					// We don't want this frame. We've either moved past where it could matter or it is a redundant receive, or its so far ahead in the future that it would overwrite our buffer.
+					ScratchNetworkSimTime.NetSerialize(P.Ar);
+					ScratchInputCmd.NetSerialize(P);
+					
+					//UE_NP_TRACE_SYSTEM_FAULT("TRepController_Server::NetSerialize received Frame %d outside of acceptable range %d-%d.", Frame, MinAcceptableFrame, MaxAcceptableFrame);
+					continue;
+				}
+				else
+				{
+					LastSerializedFrame = Frame;
 				}
-			}
 
-			for (int32 Frame = StartFrame; Frame <= HeadFrame; ++Frame)
-			{
-				auto* InputCmd = InputBuffer.WriteFrame(Frame);
-				InputCmd->NetSerialize(P);
-				
-				if (Frame > ExistingInputHeadFrame)
+				if (TFrameState* ExistingFrame = State.GetFrameState(Frame))
 				{
-					UE_NP_TRACE_NET_SERIALIZE_RECV(SimTime, Frame);
-					UE_NP_TRACE_USER_STATE_INPUT(*InputCmd, Frame);
-					UE_NP_TRACE_NET_SERIALIZE_COMMIT();
+					if (npEnsure(!ExistingFrame->HasFlag(ESimulationFrameStateFlags::InputWritten)))
+					{
+						// Recv input on frame we've created (but haven't written input to)
+						ExistingFrame->FrameDeltaTime.NetSerialize(P.Ar);
+						ExistingFrame->InputCmd.NetSerialize(P);
+						ExistingFrame->SetFlag(ESimulationFrameStateFlags::InputWritten);
 
-					SimTime += InputCmd->GetFrameDeltaTime();
+						TraceNewInputCmd(ExistingFrame->InputCmd, ExistingFrame->TotalSimTime, Frame);
+
+						State.SetLatestInputFrame(Frame);
+					}
 				}
 				else
 				{
-					// Redundant cmd, don't trace the user state
+					// Brand new frame
+					auto* NewFrameState = State.WriteFrameSequential(Frame);
+					
+					NewFrameState->FrameDeltaTime.NetSerialize(P.Ar);
+					NewFrameState->InputCmd.NetSerialize(P);
+					NewFrameState->SetFlag(ESimulationFrameStateFlags::InputWritten);
+
+					State.SetLatestInputFrame(Frame);
+
+					TraceNewInputCmd(NewFrameState->InputCmd, NewFrameState->TotalSimTime, Frame);
 				}
 			}
 		}
-
-		LastSerializedFrame = HeadFrame;
+		
 		NETSIM_CHECKSUM(P.Ar);
 	}
 	
 	// -------------------------------------------------------------------------------------------------
 	// TRepController_Server::Reconcile
-	//	This is the server reconciling the simulation state with what the client sent him
+	//	This is the server reconciling the simulation state with what the client sent
 	// -------------------------------------------------------------------------------------------------
-	void Reconcile(TSimulation* Simulation, TDriver* Driver, TNetworkSimBufferContainer<Model>& Buffers, TSimulationTicker<TTickSettings>& Ticker)
+	void Reconcile(TSimulation* Simulation, TDriver* Driver, TNetworkedSimulationState<Model>& State)
 	{
 		// After receiving input, server may process up to the latest received frames.
 		// (If we needed to buffer input server side for whatever reason, we would do it here)
 		// (also note that we will implicitly guard against speed hacks in the core update loop by not processing cmds past what we have been "allowed")
-		Ticker.MaxAllowedFrame = Buffers.Input.HeadFrame();
+		State.SetMaxTickFrame(FMath::Max(LastSerializedFrame, State.GetMaxTickFrame()));
 
-		// Check for gaps in commands
-		if ( Ticker.PendingFrame < Buffers.Input.TailFrame() )
-		{
-			UE_LOG(LogNetworkSim, Warning, TEXT("TReplicator_Server::Reconcile missing inputcmds. PendingFrame: %d. %s. This can happen via packet loss"), Ticker.PendingFrame, *Buffers.Input.GetBasicDebugStr());
-			Ticker.PendingFrame = Buffers.Input.TailFrame();
-		}
+		
 	}
 	
 	// -------------------------------------------------------------------------------------------------
-	// TRepController_Server::PreSimTick
-	//	This is the server deciding if he should generate local input for the simulation
+	// TRepController_Server::Tick
+	//	This is the server ticking locally
 	// -------------------------------------------------------------------------------------------------
-	void PreSimTick(TDriver* Driver, TNetworkSimBufferContainer<Model>& Buffers, TSimulationTicker<TTickSettings>& Ticker, const FNetSimTickParameters& TickParameters)
+	void Tick(TDriver* Driver, TSimulation* Simulation, TNetworkedSimulationState<Model>& State, const FNetSimTickParameters& TickParameters)
 	{
-		Ticker.GiveSimulationTime(TickParameters.LocalDeltaTimeSeconds);
+		State.GiveSimulationTime(TickParameters.LocalDeltaTimeSeconds);
 		if (TickParameters.bGenerateLocalInputCmds)
 		{
-			TryGenerateLocalInput(Driver, Buffers, Ticker);
+			TryGenerateLocalInput(Driver, State, true);
 		}
+
+		DoSimulationTicks(*Simulation, State, ESimulationTickContext::Authority);
+
+		// Sync to latest frame if there is any
+		const int32 PendingFrame = State.GetPendingTickFrame();
+		Driver->FinalizeFrame(State.GetFrameState(PendingFrame)->SyncState, *State.GetAuxState(PendingFrame));
 	}
 
 	int32 GetLastSerializedFrame() const { return LastSerializedFrame; }
 
 private:
 	
-	int32 LastSerializedFrame = 0;
+	int32 LastSerializedFrame = -1;
 
 };
 
-// Simulated: "non locally controlled" simulations. We support "Simulation Extrapolation" here (using the sim to fake inputs to advance the sim)
-//	-TODO: this is replicating Input/Sync/Aux which is required to do accurate Simulation Extrapolation. For interpolated simulated proxies, we can skip the input and possibly the aux.
-//	More settings/config options would be nice here
+// ----------------------------------------------------------------------------------------------------------------------------------------
+//	TRepController_Autonomous
+//	Predictive client: reconciles with received authoritative state
+// ----------------------------------------------------------------------------------------------------------------------------------------
 template<typename Model, typename TBase=TRepController_Empty<Model>>
-struct TRepController_Simulated : public TBase
+struct TRepController_Autonomous: public TBase
 {
 	using typename TBase::TSimulation;
 	using typename TBase::TTickSettings;
@@ -411,335 +423,274 @@ struct TRepController_Simulated : public TBase
 	using typename TBase::TSyncState;
 	using typename TBase::TAuxState;
 
-	// Parent Simulation. If this is set, this simulation will forward predict in sync with this parent sim. The parent sim should be an autonomous proxy driven simulation
-	INetworkedSimulationModel* ParentSimulation = nullptr;
-
-	// Instance flag for enabling simulated extrapolation
-	bool bAllowSimulatedExtrapolation = true;
-
-	// Interpolated that will be used if bAllowSimulatedExtrapolation == false && ParentSimulation == nullptr
-	TNetSimInterpolator<Model> Interpolator;
+	using TFrameState = TSimulationFrameState<Model>;
 
-	// Last Serialized time and state
-	FNetworkSimTime GetLastSerializedSimulationTime() const { return LastSerializedSimulationTime; }
-	const TSyncState& GetLastSerializedSyncState() const { return LastSerializedSyncState; }
-	const TAuxState& GetLastSerializedAuxState() const { return LastSerializedAuxState; }
-	const TInputCmd& GetLastSerializedInputCmd() const { return LastSerializedInputCmd; }
-	int32 GetLastSerializedFrame() const { return LastSerializedKeyframe; }
-	
-	ESimulatedUpdateMode GetSimulatedUpdateMode() const
-	{
-		if (ParentSimulation)
-		{
-			return ESimulatedUpdateMode::ForwardPredict;
-		}
-		if (bAllowSimulatedExtrapolation && NetworkSimulationModelCVars::EnableSimulatedExtrapolation())
-		{
-			return ESimulatedUpdateMode::Extrapolate;
-		}
+	int32 GetLastSerializedFrame() const { return SerializedFrame; }
+	bool IsReconcileFaultDetected() const { return bReconcileFaultDetected; }
+	const FNetworkSimTime& GetLastSerializedSimTime() const { return SerializedTime; }
 
-		return ESimulatedUpdateMode::Interpolate;
-	}
-	
-	int32 GetProxyDirtyCount(TNetworkSimBufferContainer<Model>& Buffers) const
-	{
-		return Buffers.Sync.GetDirtyCount() ^ (TBase::GetProxyDirtyCount(Buffers) << 2); 
-	}
+	TArray<INetworkedSimulationModel*> DependentSimulations;
+	bool bDependentSimulationNeedsReconcile = false;
 
-	void NetSerialize(const FNetSerializeParams& P, TNetworkSimBufferContainer<Model>& Buffers, TSimulationTicker<TTickSettings>& Ticker)
+	// --------------------------------------------------------------------
+	//	TRepController_Autonomous::NetSerialize
+	// --------------------------------------------------------------------
+	void NetSerialize(const FNetSerializeParams& P, TNetworkedSimulationState<Model>& State)
 	{
-		FArchive& Ar = P.Ar;
 		NETSIM_CHECKSUM(P.Ar);
-
-		FNetworkSimTime PrevLastSerializedSimulationTime = LastSerializedSimulationTime;
-
-		// Serialize latest simulation time
-		LastSerializedSimulationTime = Ticker.GetTotalProcessedSimulationTime();
-		LastSerializedSimulationTime.NetSerialize(P.Ar);
-
-		// Serialize latest element
-		TInputCmd* InputCmd = nullptr;
-		TSyncState* SyncState = nullptr;
-		TAuxState* AuxState = nullptr;
+		FArchive& Ar = P.Ar;
+		
+		SerializedFrame = FNetworkSimulationSerialization::SerializeFrame(Ar, State.GetPendingTickFrame(), SerializedFrame);
 		
+		SerializedTime = State.GetTotalProcessedSimulationTime();
+		SerializedTime.NetSerialize(P.Ar);
+
 		if (Ar.IsSaving())
 		{
-			const int32 Frame = Buffers.Sync.HeadFrame();
-
-			InputCmd = Buffers.Input.HeadElement();
-			if (!InputCmd)
-			{
-				InputCmd = &LastSerializedInputCmd;
-			}
-			SyncState = Buffers.Sync.HeadElement();
-			AuxState = Buffers.Aux[Frame];
-			check(InputCmd && SyncState && AuxState	); // We should not be here if the buffer is empty. Want to avoid serializing an "empty" flag at the top here.
-
-			LastSerializedKeyframe = Frame;
+			// Server serialize the latest state
+			State.GetFrameState(SerializedFrame)->SyncState.NetSerialize(Ar);
+			State.GetAuxState(SerializedFrame)->NetSerialize(Ar);
+			State.CueDispatcher.NetSerializeSavedCues(Ar, ENetSimCueReplicationTarget::AutoProxy);
 		}
 		else
 		{
-			check(Ticker.SimulationTimeBuffer.HeadFrame() == Buffers.Sync.HeadFrame());
-			check(Ticker.GetTotalProcessedSimulationTime() <= Ticker.GetTotalAllowedSimulationTime());
+			SerializedSyncState.NetSerialize(Ar);
+			SerializedAuxState.NetSerialize(Ar);
+			State.CueDispatcher.NetSerializeSavedCues(Ar, ENetSimCueReplicationTarget::AutoProxy);
 
-			// Cache off our "starting" time before possibly overwriting it. We will use this in Reconcile to catch back up in some cases.
-			if (Ticker.GetTotalProcessedSimulationTime() > ReconcileSimulationTime)
-			{
-				ReconcileSimulationTime = Ticker.GetTotalProcessedSimulationTime();
-			}
+			UE_NP_TRACE_NET_SERIALIZE_RECV(SerializedTime, SerializedFrame);
+			UE_NP_TRACE_USER_STATE_SYNC(SerializedSyncState, SerializedFrame);
+			UE_NP_TRACE_USER_STATE_AUX(SerializedAuxState, SerializedFrame);
+			
+			bPendingReconciliation = false;
 
-			// Find out where this should go in the local buffer based on the serialized time
-			LastSerializedKeyframe = INDEX_NONE;
-			if (LastSerializedSimulationTime > Ticker.GetTotalProcessedSimulationTime())
+			TFrameState* ClientFrameState = State.GetFrameState(SerializedFrame);
+			TAuxState* ClientAuxState = State.GetAuxState(SerializedFrame);
+
+			// The state the client predicted that corresponds to the state the server just serialized to us
+			if (ClientFrameState && ClientAuxState)
 			{
-				// We are getting new state that is ahead of what we have locally, so it can safety go right to head
-				LastSerializedKeyframe = Ticker.SimulationTimeBuffer.HeadFrame()+1;
+				const bool bForceReconcile = (NetworkSimulationModelCVars::ForceReconcile() > 0) || (NetworkSimulationModelCVars::ForceReconcileSingle() > 0);
+				
+				if (bForceReconcile || ClientFrameState->TotalSimTime != SerializedTime || Model::ShouldReconcile(SerializedSyncState, SerializedAuxState, ClientFrameState->SyncState, *ClientAuxState))
+				{
+					NetworkSimulationModelCVars::SetForceReconcileSingle(0);
+					bPendingReconciliation =  true;
+				}
 			}
 			else
 			{
-				// We are getting state that is behind what we have locally
-				for (int32 Frame = Ticker.SimulationTimeBuffer.HeadFrame(); Frame >= Ticker.SimulationTimeBuffer.TailFrame(); --Frame)
-				{
-					if (LastSerializedSimulationTime > *Ticker.SimulationTimeBuffer[Frame])
-					{
-						LastSerializedKeyframe = Frame+1;
-						break;
-					}
-				}
-
-				if (LastSerializedKeyframe == INDEX_NONE)
-				{
-					FNetworkSimTime TotalTimeAhead = Ticker.GetTotalProcessedSimulationTime() - LastSerializedSimulationTime;
-					FNetworkSimTime SerializeDelta = LastSerializedSimulationTime - PrevLastSerializedSimulationTime;
-
-					// We are way far ahead of the server... we will need to clear out sync buffers, take what they gave us, and catch up in reconcile
-					//UE_LOG(LogNetworkSim, Warning, TEXT("!!! TReplicator_Simulated. Large gap detected. SerializedTime: %s. Buffer time: [%s-%s]. %d Elements. DeltaFromHead: %s. DeltaSerialize: %s"), *LastSerializedSimulationTime.ToString(), 
-					//	*Ticker.SimulationTimeBuffer.GetElementFromTail(0)->ToString(), *Ticker.SimulationTimeBuffer.GetElementFromHead(0)->ToString(), Ticker.SimulationTimeBuffer.GetNumValidElements(), *TotalTimeAhead.ToString(), *SerializeDelta.ToString());
-					LastSerializedKeyframe = Ticker.SimulationTimeBuffer.HeadFrame()+2; // (Skip ahead 2 to force a break in continuity)
-				}
+				// We don't have this frame locally anymore, will sort it out in Reconcile
+				bPendingReconciliation =  true;
 			}
 
-			check(LastSerializedKeyframe != INDEX_NONE);
-
-			// "Finalize" our buffers and time keeping such that we serialize the latest state from the server in the right spot
-			InputCmd = Buffers.Input.WriteFrame(LastSerializedKeyframe);
-			SyncState = Buffers.Sync.WriteFrame(LastSerializedKeyframe);
-			AuxState = Buffers.Aux.WriteFrame(LastSerializedKeyframe);
-
-			ensure(Buffers.Input.HeadFrame() <= LastSerializedKeyframe);
-			ensure(Buffers.Sync.HeadFrame() <= LastSerializedKeyframe);
-			ensure(Buffers.Aux.HeadFrame() <= LastSerializedKeyframe);
-
-			// Update tick info
-			Ticker.SetTotalProcessedSimulationTime(LastSerializedSimulationTime, LastSerializedKeyframe);
-			if (Ticker.GetTotalAllowedSimulationTime() < LastSerializedSimulationTime)
-			{
-				Ticker.SetTotalAllowedSimulationTime(LastSerializedSimulationTime);
-			}
-
-			check(Ticker.GetTotalProcessedSimulationTime() <= Ticker.GetTotalAllowedSimulationTime());
-
-			Ticker.PendingFrame = LastSerializedKeyframe;	// We are about to serialize state to LastSerializedKeyframe which will be "unprocessed" (has not been used to generate a new frame)
-			Ticker.MaxAllowedFrame = LastSerializedKeyframe-1; // Do not process PendingFrame on our account. ::PreSimTick will advance this based on our interpolation/extrapolation settings
-
-			UE_NP_TRACE_NET_SERIALIZE_RECV(LastSerializedSimulationTime, LastSerializedKeyframe);
-		}
-
-		check(SyncState && AuxState);
-		InputCmd->NetSerialize(Ar);
-		SyncState->NetSerialize(Ar);
-		AuxState->NetSerialize(Ar);
-
-		if (Ar.IsSaving())
-		{
-			// Server: send SimProxy and Interpolators. Whether this is in interpolation mode is really a client side thing (we want client to make this decision and transition between the two when necessary)
-			Buffers.CueDispatcher.NetSerializeSavedCues(Ar, ENetSimCueReplicationTarget::SimulatedProxy | ENetSimCueReplicationTarget::Interpolators);
-		}
-		else
-		{
-			Buffers.CueDispatcher.NetSerializeSavedCues(Ar, GetSimulatedUpdateMode() == ESimulatedUpdateMode::Interpolate ? ENetSimCueReplicationTarget::Interpolators : ENetSimCueReplicationTarget::SimulatedProxy);
-
-			LastSerializedInputCmd = *InputCmd;
-			LastSerializedSyncState = *SyncState;
-			LastSerializedAuxState = *AuxState;
-
-			UE_NP_TRACE_USER_STATE_INPUT(*InputCmd, LastSerializedKeyframe);
-			UE_NP_TRACE_USER_STATE_SYNC(*SyncState, LastSerializedKeyframe);
-			UE_NP_TRACE_USER_STATE_AUX(*AuxState, LastSerializedKeyframe);
-			
+			State.SetConfirmedFrame(SerializedFrame);
 		}
-
 		NETSIM_CHECKSUM(P.Ar);
 	}
 
-	void Reconcile(TSimulation* Simulation, TDriver* Driver, TNetworkSimBufferContainer<Model>& Buffers, TSimulationTicker<TTickSettings>& Ticker)
+	// --------------------------------------------------------------------
+	//	TRepController_Autonomous::Reconcile
+	// --------------------------------------------------------------------
+	void Reconcile(TSimulation* Simulation, TDriver* Driver, TNetworkedSimulationState<Model>& State)
 	{
-		if (ReconcileSimulationTime.IsPositive() == false)
+		if (bPendingReconciliation == false && (bDependentSimulationNeedsReconcile == false || SerializedFrame == INDEX_NONE))
 		{
 			return;
 		}
 
 		UE_NP_TRACE_NET_SERIALIZE_COMMIT();
+		bPendingReconciliation = false;
+		bDependentSimulationNeedsReconcile = false;
 
-		check(Ticker.GetTotalProcessedSimulationTime() <= Ticker.GetTotalAllowedSimulationTime());
-
-		if (GetSimulatedUpdateMode() == ESimulatedUpdateMode::Extrapolate && NetworkSimulationModelCVars::EnableSimulatedReconcile())
+		TFrameState* ClientFrameState = State.GetFrameState(SerializedFrame);
+		if (!ClientFrameState)
 		{
-			// This is effectively a rollback: we went back in time during NetSerialize and are now going to catch up ReconcileSimulationTime ms.
-			Buffers.CueDispatcher.NotifyRollback(Ticker.GetTotalProcessedSimulationTime());
+			// We received state that we don't have in our local buffer.
+			// We can't resimulate here, just reset to the serialized state
+			UE_NP_TRACE_SYSTEM_FAULT("Could not reconcile SerializedFrame %d with local state. %s", SerializedFrame, *State.GetBasicDebugStr());
+			ResetToSerializedState(State);
+			return;
+		}
+		
+		FNetworkSimTime RollbackDeltaTime = ClientFrameState->TotalSimTime - SerializedTime;
 
-			// Simulated Reconcile requires the input buffer to be kept up to date with the Sync buffer
-			// Generate a new, fake, command since we just added a new sync state to head
-			UE_NP_TRACE_SYNTH_INPUT();
-			while (Buffers.Input.HeadFrame() < Buffers.Sync.HeadFrame())
-			{
-				TInputCmd* Next = Buffers.Input.WriteFrameInitializedFromHead(Buffers.Input.HeadFrame()+1);
-			}
+		// Commit client's state to the server version
+		ClientFrameState->SyncState = SerializedSyncState;
+		ClientFrameState->TotalSimTime = SerializedTime;
+		*State.WriteAuxState(SerializedFrame) = SerializedAuxState;
 
-			// Do we have time to make up? We may have extrapolated ahead of the server (totally fine - can happen with small amount of latency variance)
-			FNetworkSimTime DeltaSimTime = ReconcileSimulationTime - Ticker.GetTotalProcessedSimulationTime();
-			if (DeltaSimTime.IsPositive() && NetworkSimulationModelCVars::EnableSimulatedReconcile())
-			{
-				SimulationExtrapolation(Simulation, Driver, Buffers, Ticker, DeltaSimTime);
-			}
+		if (NetworkSimulationModelCVars::EnableLocalPrediction() == 0)
+		{
+			// If we aren't predicting at all, then we advanced the allowed sim time here, (since we aren't doing it in PreSimTick). This just keeps us constantly falling behind and not being able to toggle prediction on/off for debugging.
+			State.SetTotalAllowedSimulationTime(SerializedTime);
 		}
-		
-		check(Ticker.GetTotalProcessedSimulationTime() <= Ticker.GetTotalAllowedSimulationTime());
-		ReconcileSimulationTime.Reset();
-	}
-	
-	void PreSimTick(TDriver* Driver, TNetworkSimBufferContainer<Model>& Buffers, TSimulationTicker<TTickSettings>& Ticker, const FNetSimTickParameters& TickParameters)
-	{
-		// Tick if we are dependent simulation or extrapolation is enabled
-		if (ParentSimulation || (GetSimulatedUpdateMode() == ESimulatedUpdateMode::Extrapolate))
+
+		// Tell dependent simulations to rollback
+		for (INetworkedSimulationModel* DependentSim : DependentSimulations)
 		{
-			// Don't start this simulation until you've gotten at least one update from the server
-			if (Ticker.GetTotalProcessedSimulationTime().IsPositive())
-			{
-				Ticker.GiveSimulationTime(TickParameters.LocalDeltaTimeSeconds);
-			}
-			
-			// Generate local input if we are ready to tick. Note that we pass false so we won't call into the Driver to produce the input, we will use the last serialized InputCmd's values
-			TryGenerateLocalInput<Model>(Driver, Buffers, Ticker, false);
+			DependentSim->BeginRollback(RollbackDeltaTime, SerializedFrame);
 		}
-	}
 
-	void PostSimTick(TDriver* Driver, const TNetworkSimBufferContainer<Model>& Buffers, const TSimulationTicker<TTickSettings>& Ticker, const FNetSimTickParameters& TickParameters)
-	{
-		if (GetSimulatedUpdateMode() == ESimulatedUpdateMode::Interpolate)
+		// Tell cue dispatch that we've rolledback as well
+		State.CueDispatcher.NotifyRollback(SerializedTime);
+
+		const int32 StartingPendingTickFrame = State.GetPendingTickFrame();
+		const int32 ExpectedEndPendingFrame = FMath::Min(StartingPendingTickFrame, State.GetMaxTickFrame());
+		const int32 LastResimulateFrame = StartingPendingTickFrame-1;
+		if (SerializedFrame < StartingPendingTickFrame)
 		{
-			const FNetworkSimTime::FRealTime InterpolationRealTime = Interpolator.template PostSimTick<TDriver>(Driver, Buffers, Ticker, TickParameters);
-			const FNetworkSimTime InterpolationSimTime = FNetworkSimTime::FromRealTimeSeconds(InterpolationRealTime);
+			VisualLogFrameStates(Driver, State, SerializedFrame, LastResimulateFrame-1, EVisualLoggingContext::OtherMispredicted, TEXT("Resimulate Step: mispredicted"));
+			VisualLogFrameStates(Driver, State, LastResimulateFrame, LastResimulateFrame, EVisualLoggingContext::LastMispredicted, TEXT("Resimulate Step: final mispredicted"));
+
+			FNetworkSimTime StartingRemainingTime = State.GetTotalAllowedSimulationTime() - State.GetTotalProcessedSimulationTime();
+			npEnsureMsgf(!StartingRemainingTime.IsNegative(), TEXT("Negative remaining SimTime %s %s"), *State.GetTotalAllowedSimulationTime().ToString(), *State.GetTotalProcessedSimulationTime().ToString());
+
+			State.SetPendingTickFrame(SerializedFrame);
+			State.SetMaxTickFrame(StartingPendingTickFrame-1);
+
+			DoSimulationTicks(*Simulation, State, ESimulationTickContext::Resimulate);
+			
+			State.SetTotalAllowedSimulationTime(State.GetTotalProcessedSimulationTime() + StartingRemainingTime);
 
-			Buffers.CueDispatcher.SetMaxDispatchTime(InterpolationSimTime);
-			Buffers.CueDispatcher.SetConfirmedTime(InterpolationSimTime);
+			VisualLogFrameStates(Driver, State, SerializedFrame, LastResimulateFrame-1, EVisualLoggingContext::OtherPredicted, TEXT("Resimulate Step: repredicted"));
+			VisualLogFrameStates(Driver, State, LastResimulateFrame, LastResimulateFrame, EVisualLoggingContext::LastPredicted, TEXT("Resimulate Step: final repredicted"));
 		}
 		else
 		{
-			// Sync to latest frame if there is any
-			if (Buffers.Sync.Num() > 0)
-			{
-				Driver->FinalizeFrame(*Buffers.Sync.HeadElement(), *Buffers.Aux.HeadElement());
-			}
-
-			Buffers.CueDispatcher.ClearMaxDispatchTime();
-			Buffers.CueDispatcher.SetConfirmedTime(LastSerializedSimulationTime);
+			// We have no frames to resimulate, essentially we are reset to the SerializedFrame
+			State.Reset(SerializedFrame);
 		}
+
+		State.CheckInvariants();
 	}
-	
-	void DependentRollbackBegin(TSimulation* Simulation, TDriver* Driver, TNetworkSimBufferContainer<Model>& Buffers, TSimulationTicker<TTickSettings>& Ticker, const FNetworkSimTime& RollbackDeltaTime, const int32 ParentFrame)
+
+	// --------------------------------------------------------------------
+	//	TRepController_Autonomous::Tick
+	// --------------------------------------------------------------------
+	void Tick(TDriver* Driver, TSimulation* Simulation, TNetworkedSimulationState<Model>& State, const FNetSimTickParameters& TickParameters)
 	{
-		// For now, we make the assumption that our last serialized state and time match the parent simulation.
-		// This would not always be the case with low frequency simulated proxies. But could be handled by replicating the simulations together (at the replication level)
-		const int32 NewHeadFrame = Buffers.Sync.HeadFrame()+1;
-		
-		Ticker.SetTotalProcessedSimulationTime(LastSerializedSimulationTime, NewHeadFrame);
-		Ticker.SetTotalAllowedSimulationTime(LastSerializedSimulationTime);
-		Ticker.PendingFrame = NewHeadFrame;
-		Ticker.MaxAllowedFrame = NewHeadFrame;
+		if (NetworkSimulationModelCVars::ForceAutonomousFrameJump() > 0)
+		{
+			// Testing CVar to mimic grave mishaps
+			const int32 FrameJump = NetworkSimulationModelCVars::ForceAutonomousFrameJump();
+			NetworkSimulationModelCVars::SetForceAutonomousFrameJump(0);
+			
+			// Copy pending state to jumped frame
+			TFrameState PendingFrameState = *State.GetFrameState(State.GetPendingTickFrame());
+			const int32 NewFrame = State.GetPendingTickFrame() + FrameJump;
+			TFrameState* NewFrameState =  State.WriteFrameSequential(NewFrame);
+			*NewFrameState = PendingFrameState;
 
-		*Buffers.Sync.WriteFrame(NewHeadFrame) = LastSerializedSyncState;
-		*Buffers.Input.WriteFrame(NewHeadFrame) = TInputCmd();
+			State.Reset(NewFrame);
 
-		Driver->FinalizeFrame(LastSerializedSyncState, LastSerializedAuxState);
+			UE_NP_TRACE_SYSTEM_FAULT("NetworkSimulationModelCVars::ForceAutonomousFrameJump invoked. Jumped to: %d. Prev: %s", FrameJump, *State.GetBasicDebugStr());
+		}
 
-		Buffers.CueDispatcher.NotifyRollback(LastSerializedSimulationTime);
+		// The MaxInputFrame we can write without getting too far ahead of server
+		// We pass this into TryGenerateLocalInput to ensure that it doesn't.
+		// If we some how get in a state where we are past MaxInputFrame, we need to reset
+		const int32 MaxInputFrame = (SerializedFrame + State.Max() - 1);
 
-		if (NETSIM_MODEL_DEBUG)
+		if (State.GetNextInputFrame() > MaxInputFrame)
 		{
-			FVisualLoggingParameters VLogParams(EVisualLoggingContext::FirstMispredicted, ParentFrame, EVisualLoggingLifetime::Persistent, TEXT("DependentRollbackBegin"));
-			Driver->InvokeVisualLog(nullptr, &LastSerializedSyncState, &LastSerializedAuxState, VLogParams);
+			// We somehow went over the limit and must reset.
+			ResetToSerializedState(State);
+			UE_NP_TRACE_SYSTEM_FAULT("TRepController_Autonomous::PreSimTick input overrun. MaxInputFrame: %d. %s", MaxInputFrame, *State.GetBasicDebugStr());
+			return;
+		}
+
+		if (State.GetNextInputFrame() == MaxInputFrame)
+		{
+			// We are at the limit but not over it. We are essentially pausing the simulation until we recover. Set this bool to reflect this.
+			bReconcileFaultDetected = true;
+			UE_NP_TRACE_SYSTEM_FAULT("TRepController_Autonomous::PreSimTick Input Buffer Maxed. Not processing local time. MaxInputFrame: %d. %s", MaxInputFrame, *State.GetBasicDebugStr());
+			return;
 		}
-	}
-		
-	void DependentRollbackStep(TSimulation* Simulation, TDriver* Driver, TNetworkSimBufferContainer<Model>& Buffers, TSimulationTicker<TTickSettings>& Ticker, const FNetworkSimTime& StepTime, const int32 ParentFrame, const bool bFinalStep)
-	{
-		Ticker.SetTotalAllowedSimulationTime( Ticker.GetTotalAllowedSimulationTime() + StepTime );
 
-		SimulationExtrapolation(Simulation, Driver, Buffers, Ticker, StepTime);
+		bReconcileFaultDetected = false;
 
-		TSyncState* SyncState = Buffers.Sync.HeadElement();
-		check(SyncState);
-		
-		if (NETSIM_MODEL_DEBUG)
+		if (TickParameters.bGenerateLocalInputCmds)
 		{
-			FVisualLoggingParameters VLogParams(bFinalStep ? EVisualLoggingContext::LastMispredicted : EVisualLoggingContext::OtherMispredicted, ParentFrame, EVisualLoggingLifetime::Persistent, TEXT("DependentRollbackStep"));
-			Driver->InvokeVisualLog(nullptr, SyncState, Buffers.Aux.HeadElement(), VLogParams);
+			if (NetworkSimulationModelCVars::EnableLocalPrediction() > 0)
+			{
+				// Prediction: add simulation time and generate new commands
+				State.GiveSimulationTime(TickParameters.LocalDeltaTimeSeconds);
+				TryGenerateLocalInput<Model>(Driver, State, true, MaxInputFrame);
+			}
+			else
+			{
+				// Local prediction disabled: we must use a separate time accumulator to figure out when we should add more input cmds.
+				// Since we aren't processing the simulation locally, our core simulation time will only advance from network updates.
+				// (still, we need *something* to tell us when to generate a new command and what the delta time should be)
+				FNetworkSimTime NonPredictedInputTime;
+				NonPredictedInputTimeAccumulator.Accumulate(NonPredictedInputTime, TickParameters.LocalDeltaTimeSeconds);
+				if (NonPredictedInputTime.IsPositive())
+				{
+					GenerateLocalInputCmdAtFrame<Model>(Driver, State, NonPredictedInputTime, true);
+				}
+			}
 		}
+
+		DoSimulationTicks(*Simulation, State, ESimulationTickContext::Predict);
+
+		State.CueDispatcher.ClearMaxDispatchTime();
+		State.CueDispatcher.SetConfirmedTime(SerializedTime);
+
+		const int32 PendingFrame = State.GetPendingTickFrame();
+		Driver->FinalizeFrame(State.GetFrameState(PendingFrame)->SyncState, *State.GetAuxState(PendingFrame));
 	}
 
 private:
 	
-	void SimulationExtrapolation(TSimulation* Simulation, TDriver* Driver, TNetworkSimBufferContainer<Model>& Buffers, TSimulationTicker<TTickSettings>& Ticker, const FNetworkSimTime DeltaSimTime)
+	// Called as last resort to full reset to last serialized state. Commands in flight will cause correction
+	void ResetToSerializedState(TNetworkedSimulationState<Model>& State)
 	{
-		// We have extrapolated ahead of the server. The latest network update is now "in the past" from what we rendered last frame.
-		// We will insert a new frame to make up the difference from the last known state to where we want to be in the now.
-
-		ensure(Buffers.Input.HeadFrame() == Buffers.Sync.HeadFrame());
-
-		const int32 InputFrame = Buffers.Input.HeadFrame();
-		const int32 OutputFrame = InputFrame + 1;
-
-		// Override the FrameDeltaTime on our InputCmd.
-		TInputCmd* InputCmd = Buffers.Input.HeadElement();
-		InputCmd->SetFrameDeltaTime(DeltaSimTime);
-		
-		// Create new sync state to write to
-		TSyncState* PrevSyncState = Buffers.Sync[InputFrame];
-		TSyncState* NextSyncState = Buffers.Sync.WriteFrame(OutputFrame);
-		TAuxState* AuxState = Buffers.Aux[InputFrame];
-
-		if (NETSIM_MODEL_DEBUG)
+		if (SerializedFrame != INDEX_NONE)
 		{
-			FVisualLoggingParameters VLogParameters( EVisualLoggingContext::OtherPredicted, InputFrame, EVisualLoggingLifetime::Persistent, TEXT("Pre SimulationExtrapolation"));
-			Driver->InvokeVisualLog(InputCmd, PrevSyncState, AuxState, VLogParameters);
+			State.WriteFrameWithTime(SerializedFrame, SerializedTime)->SyncState = SerializedSyncState;
+			*State.WriteAuxState(SerializedFrame) = SerializedAuxState;
+			State.Reset(SerializedFrame);
 		}
+	}
 
-		// Do the actual update
-		TSimulationDoTickFunc<Model>::DoTick(*Simulation, Ticker, Buffers, OutputFrame, ESimulationTickContext::Resimulate);
-
+	void VisualLogFrameStates(TDriver* Driver, TNetworkedSimulationState<Model>& State, int32 StartFrame, int32 EndFrame, EVisualLoggingContext Context, const TCHAR* Str)
+	{
 		if (NETSIM_MODEL_DEBUG)
 		{
-			FVisualLoggingParameters VLogParameters( EVisualLoggingContext::OtherPredicted, InputFrame, EVisualLoggingLifetime::Persistent, TEXT("Post SimulationExtrapolation"));
-			Driver->InvokeVisualLog(InputCmd, NextSyncState, Buffers.Aux[OutputFrame], VLogParameters);
+			for (int32 Frame = StartFrame; Frame <= EndFrame; ++Frame)
+			{
+				TFrameState* FrameState = State.GetFrameState(Frame);
+				if (npEnsureMsgf(FrameState, TEXT("Invalid Frame %d. %s"), Frame, *State.GetBasicDebugStr()))
+				{
+					FVisualLoggingParameters VLogParameters(Context, Frame, EVisualLoggingLifetime::Persistent, Str);
+					Driver->InvokeVisualLog(&FrameState->InputCmd, &FrameState->SyncState, State.GetAuxState(Frame), VLogParameters);
+				}
+			}
 		}
-
-		Ticker.MaxAllowedFrame = OutputFrame;
 	}
-
-	int32 LastSerializedKeyframe;
 	
-	FNetworkSimTime ReconcileSimulationTime;
-	FNetworkSimTime LastSerializedSimulationTime;
-	
-	TInputCmd LastSerializedInputCmd;
-	TSyncState LastSerializedSyncState;
-	TAuxState LastSerializedAuxState;	// Temp? This should be conditional or optional. We want to support not replicating the aux state to simulated proxies
+	TSyncState SerializedSyncState;
+	TAuxState SerializedAuxState;
+	FNetworkSimTime SerializedTime;
+	int32 SerializedFrame = -1;
+
+	bool bPendingReconciliation = false;	// Reconciliation is pending: we need to reconcile state from the server that differs from the locally predicted state
+	bool bReconcileFaultDetected = false;	// A fault was detected: we received state from the server that we are unable to reconcile with locally predicted state
+
+	TRealTimeAccumulator<TTickSettings> NonPredictedInputTimeAccumulator; // for tracking input time in the non predictive case
 };
 
-/** Replicates TSyncState and does basic reconciliation. */
-template<typename Model, typename TBase=TRepController_Base<Model>>
-struct TRepController_Autonomous: public TBase
+// ----------------------------------------------------------------------------------------------------------------------------------------
+//	TRepController_Simulated
+//	Clients receive and tick on this controller for non locally controlled simulations. Server sends on controller to those clients.
+//	Can run in three modes:
+//		-Interpolate: no running on SimulationTick
+//		-SimulationExtrapolation: run SimulationTick in between network updates
+//		-ForwardPredict: forward predict in step with a leading Autonomous Proxy simulation
+// ----------------------------------------------------------------------------------------------------------------------------------------
+template<typename Model, typename TBase=TRepController_Empty<Model>>
+struct TRepController_Simulated : public TBase
 {
 	using typename TBase::TSimulation;
 	using typename TBase::TTickSettings;
@@ -749,317 +700,246 @@ struct TRepController_Autonomous: public TBase
 	using typename TBase::TSyncState;
 	using typename TBase::TAuxState;
 
-	int32 GetLastSerializedFrame() const { return SerializedFrame; }
-	bool IsReconcileFaultDetected() const { return bReconcileFaultDetected; }
-	const FNetworkSimTime& GetLastSerializedSimTime() const { return SerializedTime; }
+	using TFrameState = TSimulationFrameState<Model>;
 
-	TArray<INetworkedSimulationModel*> DependentSimulations;
-	bool bDependentSimulationNeedsReconcile = false;
+	// Parent Simulation. If this is set, this simulation will forward predict in sync with this parent sim. The parent sim should be an autonomous proxy driven simulation
+	INetworkedSimulationModel* ParentSimulation = nullptr;
+
+	// Instance flag for enabling simulated extrapolation
+	bool bAllowSimulatedExtrapolation = true;
+
+	// Interpolated that will be used if bAllowSimulatedExtrapolation == false && ParentSimulation == nullptr
+	TNetSimInterpolator<Model> Interpolator;
 
-	int32 GetProxyDirtyCount(TNetworkSimBufferContainer<Model>& Buffers) const
+	// Last Serialized time and state (this is for the debugger which should go away)
+	FNetworkSimTime GetLastSerializedSimulationTime() const { return LastSerializedSimulationTime; }
+	const TSyncState& GetLastSerializedSyncState() const { return LastSerializedSyncState; }
+	const TAuxState& GetLastSerializedAuxState() const { return LastSerializedAuxState; }
+	const TInputCmd& GetLastSerializedInputCmd() const { return LastSerializedInputCmd; }
+	int32 GetLastSerializedFrame() const { return LastSerializedKeyframe; }
+	
+	ESimulatedUpdateMode GetSimulatedUpdateMode() const
 	{
-		return (Buffers.Sync.GetDirtyCount()) ^ (Buffers.Aux.GetDirtyCount() << 1) ^ (TBase::GetProxyDirtyCount(Buffers) << 2); 
+		if (ParentSimulation)
+		{
+			return ESimulatedUpdateMode::ForwardPredict;
+		}
+		if (bAllowSimulatedExtrapolation && NetworkSimulationModelCVars::EnableSimulatedExtrapolation())
+		{
+			return ESimulatedUpdateMode::Extrapolate;
+		}
+
+		return ESimulatedUpdateMode::Interpolate;
 	}
 
-	// --------------------------------------------------------------------
-	//	NetSerialize
-	// --------------------------------------------------------------------
-	void NetSerialize(const FNetSerializeParams& P, TNetworkSimBufferContainer<Model>& Buffers, TSimulationTicker<TTickSettings>& Ticker)
+	// -------------------------------------------------------------------------------------------------
+	// TRepController_Simulated::NetSerialize
+	// -------------------------------------------------------------------------------------------------
+	void NetSerialize(const FNetSerializeParams& P, TNetworkedSimulationState<Model>& State)
 	{
-		NETSIM_CHECKSUM(P.Ar);
 		FArchive& Ar = P.Ar;
-		
-		SerializedFrame = FNetworkSimulationSerialization::SerializeFrame(Ar, Buffers.Sync.HeadFrame());
-		
-		SerializedTime = Ticker.GetTotalProcessedSimulationTime();
-		SerializedTime.NetSerialize(P.Ar);
+		NETSIM_CHECKSUM(P.Ar);
 
 		if (Ar.IsSaving())
 		{
-			// Server serialize the latest state
-			Buffers.Sync.HeadElement()->NetSerialize(Ar);
-			Buffers.Aux.HeadElement()->NetSerialize(Ar);
+			LastSerializedKeyframe = State.GetPendingTickFrame();
+			
+			TFrameState* FrameState = State.GetFrameState(LastSerializedKeyframe);
+			npCheckSlow(FrameState);
+
+			TFrameState* InputFrameState = State.GetFrameState(State.GetLatestInputFrame());
+			npCheckfSlow(InputFrameState, TEXT("Invalid InputFrame. %s"), *State.GetBasicDebugStr());
+
+			TAuxState* AuxState = State.GetAuxState(LastSerializedKeyframe);
+			npCheckSlow(AuxState);
+			
+			FrameState->TotalSimTime.NetSerialize(P.Ar);	// 1. Time
+			FrameState->SyncState.NetSerialize(P.Ar);		// 2. Sync
+			InputFrameState->InputCmd.NetSerialize(P.Ar);	// 3. Input
+			AuxState->NetSerialize(P.Ar); // 4. Aux
+
+			// Server: send SimProxy and Interpolators. Whether this is in interpolation mode is really a client side thing (we want client to make this decision and transition between the two when necessary)
+			State.CueDispatcher.NetSerializeSavedCues(Ar, ENetSimCueReplicationTarget::SimulatedProxy | ENetSimCueReplicationTarget::Interpolators); // 5. Cues
 		}
 		else
 		{
-			SerializedSyncState.NetSerialize(Ar);
-			SerializedAuxState.NetSerialize(Ar);
-		}
+			// ReconcileSimulationTime is where we will reconcile up to if we get an update from the server "in the past"
+			// (But since we can get multiple NetSerializes before reconciling, don't let a stale update mess things up)
+			if (State.GetTotalProcessedSimulationTime() > ReconcileSimulationTime)
+			{
+				ReconcileSimulationTime = State.GetTotalProcessedSimulationTime();
+			}
+			
+			LastSerializedSimulationTime.NetSerialize(P.Ar); // 1. Time
 
-		Buffers.CueDispatcher.NetSerializeSavedCues(Ar, ENetSimCueReplicationTarget::AutoProxy);
+			LastSerializedKeyframe = State.GetPendingTickFrame()+1;
+			
+			TFrameState* FrameState = State.WriteFrameWithTime(LastSerializedKeyframe, LastSerializedSimulationTime);
+			npCheckSlow(FrameState);
 
-		if (Ar.IsLoading())
-		{
-			UE_NP_TRACE_NET_SERIALIZE_RECV(SerializedTime, SerializedFrame);
-			UE_NP_TRACE_USER_STATE_SYNC(SerializedSyncState, SerializedFrame);
-			UE_NP_TRACE_USER_STATE_AUX(SerializedAuxState, SerializedFrame);
+			TAuxState* AuxState = State.WriteAuxState(LastSerializedKeyframe);
+			npCheckSlow(AuxState);
 
-			bReconcileFaultDetected = false;
-			bPendingReconciliation = false;
+			FrameState->SyncState.NetSerialize(P.Ar);	// 2. Sync
+			FrameState->InputCmd.NetSerialize(P.Ar);	// 3. Input
+			FrameState->SetFlag(ESimulationFrameStateFlags::InputWritten);
+			
+			AuxState->NetSerialize(P.Ar); // 4. Aux
 
-			TSyncState* ClientSync = Buffers.Sync[SerializedFrame];
-			TAuxState* ClientAux = Buffers.Aux[SerializedFrame];
+			State.CueDispatcher.NetSerializeSavedCues(Ar, GetSimulatedUpdateMode() == ESimulatedUpdateMode::Interpolate ? ENetSimCueReplicationTarget::Interpolators : ENetSimCueReplicationTarget::SimulatedProxy); // 5. Cues
 
-			if (Ticker.SimulationTimeBuffer.IsValidFrame(SerializedFrame))
+			State.SetLatestInputFrame(LastSerializedKeyframe);
+			State.SetMaxTickFrame(LastSerializedKeyframe);
+			State.SetPendingTickFrame(LastSerializedKeyframe);
+			State.SetConfirmedFrame(LastSerializedKeyframe);
+			
+			if (State.GetTotalAllowedSimulationTime() < FrameState->TotalSimTime)
 			{
-				FNetworkSimTime ClientSimTime = *Ticker.SimulationTimeBuffer[SerializedFrame];
-				if (SerializedTime != ClientSimTime)
-				{
-					bPendingReconciliation = true;
-				}
+				State.SetTotalAllowedSimulationTime(FrameState->TotalSimTime);
 			}
 
-			// The state the client predicted that corresponds to the state the server just serialized to us
-			if (ClientSync && ClientAux)
-			{
-				const bool bForceReconcile = (NetworkSimulationModelCVars::ForceReconcile() > 0) || (NetworkSimulationModelCVars::ForceReconcileSingle() > 0);
-				
-				if (bForceReconcile || Model::ShouldReconcile(SerializedSyncState, SerializedAuxState, *ClientSync, *ClientAux))
-				{
-					NetworkSimulationModelCVars::SetForceReconcileSingle(0);
-					UE_CLOG(!Buffers.Input.IsValidFrame(SerializedFrame-1), LogNetworkSim, Error, TEXT("::NetSerialize: Client InputBuffer does not contain data for frame %d. {%s} {%s}"), SerializedFrame, *Buffers.Input.GetBasicDebugStr(), *Buffers.Sync.GetBasicDebugStr());
-					bPendingReconciliation =  true;
-				}
-			}
-			else
-			{
-				if (SerializedFrame < Buffers.Sync.TailFrame())
-				{
-					// Case 1: the serialized state is older than what we've kept in our buffer. A bigger buffer would solve this! (at the price of more resimulated frames to recover when this happens)
-					// This is a reconcile fault and we just need to chill. We'll stop sending user commands until the cmds in flight flush through the system and we catch back up.
-					bReconcileFaultDetected = true;
-					UE_NP_TRACE_NET_SERIALIZE_FAULT();
-				}
-				else
-				{
-					// Case 2: We've received a newer frame than what we've processed locally. This could happen if we are buffering our inputs locally (but still sending to the server) or just not predicting
-					bPendingReconciliation =  true;
-				}
-			}
+			UE_NP_TRACE_NET_SERIALIZE_RECV(LastSerializedSimulationTime, LastSerializedKeyframe);
+
+			UE_NP_TRACE_USER_STATE_INPUT(FrameState->InputCmd, LastSerializedKeyframe);
+			UE_NP_TRACE_USER_STATE_SYNC(FrameState->SyncState, LastSerializedKeyframe);
+			UE_NP_TRACE_USER_STATE_AUX(*AuxState, LastSerializedKeyframe);
+
+			// Temp: this will go away
+			LastSerializedInputCmd = FrameState->InputCmd;
+			LastSerializedSyncState = FrameState->SyncState;
+			LastSerializedAuxState = *AuxState;
 		}
+
 		NETSIM_CHECKSUM(P.Ar);
 	}
 
-	// --------------------------------------------------------------------
-	//	Reconcile
-	// --------------------------------------------------------------------
-	void Reconcile(TSimulation* Simulation, TDriver* Driver, TNetworkSimBufferContainer<Model>& Buffers, TSimulationTicker<TTickSettings>& Ticker)
+	// -------------------------------------------------------------------------------------------------
+	// TRepController_Simulated::Reconcile
+	// -------------------------------------------------------------------------------------------------
+	void Reconcile(TSimulation* Simulation, TDriver* Driver, TNetworkedSimulationState<Model>& State)
 	{
-		if (bPendingReconciliation == false && (bDependentSimulationNeedsReconcile == false || SerializedFrame == INDEX_NONE))
+		if (ReconcileSimulationTime.IsPositive() == false)
 		{
 			return;
 		}
-		bPendingReconciliation = false;
-		bDependentSimulationNeedsReconcile = false;
 
-		TSyncState* ClientSyncState = Buffers.Sync[SerializedFrame];
-		const bool bDoVisualLog = NetworkSimulationModelCVars::EnableLocalPrediction() > 0 && NETSIM_MODEL_DEBUG; // don't visual log if we have prediction disabled
-		
-		if (bDoVisualLog)
-		{
-			FVisualLoggingParameters VLogParameters(EVisualLoggingContext::LastConfirmed, SerializedFrame, EVisualLoggingLifetime::Persistent, TEXT("Serialized State"));
-			Driver->InvokeVisualLog(Buffers.Input[SerializedFrame], &SerializedSyncState, &SerializedAuxState, VLogParameters);
-		}
+		// The last serialized state is now committed. If we received multiple NetSerializes before a reconcile, only the last is committed.
+		UE_NP_TRACE_NET_SERIALIZE_COMMIT();
+
+		FNetworkSimTime TotalSimTime = State.GetTotalProcessedSimulationTime();
+		check(TotalSimTime <= State.GetTotalAllowedSimulationTime());
 
-		if (ClientSyncState)
+		if (GetSimulatedUpdateMode() == ESimulatedUpdateMode::Extrapolate && NetworkSimulationModelCVars::EnableSimulatedReconcile())
 		{
-			// Existing ClientSyncState, log it before overwriting it
-			if (bDoVisualLog)
+			// This is effectively a rollback: we went back in time during NetSerialize and are now going to catch up ReconcileSimulationTime ms.
+			State.CueDispatcher.NotifyRollback(TotalSimTime);
+
+			// Do we have time to make up? We may have extrapolated ahead of the server (totally fine - can happen with small amount of latency variance)
+			FNetworkSimTime DeltaSimTime = ReconcileSimulationTime - TotalSimTime;
+			if (DeltaSimTime.IsPositive() && NetworkSimulationModelCVars::EnableSimulatedReconcile())
 			{
-				FVisualLoggingParameters VLogParameters(EVisualLoggingContext::FirstMispredicted, SerializedFrame, EVisualLoggingLifetime::Persistent, TEXT("Mispredicted State"));
-				Driver->InvokeVisualLog(Buffers.Input[SerializedFrame], ClientSyncState, Buffers.Aux[SerializedFrame], VLogParameters);
+				SimulationExtrapolation(Simulation, Driver, State, DeltaSimTime);
 			}
 		}
-		else
+		
+		check(State.GetTotalProcessedSimulationTime() <= State.GetTotalAllowedSimulationTime());
+		ReconcileSimulationTime.Reset();
+	}
+
+	// -------------------------------------------------------------------------------------------------
+	// TRepController_Simulated::Tick
+	// -------------------------------------------------------------------------------------------------
+	void Tick(TDriver* Driver, TSimulation* Simulation, TNetworkedSimulationState<Model>& State, const FNetSimTickParameters& TickParameters)
+	{
+		// Tick if we are dependent simulation or extrapolation is enabled
+		if (ParentSimulation || (GetSimulatedUpdateMode() == ESimulatedUpdateMode::Extrapolate))
 		{
-			// No existing state, so create add it explicitly
-			ClientSyncState = Buffers.Sync.WriteFrame( SerializedFrame );
+			// Don't start this simulation until you've gotten at least one update from the server
+			if (State.GetTotalProcessedSimulationTime().IsPositive())
+			{
+				State.GiveSimulationTime(TickParameters.LocalDeltaTimeSeconds);
+			}
+			
+			// Generate local input if we are ready to tick. Note that we pass false so we won't call into the Driver to produce the input, we will use the last serialized InputCmd's values
+			TryGenerateLocalInput<Model>(Driver, State, false);
 		}
 
-		UE_NP_TRACE_NET_SERIALIZE_COMMIT();
+		DoSimulationTicks(*Simulation, State, ESimulationTickContext::SimExtrapolate);
 
-		// Set client's sync state to the server version
-		check(ClientSyncState);
-		*ClientSyncState = SerializedSyncState;
-
-		// Set Client's aux state to the server version
-		*Buffers.Aux.WriteFrame(SerializedFrame) = SerializedAuxState;
-
-		const FNetworkSimTime RollbackDeltaTime = SerializedTime - Ticker.GetTotalProcessedSimulationTime();
-
-		// Set the canonical simulation time to what we received (we will advance it as we resimulate)
-		Ticker.SetTotalProcessedSimulationTime(SerializedTime, SerializedFrame);
-		Ticker.PendingFrame = SerializedFrame;
-		//Ticker.MaxAllowedFrame = FMath::Max(Ticker.MaxAllowedFrame, Ticker.PendingFrame); // Make sure this doesn't lag behind. This is the only place we should need to do this.
-
-		if (NetworkSimulationModelCVars::EnableLocalPrediction() == 0)
+		if (GetSimulatedUpdateMode() == ESimulatedUpdateMode::Interpolate)
 		{
-			// If we aren't predicting at all, then we advanced the allowed sim time here, (since we aren't doing it in PreSimTick). This just keeps us constantly falling behind and not being able to toggle prediction on/off for debugging.
-			Ticker.SetTotalAllowedSimulationTime(SerializedTime);
-		}
+			const FNetworkSimTime::FRealTime InterpolationRealTime = Interpolator.template PostSimTick<TDriver>(Driver, State, TickParameters);
+			const FNetworkSimTime InterpolationSimTime = FNetworkSimTime::FromRealTimeSeconds(InterpolationRealTime);
 
-		// Tell dependent simulations to rollback
-		for (INetworkedSimulationModel* DependentSim : DependentSimulations)
-		{
-			DependentSim->BeginRollback(RollbackDeltaTime, SerializedFrame);
+			State.CueDispatcher.SetMaxDispatchTime(InterpolationSimTime);
+			State.CueDispatcher.SetConfirmedTime(InterpolationSimTime);
 		}
-
-		// Tell cue dispatch that we've rolledback as well
-		Buffers.CueDispatcher.NotifyRollback(SerializedTime);
-		
-		// Resimulate all user commands 
-		const int32 LastFrameToProcess = Ticker.MaxAllowedFrame;
-		for (int32 Frame = SerializedFrame; Frame <= LastFrameToProcess; ++Frame)
+		else
 		{
-			const int32 OutputFrame = Frame+1;
-
-			// Frame is the frame we are resimulating right now.
-			TInputCmd* ResimulateCmd  = Buffers.Input[Frame];
-			TAuxState* AuxState = Buffers.Aux[Frame];
-			TSyncState* PrevSyncState = Buffers.Sync[Frame];
-			TSyncState* NextSyncState = Buffers.Sync.WriteFrame(OutputFrame);
-			
-			check(ResimulateCmd);
-			check(PrevSyncState);
-			check(NextSyncState);
-			check(AuxState);
-
-			// Log out the Mispredicted state that we are about to overwrite.
-			if (NETSIM_MODEL_DEBUG)
+			// Sync to latest frame if there is any
+			const int32 PendingFrame = State.GetPendingTickFrame();
+			if (PendingFrame >= 0)
 			{
-				FVisualLoggingParameters VLogParameters(Frame == LastFrameToProcess ? EVisualLoggingContext::LastMispredicted : EVisualLoggingContext::OtherMispredicted, Frame, EVisualLoggingLifetime::Persistent, TEXT("Resimulate Step: mispredicted"));
-				Driver->InvokeVisualLog(Buffers.Input[OutputFrame], Buffers.Sync[OutputFrame], Buffers.Aux[OutputFrame], VLogParameters);
-			}
-
-			// Do the actual update
-			TSimulationDoTickFunc<Model>::DoTick(*Simulation, Ticker, Buffers, OutputFrame, ESimulationTickContext::Resimulate);
-
-			// Log out the newly predicted state that we got.
-			if (NETSIM_MODEL_DEBUG)
-			{			
-				FVisualLoggingParameters VLogParameters(Frame == LastFrameToProcess ? EVisualLoggingContext::LastPredicted : EVisualLoggingContext::OtherPredicted, Frame, EVisualLoggingLifetime::Persistent, TEXT("Resimulate Step: repredicted"));
-				Driver->InvokeVisualLog(Buffers.Input[OutputFrame], Buffers.Sync[OutputFrame], Buffers.Aux[OutputFrame], VLogParameters);
+				Driver->FinalizeFrame(State.GetFrameState(PendingFrame)->SyncState, *State.GetAuxState(PendingFrame));
 			}
 
-			// Tell dependent simulations to advance
-			for (INetworkedSimulationModel* DependentSim : DependentSimulations)
-			{
-				DependentSim->StepRollback(ResimulateCmd->GetFrameDeltaTime(), Frame, (Frame == LastFrameToProcess));
-			}
+			State.CueDispatcher.ClearMaxDispatchTime();
+			State.CueDispatcher.SetConfirmedTime(LastSerializedSimulationTime);
 		}
 	}
-
-	// --------------------------------------------------------------------
-	//	PreSimTick
-	// --------------------------------------------------------------------
-	void PreSimTick(TDriver* Driver, TNetworkSimBufferContainer<Model>& Buffers, TSimulationTicker<TTickSettings>& Ticker, const FNetSimTickParameters& TickParameters)
+	
+	// -------------------------------------------------------------------------------------------------
+	// TRepController_Simulated::DependentRollbackBegin
+	// -------------------------------------------------------------------------------------------------
+	void DependentRollbackBegin(TSimulation* Simulation, TDriver* Driver, TNetworkedSimulationState<Model>& State, const FNetworkSimTime& RollbackDeltaTime, const int32 ParentFrame)
 	{
-		// If we have a reconcile fault, we cannot continue on with the simulation until it clears itself out. This effectively drops the input time and does not sample new inputs
-		if (bReconcileFaultDetected)
-		{
-			return;
-		}
+		// FIXME: this isn't using RollbackDeltaTime
 
-		if (TickParameters.bGenerateLocalInputCmds)
-		{
-			if (NetworkSimulationModelCVars::EnableLocalPrediction() > 0)
-			{
-				// Prediction: add simulation time and generate new commands
-				Ticker.GiveSimulationTime(TickParameters.LocalDeltaTimeSeconds);
-				TryGenerateLocalInput<Model>(Driver, Buffers, Ticker);
-			}
-			else
-			{
-				// Local prediction disabled: we must use a separate time accumulator to figure out when we should add more input cmds.
-				// Since we aren't processing the simulation locally, our core simulation time will only advance from network updates.
-				// (still, we need *something* to tell us when to generate a new command and what the delta time should be)
-				FNetworkSimTime NonPredictedInputTime;
-				NonPredictedInputTimeAccumulator.Accumulate(NonPredictedInputTime, TickParameters.LocalDeltaTimeSeconds);
-				if (NonPredictedInputTime.IsPositive())
-				{
-					GenerateLocalInputCmdAtFrame<Model>(Driver, Buffers, NonPredictedInputTime, Buffers.Input.HeadFrame() + 1);
-				}
-			}
-		}
-	}
+		// Rollback to the last serialized frame
+		// For now, we make the assumption that our last serialized state and time match the parent simulation.
+		// This would not always be the case with low frequency simulated proxies. But could be handled by replicating the simulations together (at the replication level)
+		// npEnsure(Buffers.IsValidFrame(LastSerializedSimulationTime));
+		npEnsure(State.WriteFrameSequential(LastSerializedKeyframe)); // fixme
+		
+		State.SetPendingTickFrame(LastSerializedKeyframe);
+		State.SetMaxTickFrame(LastSerializedKeyframe);
 
-	void PostSimTick(TDriver* Driver, const TNetworkSimBufferContainer<Model>& Buffers, const TSimulationTicker<TTickSettings>& Ticker, const FNetSimTickParameters& TickParameters)
+		State.CueDispatcher.NotifyRollback(LastSerializedSimulationTime);
+	}
+		
+	void DependentRollbackStep(TSimulation* Simulation, TDriver* Driver, TNetworkedSimulationState<Model>& State, const FNetworkSimTime& StepTime, const int32 ParentFrame, const bool bFinalStep)
 	{
-		TBase::PostSimTick(Driver, Buffers, Ticker, TickParameters);
-
-		Buffers.CueDispatcher.ClearMaxDispatchTime();
-		Buffers.CueDispatcher.SetConfirmedTime(SerializedTime);
+		State.SetTotalAllowedSimulationTime( State.GetTotalAllowedSimulationTime() + StepTime );
+		SimulationExtrapolation(Simulation, Driver, State, StepTime);
 	}
-	
 
 private:
 	
-	TSyncState SerializedSyncState;
-	TAuxState SerializedAuxState;
-	FNetworkSimTime SerializedTime;
-	int32 SerializedFrame = -1;
-
-	bool bPendingReconciliation = false;	// Reconciliation is pending: we need to reconcile state from the server that differs from the locally predicted state
-	bool bReconcileFaultDetected = false;	// A fault was detected: we received state from the server that we are unable to reconcile with locally predicted state
-
-	TRealTimeAccumulator<TTickSettings> NonPredictedInputTimeAccumulator; // for tracking input time in the non predictive case
-};
-
-/** Special replicator for debug buffer, this preserves the local buffer and receives into a replicator-owned buffer (we want these buffers to be distinct/not merged) */
-template<typename Model, int32 MaxNumElements=5, typename TBase=TRepController_Empty<Model>>
-struct TRepController_Debug : public TBase
-{
-	using typename TBase::TSimulation;
-	using typename TBase::TTickSettings;
-	using typename TBase::TDriver;
-	using typename TBase::TBufferTypes;
-
-	using TDebugState = typename TBase::TBufferTypes::TDebugState;
-	using TDebugBuffer = typename TNetworkSimBufferContainer<Model>::TDebugBuffer;
-
-	int32 GetProxyDirtyCount(TNetworkSimBufferContainer<Model>& Buffers) const 
+	void SimulationExtrapolation(TSimulation* Simulation, TDriver* Driver, TNetworkedSimulationState<Model>& State, const FNetworkSimTime DeltaSimTime)
 	{
-		if (NetworkSimulationModelCVars::EnableDebugBufferReplication() == 0)
-		{
-			return 0;
-		}
-
-		return Buffers.Debug.GetDirtyCount() ^ (TBase::GetProxyDirtyCount(Buffers) << 2); 
-	}
-
-	void NetSerialize(const FNetSerializeParams& P, TNetworkSimBufferContainer<Model>& Buffers, TSimulationTicker<TTickSettings>& Ticker)
-	{
-		NETSIM_CHECKSUM(P.Ar);
-		TBase::NetSerialize(P, Buffers, Ticker);
-		FArchive& Ar = P.Ar;
+		// We have extrapolated ahead of the server. The latest network update is now "in the past" from what we rendered last frame.
+		// We will insert a new frame to make up the difference from the last known state to where we want to be in the now.
 
-		TDebugBuffer& Buffer = Ar.IsSaving() ? Buffers.Debug : ReceivedBuffer;
-		
-		uint8 SerializedNumElements = FMath::Min<uint8>(MaxNumElements, Buffer.Num());
-		Ar << SerializedNumElements;
+		const int32 InputFrame = State.GetPendingTickFrame();
+		const int32 OutputFrame = InputFrame + 1;
 
-		const int32 HeadFrame = FNetworkSimulationSerialization::SerializeFrame(Ar, Buffer.HeadFrame());
-		const int32 StartingFrame = FMath::Max(0, HeadFrame - SerializedNumElements + 1);
+		// Override the FrameDeltaTime
+		TFrameState* InputFrameState = State.GetFrameState(InputFrame);
+		npCheckSlow(InputFrameState);
+		npEnsureSlow(InputFrameState->HasFlag(ESimulationFrameStateFlags::InputWritten));
+		InputFrameState->FrameDeltaTime.Set(DeltaSimTime);
 
-		if (Ar.IsLoading())
-		{
-			const int32 PrevHead = Buffer.HeadFrame();
-			if (PrevHead < StartingFrame && PrevHead >= 0)
-			{
-				// There is a gap in the stream. In some cases, we want this to be a "fault" and bubble up. We may want to synthesize state or maybe we just skip ahead.
-				UE_LOG(LogNetworkSim, Warning, TEXT("Fault: gap in received Debug buffer. PrevHead: %d. Received: %d-%d. Reseting previous buffer contents"), PrevHead, StartingFrame, HeadFrame);
-			}
-		}
-
-		for (int32 Frame = StartingFrame; Frame <= HeadFrame; ++Frame)
-		{
-			// This, as is, is bad. The intention is that these functions serialize multiple items in some delta compressed fashion.
-			// As is, we are just serializing the elements individually.
-			auto* Cmd = Ar.IsLoading() ? Buffer.WriteFrame(Frame) : Buffer[Frame];
-			Cmd->NetSerialize(P);
-		}
-		NETSIM_CHECKSUM(P.Ar);
+		// Do the actual update
+		State.SetMaxTickFrame(InputFrame);
+		TSimulationDoTickFunc<Model>::DoTick(*Simulation, State, ESimulationTickContext::ResimExtrapolate);
 	}
 
-	TDebugBuffer ReceivedBuffer;
+	int32 LastSerializedKeyframe;
+	FNetworkSimTime ReconcileSimulationTime;
+
+	// Temp: this should go away
+	FNetworkSimTime LastSerializedSimulationTime;
+	TInputCmd LastSerializedInputCmd;
+	TSyncState LastSerializedSyncState;
+	TAuxState LastSerializedAuxState;
 };
\ No newline at end of file
diff --git a/Engine/Plugins/Runtime/NetworkPrediction/Source/NetworkPrediction/Public/NetworkedSimulationModelTick.h b/Engine/Plugins/Runtime/NetworkPrediction/Source/NetworkPrediction/Public/NetworkedSimulationModelTick.h
--- a/Engine/Plugins/Runtime/NetworkPrediction/Source/NetworkPrediction/Public/NetworkedSimulationModelTick.h
+++ b/Engine/Plugins/Runtime/NetworkPrediction/Source/NetworkPrediction/Public/NetworkedSimulationModelTick.h
@@ -84,89 +84,6 @@ private:
 	TRealTime AccumulatedTime;
 };
 
-// ----------------------------------------------------------------------------------------------------------------------------------------------
-//	FSimulationTickState: Holds active state for simulation ticking. We track two things: frames and time.
-//
-//	PendingFrame is the next frame we will process: the input/sync/aux state @ PendingFrame will be run through ::SimulationTick and produce the 
-//	next frame's (PendingFrame+1) Sync and possibly Aux state (if it changes). "Out of band" modifications to the sync/aux state should happen
-//	to PendingFrame (e.g, before it is processed. Once a frame has been processed, we won't run it through ::SimulationTick again!).
-//
-//	MaxAllowedFrame is a frame based limiter on simulation updates. This must be incremented to allow the simulation to advance.
-//
-//	Time is also tracked. We keep running total for how much the sim has advanced and how much it is allowed to advance. There is also a historic buffer of
-//	simulation time in SimulationTimeBuffer.
-//
-//	Consider that Frames are essentially client dependent and gaps can happen due to packet loss, etc. Time will always be continuous though.
-//	
-// ----------------------------------------------------------------------------------------------------------------------------------------------
-
-struct FSimulationTickState
-{
-	int32 PendingFrame = 0;
-	int32 MaxAllowedFrame = -1;
-	bool bUpdateInProgress = false;
-	
-	FNetworkSimTime GetTotalProcessedSimulationTime() const 
-	{ 
-		return TotalProcessedSimulationTime; 
-	}
-
-	void SetTotalProcessedSimulationTime(const FNetworkSimTime& SimTime, int32 Frame)
-	{
-		TotalProcessedSimulationTime = SimTime;
-		*SimulationTimeBuffer.WriteFrame(Frame) = SimTime;
-	}
-
-	void IncrementTotalProcessedSimulationTime(const FNetworkSimTime& DeltaSimTime, int32 Frame)
-	{
-		TotalProcessedSimulationTime += DeltaSimTime;
-		*SimulationTimeBuffer.WriteFrame(Frame) = TotalProcessedSimulationTime;
-	}
-		
-	// Historic tracking of simulation time. This allows us to timestamp sync data as its produced
-	TNetworkSimContiguousBuffer<FNetworkSimTime>	SimulationTimeBuffer;
-
-	// How much granted simulation time is left to process
-	FNetworkSimTime GetRemainingAllowedSimulationTime() const
-	{
-		return TotalAllowedSimulationTime - TotalProcessedSimulationTime;
-	}
-
-	FNetworkSimTime GetTotalAllowedSimulationTime() const
-	{
-		return TotalAllowedSimulationTime;
-	}
-
-protected:
-
-	FNetworkSimTime TotalAllowedSimulationTime;	// Total time we have been "given" to process. We cannot process more simulation time than this: doing so would be speed hacking.
-	FNetworkSimTime TotalProcessedSimulationTime;	// How much time we've actually processed. The only way to increment this is to process user commands or receive authoritative state from the network.
-};
-
-// "Ticker" that actually allows us to give the simulation time. This struct will generally not be passed around outside of the core TNetworkedSimulationModel/Replicators
-template<typename TickSettings=TNetworkSimTickSettings<>>
-struct TSimulationTicker : public FSimulationTickState
-{
-	using TSettings = TickSettings;
-
-	void SetTotalAllowedSimulationTime(const FNetworkSimTime& SimTime)
-	{
-		TotalAllowedSimulationTime = SimTime;
-		RealtimeAccumulator.Reset();
-	}
-
-	// "Grants" allowed simulation time to this tick state. That is, we are now allowed to advance the simulation by this amount the next time the sim ticks.
-	// Note the input is RealTime in SECONDS. This is what the rest of the engine uses when dealing with float delta time.
-	void GiveSimulationTime(float RealTimeSeconds)
-	{
-		RealtimeAccumulator.Accumulate(TotalAllowedSimulationTime, RealTimeSeconds);
-	}	
-
-private:
-
-	TRealTimeAccumulator<TSettings>	RealtimeAccumulator;	
-};
-
 // ----------------------------------------------------------------------------------------------------------------------------------------------
 //	FNetSimTimeStep Time data structure that is passed into SimulationTick
 // ----------------------------------------------------------------------------------------------------------------------------------------------
@@ -175,7 +92,12 @@ struct FNetSimTimeStep
 {
 	// The delta time step for this tick (in MS by default)
 	const FNetworkSimTime& StepMS;
-	// The tick state of the simulation prior to running this tick. E.g, does not "include" the above StepMS that we are simulating now.
-	// The first time ::SimulationTick runs, TickState.GetTotalProcessedSimulationTime() == 0.
-	const FSimulationTickState& TickState;
+
+	// How much simulation time has ran up until this point.
+	// This will be 0 the first time ::SimulationTick runs.
+	const FNetworkSimTime& TotalSimulationTime;
+
+	// The Simulation Frame number we are generating in this tick, E.g, the output frame.
+	// This will be 1 the first time ::SimulationTick runs. (0 is the starting input and is not generated in a Tick)
+	const int32& Frame;
 };
\ No newline at end of file
diff --git a/Engine/Plugins/Runtime/NetworkPrediction/Source/NetworkPrediction/Public/NetworkedSimulationModelTraits.h b/Engine/Plugins/Runtime/NetworkPrediction/Source/NetworkPrediction/Public/NetworkedSimulationModelTraits.h
--- a/Engine/Plugins/Runtime/NetworkPrediction/Source/NetworkPrediction/Public/NetworkedSimulationModelTraits.h
+++ b/Engine/Plugins/Runtime/NetworkPrediction/Source/NetworkPrediction/Public/NetworkedSimulationModelTraits.h
@@ -2,9 +2,20 @@
 #pragma once
 #include "NetworkedSimulationModelBufferTypes.h"
 
-/** This is the "system driver", it has functions that the TNetworkedSimulationModel needs to call internally, that are specific to the types but not specific to the simulation itself. */
+// This is the "system driver", it has functions that the TNetworkedSimulationModel needs to call internally, 
+// that are specific to the types but not specific to the simulation itself. This is main:
+//	Finalize Frame: how to push the sync/aux state to the UE4 side.
+//	Produce Input: how to produce the InputCmd when necessary.
+//
+// The driver is templatized on the BufferTypes. This is so that the driver (UE4 actor/component) just has to know about
+// the underlying user structure types, and not about the networking model or simulation class.
+//
+// If we templatized the driver on the Model def, we would need separate UE4 actor/component for fixed vs non fixed tick simulations.
+// This way, we can define a single UE4 piece that just handles the user input/sync/aux state and doesn't care about the other stuff
+// defined in the Model def.
+
 template<typename TBufferTypes>
-class TNetworkedSimulationModelDriver
+class TNetworkedSimulationModelDriverBase
 {
 public:
 	using TInputCmd = typename TBufferTypes::TInputCmd;
@@ -57,19 +68,9 @@ protected:
 
 	// Called to visual log the given states. Note that not all 3 will always be present and you should check for nullptrs.
 	virtual void VisualLog(const TInputCmd* Input, const TSyncState* Sync, const TAuxState* Aux, const FVisualLoggingParameters& SystemParameters) const = 0;
-
 };
 
-// Internal traits that are derived from the user settings (E.g, these can't go in FNetSimModelDefBase because they are dependent on the user types)
-// These can be overridden as well via template specialization but we don't expect this to be common.
-template<typename T>
-struct TNetSimModelTraitsBase
-{
-	// Internal buffer types: the model may wrap the user structs (e.g, input cmds + frame time) for internal usage
-	using InternalBufferTypes = TInternalBufferTypes< typename T::BufferTypes, typename T::TickSettings>;
-	
-	// Driver class: interface that the NetSimModel needs user code to implement in order to function
-	using Driver = TNetworkedSimulationModelDriver<typename T::BufferTypes>;
-};
-
-template<typename T> struct TNetSimModelTraits : public TNetSimModelTraitsBase<T> { };
\ No newline at end of file
+// Outward facing class. This allows user simulations to extend the driver interface if necessary.
+// Use case would be if writing custom RepControllers that required special communication with the driver.
+template<typename TBufferTypes>
+class TNetworkedSimulationModelDriver : public TNetworkedSimulationModelDriverBase<TBufferTypes> { };
\ No newline at end of file
diff --git a/Engine/Plugins/Runtime/NetworkPrediction/Source/NetworkPrediction/Public/NetworkedSimulationModelTypes.h b/Engine/Plugins/Runtime/NetworkPrediction/Source/NetworkPrediction/Public/NetworkedSimulationModelTypes.h
--- a/Engine/Plugins/Runtime/NetworkPrediction/Source/NetworkPrediction/Public/NetworkedSimulationModelTypes.h
+++ b/Engine/Plugins/Runtime/NetworkPrediction/Source/NetworkPrediction/Public/NetworkedSimulationModelTypes.h
@@ -7,138 +7,324 @@
 #include "NetworkedSimulationModelTime.h"
 #include "NetworkedSimulationModelTraits.h"
 #include "Trace/NetworkPredictionTrace.h"
+#include "NetworkPredictionCheck.h"
 
 // ---------------------------------------------------------------------------------------------------------------------
-//	TNetworkSimBufferContainer
-//	Container for the actual replicated buffers that the system uses.
-//	Has compile time accessors for retrieving the buffers based on templated enum value.
+//	TNetworkedSimulationState:
+//
+//	This structure contains all the networked state used by the system.
+//	
+//	Simulation's state:
+//		-FrameBuffer:	frame-to-frame state
+//		-AuxBuffer:		sparse frame based state
+//		-CueDispatcher: non simulation affecting event history, mappings
+//		-TimingInfo:	time keeper for advancing simulation ticks
+//	
+//	-Frame Markers
+//		-Pending		Next frame that will be ticked. Latest data.
+//		-LatestInput	Latest frame to have input written to it.
+//		-NextInput		Where next input should be written to.
+//		-MaxTick		Max frame that is allowed to be ticked. To facilitate input buffering when necessary.
+//		-Confirmed		Last valid frame. Frames prior to confirmed frame do not need to be kept around.
+//
 // ---------------------------------------------------------------------------------------------------------------------
 
-// Helper struct for enum-based access. This has to be done in an outside struct because we cant specialize inside TNetworkSimBufferContainer on all compilers
-template<typename TContainer, ENetworkSimBufferTypeId BufferId>
-struct TBufferGetterHelper
+template<typename Model>
+struct TNetworkedSimulationState
 {
-	static typename TContainer::template select_buffer_type<BufferId>::type& Get(TContainer& Container)
+	// Collection of types we were assigned
+	using TBufferTypes = typename Model::BufferTypes;
+	using TickSettings = typename Model::TickSettings;
+
+	using TInputCmd = typename TBufferTypes::TInputCmd;
+	using TSyncState = typename TBufferTypes::TSyncState;
+	using TAuxState = typename TBufferTypes::TAuxState;
+
+	TNetSimCueDispatcher<Model> CueDispatcher;
+
+	TNetworkedSimulationState()
+		: FrameBuffer(64), AuxBuffer(64)
 	{
-		static_assert(!Container, "Failed to find specialized Get for your BufferId");
 	}
-};
 
-template<typename TContainer>
-struct TBufferGetterHelper<TContainer, ENetworkSimBufferTypeId::Input>
-{
-	static typename TContainer::TInputBuffer& Get(TContainer& Container)
+	int32 Max() const { return FrameBuffer.Capacity(); }
+	FString GetBasicDebugStr() const { return FString::Printf(TEXT("PendingTick: %d MaxTick: %d LatestInput: %d NextIntput: %d Size: %d. PT: %s. IT: %s. MT: %s"), 
+		PendingTickFrame, MaxTickFrame, LatestInputFrame, NextInputFrame, Max(), *GetTotalProcessedSimulationTime().ToString(), *GetTotalInputSimulationTime().ToString(), *GetTotalAllowedSimulationTime().ToString()); }
+	
+	// ---------------------------------------------------------------------
+	// Frame State
+	// ---------------------------------------------------------------------
+
+	// Writes a new frame to the buffer, with the given explicit TotalSimTime.
+	// This is used in places where we might break sim time continuity (NetRecvs) or are possibly resimulating over a frame and want to explicitly set the TotalSimTime.
+	TSimulationFrameState<Model>* WriteFrameWithTime(int32 NewFrame, FNetworkSimTime TotalSimTime)
 	{
-		return Container.Input;
+		if (TSimulationFrameState<Model>* ExistingFrame = GetValidFrame(NewFrame))
+		{
+			ExistingFrame->TotalSimTime = TotalSimTime;
+			return ExistingFrame;
+		}
+		
+		const int32 Delta = NewFrame - PendingTickFrame;
+		if (Delta >= Max())
+		{
+			// PendingTickFrame must always be preserved
+			UE_NP_TRACE_SYSTEM_FAULT("WriteFrameWithTime frame %d outside of buffer range. Reseting contents. %s", PendingTickFrame, *GetBasicDebugStr());
+				
+			TSimulationFrameState<Model>& NewFrameState = FrameBuffer[NewFrame];
+			NewFrameState = FrameBuffer[PendingTickFrame]; // Init with previous state
+			NewFrameState.TotalSimTime = TotalSimTime;
+			NewFrameState.ClearFlags();
+				
+			PendingTickFrame = NewFrame;
+			MaxTickFrame = NewFrame-1;
+			LatestInputFrame = INDEX_NONE;
+			NextInputFrame = NewFrame;
+			CheckInvariants();
+
+			return &NewFrameState;
+		}
+			
+		TSimulationFrameState<Model>& NewFrameState = FrameBuffer[NewFrame];
+		NewFrameState.TotalSimTime = TotalSimTime;
+
+		for (int32 Frame = LatestInputFrame+1; Frame <= NewFrame; ++Frame)
+		{
+			FrameBuffer[Frame].ClearFlags();
+		}
+		
+		return &NewFrameState;
 	}
-};
-template<typename TContainer>
-struct TBufferGetterHelper<TContainer, ENetworkSimBufferTypeId::Sync>
-{
-	static typename TContainer::TSyncBuffer& Get(TContainer& Container)
+
+	// Writes a sequential frame to the buffer, even if it creates a gap. This is used in places where we want to maintain continuity of TotalSimTime for the new frame.
+	// That is, regardless if frame# is not continuous, we want this new frame to have the correct TotalSimTime, based on the previous valid frame.
+	//
+	// If the frame already exists, will just return the frame
+	// If this causes a new frame to be created, its flags will be cleared and TotalSimTime will be set for you. But contents will otherwise be stale
+	// If a gap is created, those frames will also have flags cleared
+	TSimulationFrameState<Model>* WriteFrameSequential(int32 NewFrame)
 	{
-		return Container.Sync;
+		if (TSimulationFrameState<Model>* ExistingFrame = GetValidFrame(NewFrame))
+		{
+			CheckInvariants();
+			return ExistingFrame;
+		}
+				
+		const int32 Delta = NewFrame - PendingTickFrame;
+		if (Delta >= Max())
+		{
+			// PendingTickFrame must always be preserved
+			UE_NP_TRACE_SYSTEM_FAULT("WriteFrameSequential frame %d outside of buffer range. Reseting contents. %s", PendingTickFrame, *GetBasicDebugStr());
+
+			TSimulationFrameState<Model>& NewFrameState = FrameBuffer[NewFrame];
+			NewFrameState = FrameBuffer[PendingTickFrame]; // Init with previous state
+
+			PendingTickFrame = NewFrame;
+			MaxTickFrame = NewFrame-1;
+
+			if (NewFrameState.HasFlag(ESimulationFrameStateFlags::InputWritten))
+			{
+				LatestInputFrame = NewFrame;
+				NextInputFrame = NewFrame+1;
+			}
+			else
+			{
+				LatestInputFrame = INDEX_NONE;
+				NextInputFrame = NewFrame;
+			}
+			
+			CheckInvariants();
+			return &NewFrameState;
+		}
+
+		// Create the new frame
+		TSimulationFrameState<Model>& NewFrameState = FrameBuffer[NewFrame];
+		NewFrameState.ClearFlags();
+
+		// Clear gap frame flags so they don't get processed
+		const int32 PrevFrame = FMath::Max(PendingTickFrame, LatestInputFrame);
+		for (int32 Frame = PrevFrame+1; Frame < NewFrame; ++Frame)
+		{
+			FrameBuffer[Frame].ClearFlags();
+		}
+
+		// Look at the previous frame sequential frame and calculate the new Frame's TotalSimTime
+		TSimulationFrameState<Model>& PrevFrameState = FrameBuffer[PrevFrame];
+		
+		if (PrevFrameState.HasFlag(ESimulationFrameStateFlags::InputWritten))
+		{
+			NewFrameState.TotalSimTime = PrevFrameState.FrameDeltaTime.Get() + PrevFrameState.TotalSimTime;
+		}
+		else
+		{
+			NewFrameState.TotalSimTime = PrevFrameState.TotalSimTime;
+		}
+		CheckInvariants();
+		return &NewFrameState;
 	}
-};
-template<typename TContainer>
-struct TBufferGetterHelper<TContainer, ENetworkSimBufferTypeId::Aux>
-{
-	static typename TContainer::TAuxBuffer& Get(TContainer& Container)
+
+	TSimulationFrameState<Model>* WriteFrameInitial()
 	{
-		return Container.Aux;
+		return &FrameBuffer[0];
 	}
-};
-template<typename TContainer>
-struct TBufferGetterHelper<TContainer, ENetworkSimBufferTypeId::Debug>
-{
-	static typename TContainer::TDebugBuffer& Get(TContainer& Container)
-	{
-		return Container.Debug;
+
+	TSimulationFrameState<Model>* GetFrameState(int32 Frame) { return GetValidFrame(Frame); }
+	const TSimulationFrameState<Model>* GetFrameState(int32 Frame) const { return GetValidFrame(Frame); }
+
+	// ---------------------------------------------------------------------
+	// Aux State
+	// ---------------------------------------------------------------------
+	
+	TAuxState* GetAuxState(int32 Frame) { return AuxBuffer[Frame]; }
+	const TAuxState* GetAuxState(int32 Frame) const { return AuxBuffer[Frame]; }
+
+	TNetSimLazyWriterFunc<TAuxState> GetAuxStateLazyWriter(int32 Frame) { return AuxBuffer.LazyWriter(Frame); }
+	TAuxState* WriteAuxState(int32 Frame) { return AuxBuffer.WriteAtFrame(Frame); }
+
+	// ---------------------------------------------------------------------
+	// Ticking
+	// ---------------------------------------------------------------------
+
+	// Returns how much time has been ticked
+	FNetworkSimTime GetTotalProcessedSimulationTime() const { return FrameBuffer[PendingTickFrame].TotalSimTime; }
+
+	// Returns how much time worth of input has been submitted. E.g, if every submitted frame was processed, what would the TotalProcesedSimulationTime be?
+	FNetworkSimTime GetTotalInputSimulationTime() const
+	{ 
+		if (LatestInputFrame != INDEX_NONE)
+		{
+			return FrameBuffer[LatestInputFrame].TotalSimTime + FrameBuffer[LatestInputFrame].FrameDeltaTime.Get();
+		}
+		return FrameBuffer[PendingTickFrame].TotalSimTime;
 	}
-};
 
-// Helper for accessing a netsim buffer given the underlying type
-namespace NetSimBufferSelect
-{
-	template<typename TContainer, typename TState>
-	typename TEnableIf< TIsDerivedFrom<typename TContainer::TInputCmd, TState>::IsDerived, typename TContainer::TInputBuffer&>::Type
-	Get(TContainer& Container)
+	// How much time has been given to the simulation to be ticked
+	FNetworkSimTime GetTotalAllowedSimulationTime() const {	return TotalAllowedSimulationTime; }
+
+	void SetTotalAllowedSimulationTime(const FNetworkSimTime& SimTime)
 	{
-		return Container.Input;
+		TotalAllowedSimulationTime = SimTime;
+		RealtimeAccumulator.Reset();
 	}
 
-	template<typename TContainer, typename TState>
-	typename TEnableIf< TIsDerivedFrom<typename TContainer::TSyncState, TState>::IsDerived, typename TContainer::TSyncBuffer&>::Type
-	Get(TContainer& Container)
+	// "Grants" allowed simulation time to this tick state. That is, we are now allowed to advance the simulation by this amount the next time the sim ticks.
+	// Note the input is RealTime in SECONDS. This is what the rest of the engine uses when dealing with float delta time.
+	void GiveSimulationTime(float RealTimeSeconds)
+	{
+		RealtimeAccumulator.Accumulate(TotalAllowedSimulationTime, RealTimeSeconds);
+	}
+	
+	// Scoped bool set while running a ::SimulationTick
+	bool GetTickInProgress() const { return bTickInProgress; }	
+	void SetTickInProgress(bool b) { bTickInProgress = b; }
+		
+	// The next frame that will be ticked: used as Input to create a new frame at PendingTickFrame+1.
+	// This is where all OOB mods to state should happen. E.g, the "mutable" frame. Must be preserved whenever writing new frames.
+	int32 GetPendingTickFrame() const { return PendingTickFrame; }
+	void SetPendingTickFrame(int32 Frame) { PendingTickFrame = Frame; }
+
+	// Oldest valid frame. This is set explicitly as authoritative state is received by clients. 
+	// On server, it is the tail frame of the buffer: the oldest valid frame.
+	int32 GetConfirmedFrame() const { return ConfirmedFrame == INDEX_NONE ? TailFrame() : ConfirmedFrame; }
+	void SetConfirmedFrame(int32 Frame) { ConfirmedFrame = Frame; CheckInvariants(); }
+
+	// Max frame we are allowed to process. This is how local buffer can be done.
+	int32 GetMaxTickFrame() const { return MaxTickFrame; }
+	void SetMaxTickFrame(int32 Frame) { MaxTickFrame = Frame; CheckInvariants(); }
+
+	// The last frame that had input written to it. This can be ahead of PendingTickFrame but never more than 1 frame behind.
+	int32 GetLatestInputFrame() const { return LatestInputFrame; }
+	void SetLatestInputFrame(int32 Frame)
 	{
-		return Container.Sync;
+		LatestInputFrame = Frame;
+		NextInputFrame = Frame == INDEX_NONE ? PendingTickFrame : (Frame+1);
+		CheckInvariants(); 
 	}
 
-	template<typename TContainer, typename TState>
-	typename TEnableIf< TIsDerivedFrom<typename TContainer::TAuxState, TState>::IsDerived, typename TContainer::TAuxBuffer&>::Type
-	Get(TContainer& Container)
+	// Where input is expected to be written to next. Gaps in NetSerialization can cause this to not be the case though.
+	int32 GetNextInputFrame() const { return NextInputFrame; }
+
+	// Resets the system to the existing NewPendingTickFrame. This is used when we need to hard reset things to a given authoritative state
+	void Reset(int32 NewPendingTickFrame)
 	{
-		return Container.Aux;
+		PendingTickFrame = NewPendingTickFrame;
+		ConfirmedFrame = NewPendingTickFrame;
+		MaxTickFrame = INDEX_NONE;
+		LatestInputFrame = INDEX_NONE;
+		NextInputFrame = NewPendingTickFrame;
+
+		TotalAllowedSimulationTime = FrameBuffer[PendingTickFrame].TotalSimTime;
+		CheckInvariants();
 	}
 
-	template<typename TContainer, typename TState>
-	typename TEnableIf< TIsDerivedFrom<typename TContainer::TDebugState, TState>::IsDerived, typename TContainer::TDebugState&>::Type
-	Get(TContainer& Container)
+	void CheckInvariants()
 	{
-		return Container.Debug;
+		npEnsureMsgf(IsValidFrame(PendingTickFrame), TEXT("PendingTickFrame not valid. %s"), *GetBasicDebugStr());
+		npEnsureMsgf(ConfirmedFrame <= PendingTickFrame, TEXT("ConfirmedFrame %d <= PendingTickFrame"), ConfirmedFrame, PendingTickFrame);
+
+		if (LatestInputFrame != INDEX_NONE)
+		{
+			npEnsureMsgf(MaxTickFrame <= LatestInputFrame, TEXT("MaxTickFrame %d <= LatestInputFrame %d"), MaxTickFrame, LatestInputFrame);
+			npEnsureMsgf(PendingTickFrame <= LatestInputFrame+1, TEXT("PendingTickFrame %d <= (LatestInputFrame+1) %d"), PendingTickFrame, LatestInputFrame+1);
+			npEnsureMsgf(LatestInputFrame == NextInputFrame-1, TEXT("LatestInputFrame %d != (NextInputFrame+1) %d"),  LatestInputFrame, NextInputFrame+1);			
+			npEnsureMsgf(IsValidFrame(LatestInputFrame) && FrameBuffer[LatestInputFrame].HasFlag(ESimulationFrameStateFlags::InputWritten), TEXT("LatestInputFrame not valid. %s"), *GetBasicDebugStr());
+		}
+
+		if (MaxTickFrame != INDEX_NONE)
+		{
+			npEnsureMsgf(PendingTickFrame <= MaxTickFrame+1, TEXT("PendingTickFrame %d <= (MaxTickFrame+1) %d"), PendingTickFrame, MaxTickFrame+1);
+		}
+
+#if NP_CHECKS_AND_ENSURES_SLOW
+		// If PendingTickFrame is ahead of input, ensure that calculated TotalSimTime is what we expected.
+		// (If this fires, its likely LatestInputFrame is wrong. If we took a frame from a NetSerialize for example,
+		// latest InputFrame should be reset.
+		if (LatestInputFrame != INDEX_NONE)
+		{
+			if (LatestInputFrame < PendingTickFrame)
+			{
+				ensure(LatestInputFrame+1 == PendingTickFrame);
+				FNetworkSimTime PendingFrameTotalSimTime = FrameBuffer[PendingTickFrame].TotalSimTime;
+				FNetworkSimTime InputFramesExpectedTotalSimTime = FrameBuffer[LatestInputFrame].TotalSimTime + FrameBuffer[LatestInputFrame].FrameDeltaTime.Get();
+				ensure(InputFramesExpectedTotalSimTime == PendingFrameTotalSimTime);
+			}
+
+			npEnsureMsgf(FrameBuffer[LatestInputFrame].HasFlag(ESimulationFrameStateFlags::InputWritten), TEXT("LatestInputFrame %d has no InputWritten"), LatestInputFrame);
+		}
+#endif
+
 	}
-}
 
-// The main container for all of our buffers
-template<typename Model>
-struct TNetworkSimBufferContainer
-{
-	// Collection of types we were assigned
-	using TBufferTypes = typename TNetSimModelTraits<Model>::InternalBufferTypes;
+private:
 
-	using TInputCmd = typename TBufferTypes::TInputCmd;
-	using TSyncState = typename TBufferTypes::TSyncState;
-	using TAuxState = typename TBufferTypes::TAuxState;
-	using TDebugState = typename TBufferTypes::TDebugState;
+	int32 HeadFrame() const { return FMath::Max(PendingTickFrame, LatestInputFrame); }
+	int32 TailFrame() const { return FMath::Max(ConfirmedFrame, HeadFrame() - FrameBuffer.Capacity() + 1); }
+	bool IsValidFrame(int32 Frame) const { return Frame >= TailFrame() && Frame <= HeadFrame(); }
 
-	// helper that returns the buffer type for a given BufferId (not the underlying type: the actual TReplicatedBuffer<TUnderlyingType>
-	template<ENetworkSimBufferTypeId TypeId>
-	struct select_buffer_type
+	const TSimulationFrameState<Model>* GetValidFrame(int32 Frame) const
 	{
-		// We are just wrapping a TNetworkSimContiguousBuffer around whatever type TBufferTypes::select_type returns (which returns the underlying type)
-		using type = TNetworkSimContiguousBuffer<typename TBufferTypes::template select_type<TypeId>::type>;
-	};
+		return IsValidFrame(Frame) ? &FrameBuffer[Frame] : nullptr;
+	}
 
-	template<ENetworkSimBufferTypeId TypeId>
-	struct select_buffer_type_sparse
+	TSimulationFrameState<Model>* GetValidFrame(int32 Frame)
 	{
-		// We are just wrapping a TNetworkSimSparseBuffer around whatever type TBufferTypes::select_type returns (which returns the underlying type)
-		using type = TNetworkSimSparseBuffer<typename TBufferTypes::template select_type<TypeId>::type>;
-	};
-
-	// The buffer types. This may look intimidating but is not that bad!
-	// We are using select_buffer_type to decide what the type should be for a given ENetworkSimBufferTypeId.
-	using TInputBuffer = typename select_buffer_type<ENetworkSimBufferTypeId::Input>::type;
-	using TSyncBuffer = typename select_buffer_type<ENetworkSimBufferTypeId::Sync>::type;
-	using TAuxBuffer = typename select_buffer_type_sparse<ENetworkSimBufferTypeId::Aux>::type;
-	using TDebugBuffer = typename select_buffer_type<ENetworkSimBufferTypeId::Debug>::type;
-
-	// The buffers themselves. Types are already declared above.
-	// If you are reading just to find the damn underlying type here, its (TNetworkSimSparseBuffer||TNetworkSimContiguousBuffer) < whatever your struct type is >
-	TInputBuffer Input;
-	TSyncBuffer Sync;
-	TAuxBuffer Aux;
-	TDebugBuffer Debug;
+		return IsValidFrame(Frame) ? &FrameBuffer[Frame] : nullptr;
+	}
 
-	TNetSimCueDispatcher<Model> CueDispatcher;
+	bool bTickInProgress = false;
 
-	// Finally, template accessor for getting buffers based on enum. This is really what all this junk is about.
-	// This allows other templated classes in the system to access a specific buffer from another templated argument
-	template<ENetworkSimBufferTypeId BufferId>
-	typename select_buffer_type<BufferId>::type& Get()
-	{
-		return TBufferGetterHelper<TNetworkSimBufferContainer<Model>, BufferId>::Get(*this);
-	}
+	int32 PendingTickFrame = 0;
+	int32 NextInputFrame = 0;
+
+	int32 MaxTickFrame = INDEX_NONE;
+	int32 LatestInputFrame = INDEX_NONE;
+	int32 ConfirmedFrame = INDEX_NONE;
+	
+	FNetworkSimTime TotalAllowedSimulationTime;
+	TRealTimeAccumulator<TickSettings>	RealtimeAccumulator;
+	
+	TNetworkSimFrameBuffer<TSimulationFrameState<Model>> FrameBuffer;
+	TNetworkSimAuxBuffer<TAuxState> AuxBuffer;
 };
 
 // ----------------------------------------------------------------------------------------------------------------------------------------------
@@ -152,9 +338,46 @@ struct TNetworkSimBufferContainer
 //	Essentially, this allows you to plop a TNetSimStateAccessor<MyStateType> on a class and have it bind to any TNetworkedSimulationModel instantiation  
 //	that has uses that type.
 //
-//	See additional comments in TNetworkedSimulationModel::GetPendingStateRead() / GetPendingStateWrite()
+//	Explanation: During the scope of an ::SimulationTick call, we know exactly 'when' we are relative to what the server is processing. If the predicting client wants
+//	to predict a change to sync/aux state during an update, the server will do it at the exact same time (assuming not a mis prediction). When a state change
+//	happens "out of band" (outside an ::SimulationTick call) - we really have no way to correlate when the server will do it. While its tempting to think "we will get
+//	a correction anyways, might as well guess at it and maybe get a smaller correction" - but this opens us up to other problems. The server may actually not 
+//	change the state at all and you may not get an update that corrects you. You could add a timeout and track the state change somewhere but that really complicates
+//	things and could leave you open to "double" problems: if the state change is additive, you may stack the authority change on top of the local predicted change, or
+//	you may roll back the predicted change to then later receive the authority change.
 // ----------------------------------------------------------------------------------------------------------------------------------------------
 
+template<typename TNetworkSimModel, typename TState, bool IsSyncState=TIsDerivedFrom<typename TNetworkSimModel::TSyncState, TState>::Value>
+struct TNetSimAccessorHelper
+{
+	using TAccessFunc = TUniqueFunction<void(bool bForWrite, TState*& OutState, bool& OutSafe)>;
+	
+	static TAccessFunc GetStateFunc(TNetworkSimModel* NetSimModel)
+	{
+		TAccessFunc Func = [NetSimModel](bool bForWrite, TState*& OutState, bool& OutSafe)
+		{
+			OutState = &NetSimModel->State.GetFrameState(NetSimModel->State.GetPendingTickFrame())->SyncState;
+			OutSafe = NetSimModel->State.GetTickInProgress();
+		};
+		return Func;
+	}
+};
+
+template<typename TNetworkSimModel, typename TState>
+struct TNetSimAccessorHelper<TNetworkSimModel, TState, false>
+{
+	using TAccessFunc = TUniqueFunction<void(bool bForWrite, TState*& OutState, bool& OutSafe)>;
+	static TAccessFunc GetStateFunc(TNetworkSimModel* NetSimModel)
+	{
+		TAccessFunc Func = [NetSimModel](bool bForWrite, TState*& OutState, bool& OutSafe)
+		{
+			OutState = NetSimModel->State.GetAuxState(NetSimModel->State.GetPendingTickFrame());
+			OutSafe = NetSimModel->State.GetTickInProgress();
+		};
+		return Func;
+	}
+};
+
 
 template<typename TState>
 struct TNetSimStateAccessor
@@ -164,12 +387,7 @@ struct TNetSimStateAccessor
 	template<typename TNetworkSimModel>
 	void Bind(TNetworkSimModel* NetSimModel)
 	{
-		GetStateFunc = [NetSimModel](bool bWrite, TState*& OutState, bool& OutSafe)
-		{
-			auto& Buffer = NetSimBufferSelect::Get<TNetworkSimBufferContainer<typename TNetworkSimModel::Model>, TState>(NetSimModel->Buffers);
-			OutState = bWrite ? Buffer.WriteFrameInitializedFromHead(NetSimModel->Ticker.PendingFrame) : Buffer[NetSimModel->Ticker.PendingFrame];
-			OutSafe = NetSimModel->Ticker.bUpdateInProgress;
-		};
+		GetStateFunc = TNetSimAccessorHelper<TNetworkSimModel, TState>::GetStateFunc(NetSimModel);
 	}
 
 	void Clear()
@@ -203,7 +421,7 @@ struct TNetSimStateAccessor
 
 private:
 
-	TFunction<void(bool bForWrite, TState*& OutState, bool& OutSafe)> GetStateFunc;
+	TUniqueFunction<void(bool bForWrite, TState*& OutState, bool& OutSafe)> GetStateFunc;
 };
 
 // ----------------------------------------------------------------------------------------------------------------------------------------------
@@ -248,56 +466,57 @@ struct TSimulationDoTickFunc
 {
 	using TSimulation = typename Model::Simulation;
 	using TTickSettings = typename Model::TickSettings;
-	using TDriver = typename TNetSimModelTraits<Model>::Driver;
-	using TBufferTypes = typename TNetSimModelTraits<Model>::InternalBufferTypes;
+	using TBufferTypes = typename Model::BufferTypes;
+
+	using TDriver = TNetworkedSimulationModelDriver<TBufferTypes>;
+	
 	using TInputCmd = typename TBufferTypes::TInputCmd;
 	using TSyncState = typename TBufferTypes::TSyncState;
 	using TAuxState = typename TBufferTypes::TAuxState;
+	
+	using TFrameState = TSimulationFrameState<Model>;
 
-	static void DoTick(TSimulation& Simulation, TSimulationTicker<TTickSettings>& Ticker, TNetworkSimBufferContainer<Model>& Buffers, int32 OutputFrame, ESimulationTickContext Context)
+	static void DoTick(TSimulation& Simulation, TNetworkedSimulationState<Model>& State, ESimulationTickContext Context)
 	{
-		check(Ticker.bUpdateInProgress == false);
-		check(OutputFrame == Ticker.PendingFrame+1);
+		check(State.GetTickInProgress() == false);
+		const int32 InputFrame = State.GetPendingTickFrame();
+		const int32 OutputFrame = State.GetPendingTickFrame()+1;
 
 		// Get inputs. These all need to be present to tick (caller's responsibility)
-		const TInputCmd* InputCmd = Buffers.Input[Ticker.PendingFrame];
-		const TSyncState* InSyncState = Buffers.Sync[Ticker.PendingFrame];
-		const TAuxState* InAuxState = Buffers.Aux[Ticker.PendingFrame];
-
-		check(InputCmd);
-		check(InSyncState);
+		TFrameState* InFrameState = State.GetFrameState(InputFrame);
+		const TAuxState* InAuxState = State.GetAuxState(InputFrame);
+		
+		check(InFrameState);
 		check(InAuxState);
+		ensure(InFrameState->HasFlag(ESimulationFrameStateFlags::InputWritten));
 
-		const FNetworkSimTime& DeltaSimTime = InputCmd->GetFrameDeltaTime();
+		const FNetworkSimTime DeltaSimTime = InFrameState->FrameDeltaTime.Get();
+		const FNetSimTimeStep SimTimeStep { DeltaSimTime, InFrameState->TotalSimTime, OutputFrame };
+		const FNetworkSimTime OutSimTime = InFrameState->TotalSimTime + DeltaSimTime;
 
-		// Output (only sync will always be generated
-		TSyncState* OutSyncState = Buffers.Sync.WriteFrame(OutputFrame);
+		TFrameState* OutFrameState = State.WriteFrameWithTime(OutputFrame, OutSimTime);
 
 		// PreTick
-		Ticker.PendingFrame = OutputFrame;
-		Ticker.bUpdateInProgress = true;
+		State.SetPendingTickFrame(OutputFrame);
+		State.SetTickInProgress(true);
+		State.CueDispatcher.PushContext({OutFrameState->TotalSimTime, Context}); // Cues "take place" at the end of the frame
 
-		Buffers.CueDispatcher.PushContext({Ticker.GetTotalProcessedSimulationTime() + DeltaSimTime, Context}); // Cues "take place" at the end of the frame
-
-		UE_NP_TRACE_SIM_TICK(OutputFrame, DeltaSimTime, Ticker);
+		UE_NP_TRACE_SIM_TICK(OutputFrame, DeltaSimTime, SimTimeStep);
 
 		// Simulation Tick
-		Simulation.SimulationTick(
-			{ DeltaSimTime, Ticker }, // How it ticks
-			{ *InputCmd, *InSyncState, *InAuxState }, // Input States
-			{ *OutSyncState, Buffers.Aux.LazyWriter(OutputFrame), Buffers.CueDispatcher } ); // Output
+		Simulation.SimulationTick( SimTimeStep, // FNetSimTimeStep. How it ticks
+			{ InFrameState->InputCmd, InFrameState->SyncState, *InAuxState }, // TNetSimInput
+			{ OutFrameState->SyncState, State.GetAuxStateLazyWriter(OutputFrame), State.CueDispatcher } ); // TNetSimOutput
 
 		// Post
-		UE_NP_TRACE_USER_STATE_SYNC(*Buffers.Sync.HeadElement(), OutputFrame);
-		const TAuxState* AuxHead = Buffers.Aux.HeadElement();
+		UE_NP_TRACE_USER_STATE_SYNC(OutFrameState->SyncState, OutputFrame);
+		const TAuxState* AuxHead = State.GetAuxState(OutputFrame);
 		if (AuxHead != InAuxState)
 		{
 			UE_NP_TRACE_USER_STATE_AUX(*AuxHead, OutputFrame);
 		}
-
-		Ticker.IncrementTotalProcessedSimulationTime(DeltaSimTime, OutputFrame);
-		Ticker.bUpdateInProgress = false;
-
-		Buffers.CueDispatcher.PopContext();
+				
+		State.SetTickInProgress(false);
+		State.CueDispatcher.PopContext();
 	}
 };
\ No newline at end of file
diff --git a/Engine/Plugins/Runtime/NetworkPrediction/Source/NetworkPrediction/Public/Trace/NetworkPredictionTrace.h b/Engine/Plugins/Runtime/NetworkPrediction/Source/NetworkPrediction/Public/Trace/NetworkPredictionTrace.h
--- a/Engine/Plugins/Runtime/NetworkPrediction/Source/NetworkPrediction/Public/Trace/NetworkPredictionTrace.h
+++ b/Engine/Plugins/Runtime/NetworkPrediction/Source/NetworkPrediction/Public/Trace/NetworkPredictionTrace.h
@@ -21,45 +21,43 @@
 #if UE_NP_TRACE_ENABLED
 
 // Scope Push/Pop for active simulation
-#define UE_NP_TRACE_SET_SCOPE_SIM(SimulationId) FScopedSimulationTraceHelper _ScopedSimIdTraceHelper(SimulationId);
+#define UE_NP_TRACE_SET_SCOPE_SIM(SimulationId) FScopedSimulationTraceHelper _ScopedSimIdTraceHelper(SimulationId)
 
 // Called when simulation is created. (Note this also sets a Scope for tracing the initial user states next)
-#define UE_NP_TRACE_SIM_CREATED(OwningActor, GroupName) FNetworkPredictionTrace::TraceSimulationCreated(OwningActor, GroupName);
+#define UE_NP_TRACE_SIM_CREATED(OwningActor, GroupName) FNetworkPredictionTrace::TraceSimulationCreated(OwningActor, GroupName)
 
-#define UE_NP_TRACE_NETROLE(SimulationId, NetRole) FNetworkPredictionTrace::TraceSimulationNetRole(SimulationId, NetRole);
+#define UE_NP_TRACE_NETROLE(SimulationId, NetRole) FNetworkPredictionTrace::TraceSimulationNetRole(SimulationId, NetRole)
 
 // Called when simulation ticks in any context (includes resimulates, sim extrapolates, etc)
-#define UE_NP_TRACE_SIM_TICK(OutputFrame, FrameDeltaTime, TickState) FNetworkPredictionTrace::TraceSimulationTick(OutputFrame, FrameDeltaTime, TickState);
+#define UE_NP_TRACE_SIM_TICK(OutputFrame, FrameDeltaTime, TimeStep) FNetworkPredictionTrace::TraceSimulationTick(OutputFrame, FrameDeltaTime, TimeStep)
 
 // Called at the end of an engine frame (EOF) to update the current tick state of a simulation
-#define UE_NP_TRACE_SIM_EOF(TickState, LastSentKeyframe, LastReceivedKeyframe) FNetworkPredictionTrace::TraceSimulationEOF(TickState, LastSentKeyframe, LastReceivedKeyframe);
+#define UE_NP_TRACE_SIM_EOF(Buffers) FNetworkPredictionTrace::TraceSimulationEOF(Buffers)
 
 // Called whenever we receive network data: note this only indicates that data was received, not necessarily that it was committed to the simulation buffers
-#define UE_NP_TRACE_NET_SERIALIZE_RECV(ReceivedTime, ReceivedFrame) FNetworkPredictionTrace::TraceNetSerializeRecv(ReceivedTime, ReceivedFrame);
+#define UE_NP_TRACE_NET_SERIALIZE_RECV(ReceivedTime, ReceivedFrame) FNetworkPredictionTrace::TraceNetSerializeRecv(ReceivedTime, ReceivedFrame)
 
 // Called whenever a new user state has been inserted into the buffers. Analysis will determine "how" it got there from previous trace events (SIM_TICK, NET_SERIALIZE_RECV)
-#define UE_NP_TRACE_USER_STATE_INPUT(UserState, Frame) FNetworkPredictionTrace::TraceUserState(UserState, Frame, FNetworkPredictionTrace::ETraceUserState::Input);
-#define UE_NP_TRACE_USER_STATE_SYNC(UserState, Frame) FNetworkPredictionTrace::TraceUserState(UserState, Frame, FNetworkPredictionTrace::ETraceUserState::Sync);
-#define UE_NP_TRACE_USER_STATE_AUX(UserState, Frame) FNetworkPredictionTrace::TraceUserState(UserState, Frame, FNetworkPredictionTrace::ETraceUserState::Aux);
+#define UE_NP_TRACE_USER_STATE_INPUT(UserState, Frame) FNetworkPredictionTrace::TraceUserState(UserState, Frame, FNetworkPredictionTrace::ETraceUserState::Input)
+#define UE_NP_TRACE_USER_STATE_SYNC(UserState, Frame) FNetworkPredictionTrace::TraceUserState(UserState, Frame, FNetworkPredictionTrace::ETraceUserState::Sync)
+#define UE_NP_TRACE_USER_STATE_AUX(UserState, Frame) FNetworkPredictionTrace::TraceUserState(UserState, Frame, FNetworkPredictionTrace::ETraceUserState::Aux)
 
 // Called to indicate that the previous received NetSerialized data was committed to the buffers. (Not all data is commited, but we want to trace all received data)
-#define UE_NP_TRACE_NET_SERIALIZE_COMMIT() FNetworkPredictionTrace::TraceNetSerializeCommit();
-
-// Called to indicate the previous received NetSerialized data caused a fault. This means the data was too far in the past and could not be used. (Does not mean "caused correction").
-#define UE_NP_TRACE_NET_SERIALIZE_FAULT() FNetworkPredictionTrace::TraceNetSerializeFault();
+#define UE_NP_TRACE_NET_SERIALIZE_COMMIT() FNetworkPredictionTrace::TraceNetSerializeCommit()
 
 // Called prior to sampling local input
-#define UE_NP_TRACE_PRODUCE_INPUT() FNetworkPredictionTrace::TraceProduceInput();
+#define UE_NP_TRACE_PRODUCE_INPUT() FNetworkPredictionTrace::TraceProduceInput()
 
 // Called prior to generating synthesized input
-#define UE_NP_TRACE_SYNTH_INPUT() FNetworkPredictionTrace::TraceSynthInput();
+#define UE_NP_TRACE_SYNTH_INPUT() FNetworkPredictionTrace::TraceSynthInput()
 
 // Called to indicate we are about to write state to the buffers outside of the normal simulation tick/netrecive. TODO: add char* identifier to debug where the mod came from
-#define UE_NP_TRACE_OOB_STATE_MOD() FNetworkPredictionTrace::TraceOOBStateMod();
+#define UE_NP_TRACE_OOB_STATE_MOD() FNetworkPredictionTrace::TraceOOBStateMod()
 
 // Called when a PIE session is started. This is so we can keep our *sets* of worlds/simulations seperate.
-#define UE_NP_TRACE_PIE_START() FNetworkPredictionTrace::TracePIEStart();
+#define UE_NP_TRACE_PIE_START() FNetworkPredictionTrace::TracePIEStart()
 
+#define UE_NP_TRACE_SYSTEM_FAULT(Format, ...) FNetworkPredictionTrace::TraceSystemFault(TEXT(Format), ##__VA_ARGS__)
 
 #else
 
@@ -75,11 +73,11 @@
 #define UE_NP_TRACE_USER_STATE_AUX(...)
 #define UE_NP_TRACE_ROLLBACK(...)
 #define UE_NP_TRACE_NET_SERIALIZE_COMMIT(...)
-#define UE_NP_TRACE_NET_SERIALIZE_FAULT(...)
 #define UE_NP_TRACE_PRODUCE_INPUT(...)
 #define UE_NP_TRACE_SYNTH_INPUT(...)
 #define UE_NP_TRACE_OOB_STATE_MOD(...)
 #define UE_NP_TRACE_PIE_START(...)
+#define UE_NP_TRACE_SYSTEM_FAULT(Format, ...) UE_LOG(LogNetworkSim, Warning, TEXT(Format), ##__VA_ARGS__);
 
 #endif // UE_NP_TRACE_ENABLED
 
@@ -87,6 +85,9 @@ NETWORKPREDICTION_API UE_TRACE_CHANNEL_EXTERN(NetworkPredictionChannel);
 
 class AActor;
 
+template<typename Model>
+struct TNetworkedSimulationState;
+
 // Trace is what the system+user code calls to do traces (could still be wrapped around [my own] macros, etc. This will call into the generic UE_TRACE_LOG macros)
 class NETWORKPREDICTION_API FNetworkPredictionTrace
 {
@@ -99,15 +100,14 @@ public:
 
 	static void TraceSimulationCreated(const AActor* OwningActor, const FName& GroupName);
 	static void TraceSimulationNetRole(uint32 SimulationId, ENetRole NetRole);
-	static void TraceSimulationTick(int32 OutputFrame, const FNetworkSimTime& FrameDeltaTime, const FSimulationTickState& TickState);
-	static void TraceSimulationEOF(const FSimulationTickState& TickState, int32 LastSentKeyframe, int32 LastReceivedKeyframe);
+	static void TraceSimulationTick(int32 OutputFrame, const FNetworkSimTime& FrameDeltaTime, const FNetSimTimeStep& TimeStep);
 	static void TraceNetSerializeRecv(const FNetworkSimTime& ReceivedTime, int32 ReceivedFrame);
 	static void TraceNetSerializeCommit();
-	static void TraceNetSerializeFault();
 	static void TraceProduceInput();
 	static void TraceSynthInput();
 	static void TraceOOBStateMod();
 	static void TracePIEStart();
+	static void TraceSystemFault(const TCHAR* Fmt, ...);
 
 	enum ETraceUserState
 	{
@@ -133,12 +133,19 @@ public:
 #endif
 	}
 
+	template<typename Model>
+	static void TraceSimulationEOF(const TNetworkedSimulationState<Model>& State)
+	{
+		TraceEOF_Internal(State.Max(), State.GetPendingTickFrame(), State.GetLatestInputFrame(), State.GetMaxTickFrame(), State.GetTotalProcessedSimulationTime(), State.GetTotalAllowedSimulationTime());
+	}
+
 private:
 
 	static void PushSimulationId(uint32 SimulationId);
 	static void PopSimulationId();
 
 	static void TraceUserState_Internal(int32 Frame, ETraceUserState StateType);
+	static void TraceEOF_Internal(int32 BufferSize, int32 PendingTickFrame, int32 LatestInputFrame, int32 MaxTickFrame, FNetworkSimTime TotalSimTime, FNetworkSimTime AllowedSimTime);
 
 	static FStringOutputDevice StrOut;
 
diff --git a/Engine/Plugins/Runtime/NetworkPrediction/readme.txt b/Engine/Plugins/Runtime/NetworkPrediction/readme.txt
--- a/Engine/Plugins/Runtime/NetworkPrediction/readme.txt
+++ b/Engine/Plugins/Runtime/NetworkPrediction/readme.txt
@@ -42,10 +42,7 @@ Code to look at:
  
 Extremely high level technical overview:
 -Concept of "NetworkedSimulationModel": defining a tightly constrained (gameplay) simulation that can be run predictively on clients and authoritatively on server. E.g., a "Movement System" is a NetworkedSimulationModel.
--The simulation has tightly defined input and output. At the lowest level we have generic buffers of structs that are synchronized:
-	Input Buffer: the data that the controlling client generates
-	Sync Buffer: the data that evolves frame to frame and is produced by ticking the simulation
-	Aux Buffer: additional data that can change, can be predicted, etc but does not necessarily get updated frame to frame (e.g, something else usually updates it).
+-The simulation has tightly defined input and output.
 -Implicit correction: all prediction and correction logic is now client side. Server just tells clients what the authority state was and it is up to them to reconcile.
 -Single code path: the core network simulation has one "Update" function that does not branch based on authority/prediction/simulated etc. "Write once" is the goal for gameplay code here.
  
@@ -65,8 +62,41 @@ Simulated Proxy behavior / "modes"
 TODO: Major missing elements
 -No optimizations (bandwidth, cpu) have been done.
 -[Movement specific]: must lock down SetActorLocation/Rotation API so that movement state can be changed out from underneath the system.
+
+// ----------------------------------------------------------------------------------------------------------
+// Architecture
+// ----------------------------------------------------------------------------------------------------------
+
+User Structures
+-User defines these strutures along with NetSerialize/Log function.
+	-InputCmd: client generated input state. (E.g, FMockInputCmd)
+	-Sync: frame-to-frame evolving state. (E.g, FMockSyncState)
+	-Aux: state that does not always change frame to frame (E.g, FMockAuxState)
+
+User Simulation Class (E.g, FMockNetworkSimulation)
+-Networking agnostic, user defined simulation class. Has ::SimulationTick function. Can hold non networked state.
+-Takes Inputcmd/Sync/Aux state as input and produces an output Sync and (optional) Aux state.
+
+TNetworkedSimulationModel
+-Core class that ties things together. Mostly empty itself.
+
+	Model Definition (FNetSimModelDefBase. E.g, FMockNetworkModelDef)
+		-Sets user defined structures and simulation class.
+		-Implements networking, user defined functions: ShouldReconcile, Interpolate.
+
+	TNetworkedSimulationState
+		-Holds all internal networked state: user frame buffers and ticking/frame info.
+
+	RepControllers (Server, Autonomous Proxy, Simulated Proxy)
+		-Network role based classes that define NetSeriaize, Reconcile, and Tick behavior.
+		-These hold the logic to actually advance and network the simulation machinery. 
+
  
- 
+Driver (TNetworkedSimulationModelDriver. E.g, UMockNetworkSimulationComponent)
+-UE4 side that binds with the TNetworkedSimulationModel
+-Resonsible for ::FinalizeFrame, ::ProduceInput functions. Also debugging/logging helpers.
+
+
 // ----------------------------------------------------------------------------------------------------------
 // Glossary
 // ----------------------------------------------------------------------------------------------------------
@@ -93,6 +123,14 @@ Smoothing: Taking the output of the simulation and applying an additional layer
 // Release notes
 // ----------------------------------------------------------------------------------------------------------
 
+
+Update (4-1-20)
+-Spring cleaning refactor. User code should be unaffected but this has hopefully simplified a lot of things.
+-Moving away from individual Input/Sync/Time buffers into a single internal "FrameState" structure.
+	-This is more cache effecient and less error prone. Previously was difficult to reason about which buffer's "head" was correct and how to keep them in sync.
+	-"Gaps" in frames are now allowed, rather than wiping previous contents. Gap frames are skipped.
+	-This was needed as we experimented with input buffering and packet loss.
+
 Update (3-16-20)
 -Network Prediction Insights is the new Unreal Insights based tool debugging tool for Network rediction.
 -Enable the NetworkPredictionInsights plugin in your project. 
@@ -102,6 +140,10 @@ Update (3-16-20)
 -Existing canvas based debugging will be removed soon from Network Prediction.
 -Added TSimulationDoTickFunc<Model> to replace TScopedSimulationTick. This better encapsulates 'running a tick'. 
 
+Videos of Network Prediction Insights:
+https://youtu.be/_rdt-v1nFlY
+https://youtu.be/18BNvoGP-E0
+
 Update (2-18-20)
 -Sorry for the lack of public updates. Development has been ongoing but mostly moved to a non public stream at the moment. The work being done revolves around:
 	-Unreal Insights. We are implementing an Insights tool for debugging Network Prediction. This will replace the canvas/visual logger debugging in the plugin. It is very powerful. Hopefully v.1 will make it into main branch soon.
diff --git a/Engine/Plugins/Runtime/NetworkPredictionInsights/Source/NetworkPredictionInsights/Private/NetworkPredictionAnalyzer.cpp b/Engine/Plugins/Runtime/NetworkPredictionInsights/Source/NetworkPredictionInsights/Private/NetworkPredictionAnalyzer.cpp
--- a/Engine/Plugins/Runtime/NetworkPredictionInsights/Source/NetworkPredictionInsights/Private/NetworkPredictionAnalyzer.cpp
+++ b/Engine/Plugins/Runtime/NetworkPredictionInsights/Source/NetworkPredictionInsights/Private/NetworkPredictionAnalyzer.cpp
@@ -25,11 +25,11 @@ void FNetworkPredictionAnalyzer::OnAnalysisBegin(const FOnAnalysisContext& Conte
 	Builder.RouteEvent(RouteId_SimulationEOF, "NetworkPrediction", "SimulationEOF");
 	Builder.RouteEvent(RouteId_NetSerializeRecv, "NetworkPrediction", "NetSerializeRecv");
 	Builder.RouteEvent(RouteId_NetSerializeCommit, "NetworkPrediction", "NetSerializeCommit");
-	Builder.RouteEvent(RouteId_NetSerializeFault, "NetworkPrediction", "NetSerializeFault");
 	Builder.RouteEvent(RouteId_InputCmd, "NetworkPrediction", "InputCmd");
 	Builder.RouteEvent(RouteId_SyncState, "NetworkPrediction", "SyncState");
 	Builder.RouteEvent(RouteId_AuxState, "NetworkPrediction", "AuxState");
 	Builder.RouteEvent(RouteId_PieBegin, "NetworkPrediction", "PieBegin");
+	Builder.RouteEvent(RouteId_SystemFault, "NetworkPrediction", "SystemFault");
 }
 
 void FNetworkPredictionAnalyzer::OnAnalysisEnd()
@@ -101,14 +101,15 @@ bool FNetworkPredictionAnalyzer::OnEvent(uint16 RouteId, const FOnEventContext&
 		{
 			FSimulationData::FEngineFrame& FrameData = NetworkPredictionProvider.WriteSimulationEOF(EventData.GetValue<uint32>("SimulationId"));
 			FrameData.EngineFrame = EventData.GetValue<uint64>("EngineFrameNumber");
-			FrameData.EngineFrameDeltaTime = EventData.GetValue<double>("EngineFrameDeltaTime");
-			FrameData.EngineCurrentTime = EventData.GetValue<double>("EngineCurrentTime");
+			FrameData.EngineFrameDeltaTime = EventData.GetValue<float>("EngineFrameDeltaTime");
+			
 
-			FrameData.TotalProcessedSimTimeMS = EventData.GetValue<int32>("TotalProcessedMS");
-			FrameData.TotalAllowedSimTimeMS = EventData.GetValue<int32>("TotalAllowedMS");
-		
-			FrameData.LastSentKeyframe = EventData.GetValue<int32>("LastSentKeyframe");
-			FrameData.LastReceivedKeyframe = EventData.GetValue<int32>("LastReceivedKeyframe");
+			FrameData.BufferSize = EventData.GetValue<int32>("BufferSize");
+			FrameData.PendingTickFrame = EventData.GetValue<int32>("PendingTickFrame");
+			FrameData.LatestInputFrame = EventData.GetValue<int32>("LatestInputFrame");
+			FrameData.MaxTickFrame = EventData.GetValue<int32>("MaxTickFrame");
+			FrameData.TotalSimTime = EventData.GetValue<int32>("TotalSimTime");
+			FrameData.AllowedSimTime = EventData.GetValue<int32>("AllowedSimTime");	
 			break;
 		}
 
@@ -116,8 +117,8 @@ bool FNetworkPredictionAnalyzer::OnEvent(uint16 RouteId, const FOnEventContext&
 		{
 			FSimulationData::FNetSerializeRecv NetRecv;
 			NetRecv.EngineFrame = EventData.GetValue<uint64>("EngineFrameNumber");
-			NetRecv.SimTimeMS = EventData.GetValue<uint64>("ReceivedSimTimeMS");
-			NetRecv.Frame = EventData.GetValue<uint64>("ReceivedFrame");
+			NetRecv.SimTimeMS = EventData.GetValue<int32>("ReceivedSimTimeMS");
+			NetRecv.Frame = EventData.GetValue<int32>("ReceivedFrame");
 			
 			NetworkPredictionProvider.WriteNetRecv(EventData.GetValue<uint32>("SimulationId"), MoveTemp(NetRecv));
 			break;
@@ -129,12 +130,6 @@ bool FNetworkPredictionAnalyzer::OnEvent(uint16 RouteId, const FOnEventContext&
 			break;
 		}
 
-		case RouteId_NetSerializeFault:
-		{
-			NetworkPredictionProvider.WriteNetFault(EventData.GetValue<uint32>("SimulationId"));
-			break;
-		}
-
 		case RouteId_OOBStateMod:
 		{
 			NetworkPredictionProvider.WriteOOBStateMod(EventData.GetValue<uint32>("SimulationId"));
@@ -177,6 +172,16 @@ bool FNetworkPredictionAnalyzer::OnEvent(uint16 RouteId, const FOnEventContext&
 			break;
 		}
 
+		case RouteId_SystemFault:
+		{
+			NetworkPredictionProvider.WriteSystemFault(
+				EventData.GetValue<uint32>("SimulationId"),
+				EventData.GetValue<uint64>("EngineFrameNumber"),
+				Session.StoreString(reinterpret_cast<const TCHAR*>(EventData.GetAttachment()))
+			);
+			break;
+		}
+
 		default:
 			bUnhandled = true;
 			break;
diff --git a/Engine/Plugins/Runtime/NetworkPredictionInsights/Source/NetworkPredictionInsights/Private/NetworkPredictionAnalyzer.h b/Engine/Plugins/Runtime/NetworkPredictionInsights/Source/NetworkPredictionInsights/Private/NetworkPredictionAnalyzer.h
--- a/Engine/Plugins/Runtime/NetworkPredictionInsights/Source/NetworkPredictionInsights/Private/NetworkPredictionAnalyzer.h
+++ b/Engine/Plugins/Runtime/NetworkPredictionInsights/Source/NetworkPredictionInsights/Private/NetworkPredictionAnalyzer.h
@@ -40,11 +40,11 @@ private:
 		RouteId_SimulationEOF,
 		RouteId_NetSerializeRecv,
 		RouteId_NetSerializeCommit,
-		RouteId_NetSerializeFault,
 		RouteId_InputCmd,
 		RouteId_SyncState,
 		RouteId_AuxState,
-		RouteId_PieBegin
+		RouteId_PieBegin,
+		RouteId_SystemFault
 	};
 
 
diff --git a/Engine/Plugins/Runtime/NetworkPredictionInsights/Source/NetworkPredictionInsights/Private/NetworkPredictionProvider.cpp b/Engine/Plugins/Runtime/NetworkPredictionInsights/Source/NetworkPredictionInsights/Private/NetworkPredictionProvider.cpp
--- a/Engine/Plugins/Runtime/NetworkPredictionInsights/Source/NetworkPredictionInsights/Private/NetworkPredictionProvider.cpp
+++ b/Engine/Plugins/Runtime/NetworkPredictionInsights/Source/NetworkPredictionInsights/Private/NetworkPredictionProvider.cpp
@@ -129,7 +129,10 @@ void FNetworkPredictionProvider::WriteSimulationTick(uint32 SimulationId, FSimul
 
 FSimulationData::FEngineFrame& FNetworkPredictionProvider::WriteSimulationEOF(uint32 SimulationId)
 {
-	return FindOrAdd(SimulationId)->EOFState.PushBack();
+	FSimulationData& SimulationData = FindOrAdd(SimulationId).Get();
+	FSimulationData::FEngineFrame& NewEOF = SimulationData.EOFState.PushBack();
+	NewEOF.SystemFaults = MoveTemp(SimulationData.Analysis.PendingSystemFaults);
+	return NewEOF;
 }
 
 void FNetworkPredictionProvider::WriteNetRecv(uint32 SimulationId, FSimulationData::FNetSerializeRecv&& InNetRecv)
@@ -235,15 +238,20 @@ void FNetworkPredictionProvider::WriteNetCommit(uint32 SimulationId)
 	}
 }
 
-void FNetworkPredictionProvider::WriteNetFault(uint32 SimulationId)
+void FNetworkPredictionProvider::WriteSystemFault(uint32 SimulationId, uint64 EngineFrameNumber, const TCHAR* Fmt)
 {
 	FSimulationData& SimulationData = FindOrAdd(SimulationId).Get();
 	
+	// This is akwward to trace. Should refactor some of this so it can be easily known if this happened within a
+	// NetRecv or SimTick scope.
+	
 	for (FSimulationData::FNetSerializeRecv* PendingRecv : SimulationData.Analysis.PendingNetSerializeRecv)
 	{
-		ensure(PendingRecv->Status == ENetSerializeRecvStatus::Unknown);
 		PendingRecv->Status = ENetSerializeRecvStatus::Fault;
+		PendingRecv->SystemFaults.Add({Fmt});
 	}
+
+	SimulationData.Analysis.PendingSystemFaults.Add({Fmt});
 }
 
 void FNetworkPredictionProvider::WriteOOBStateMod(uint32 SimulationId)
diff --git a/Engine/Plugins/Runtime/NetworkPredictionInsights/Source/NetworkPredictionInsights/Private/NetworkPredictionProvider.h b/Engine/Plugins/Runtime/NetworkPredictionInsights/Source/NetworkPredictionInsights/Private/NetworkPredictionProvider.h
--- a/Engine/Plugins/Runtime/NetworkPredictionInsights/Source/NetworkPredictionInsights/Private/NetworkPredictionProvider.h
+++ b/Engine/Plugins/Runtime/NetworkPredictionInsights/Source/NetworkPredictionInsights/Private/NetworkPredictionProvider.h
@@ -32,7 +32,7 @@ public:
 	FSimulationData::FEngineFrame& WriteSimulationEOF(uint32 SimulationId);
 	void WriteNetRecv(uint32 SimulationId, FSimulationData::FNetSerializeRecv&& InNetRecv);
 	void WriteNetCommit(uint32 SimulationId);
-	void WriteNetFault(uint32 SimulationId);
+	void WriteSystemFault(uint32 SimulationId, uint64 EngineFrameNumber, const TCHAR* Fmt);
 	void WriteOOBStateMod(uint32 SimulationId);
 	void WriteProduceInput(uint32 SimulationId);
 	void WriteSynthInput(uint32 SimulationId);
diff --git a/Engine/Plugins/Runtime/NetworkPredictionInsights/Source/NetworkPredictionInsights/Private/UI/SNPSimFrameContents.cpp b/Engine/Plugins/Runtime/NetworkPredictionInsights/Source/NetworkPredictionInsights/Private/UI/SNPSimFrameContents.cpp
--- a/Engine/Plugins/Runtime/NetworkPredictionInsights/Source/NetworkPredictionInsights/Private/UI/SNPSimFrameContents.cpp
+++ b/Engine/Plugins/Runtime/NetworkPredictionInsights/Source/NetworkPredictionInsights/Private/UI/SNPSimFrameContents.cpp
@@ -35,59 +35,22 @@ void SNPSimFrameContents::Construct(const FArguments& InArgs, TSharedPtr<SNPWind
 		.VAlign(VAlign_Top)
 		.Padding(FMargin(0, 0, 0, 0))
 		[
-			SAssignNew(ContentsHBoxPtr, SHorizontalBox)
-		]
-	];
-
-	/*
-
-			+ SHorizontalBox::Slot()
-			.AutoWidth()
-			.VAlign(VAlign_Top)
-			.HAlign(HAlign_Left)
-			.Padding(4.0f, 4.0f, 4.0f, 4.0f)
-			[
-				SAssignNew(SimInfoTextBlock, STextBlock)
-			]
-			
-			+ SHorizontalBox::Slot()
-			.AutoWidth()
-			.VAlign(VAlign_Top)
-			.HAlign(HAlign_Left)
-			.Padding(4.0f, 4.0f, 4.0f, 4.0f)
-			[
-				SAssignNew(SimTickTextBlock, STextBlock)
-			]
-
-			+ SHorizontalBox::Slot()
-			.AutoWidth()
-			.VAlign(VAlign_Top)
-			.HAlign(HAlign_Left)
-			.Padding(4.0f, 4.0f, 4.0f, 4.0f)
-			[
-				SAssignNew(InputCmdTextBlock, STextBlock)
-			]
-
-			+ SHorizontalBox::Slot()
-			.AutoWidth()
-			.VAlign(VAlign_Top)
-			.HAlign(HAlign_Left)
-			.Padding(4.0f, 4.0f, 4.0f, 4.0f)
+			SNew(SVerticalBox)
+			+SVerticalBox::Slot()
+			.Padding(2.0f)
+			.AutoHeight()
 			[
-				SAssignNew(SyncStateTextBlock, STextBlock)
+				SAssignNew(ContentsHBoxPtr, SHorizontalBox)
 			]
 
-			+ SHorizontalBox::Slot()
-			.AutoWidth()
-			.VAlign(VAlign_Top)
-			.HAlign(HAlign_Left)
-			.Padding(4.0f, 4.0f, 4.0f, 4.0f)
+			+SVerticalBox::Slot()
+			.Padding(2.0f)
+			.AutoHeight()
 			[
-				SAssignNew(AuxStateTextBlock, STextBlock)
+				SAssignNew(SystemFaultsVBoxPtr, SVerticalBox)
 			]
 		]
 	];
-	*/
 }
 
 void SNPSimFrameContents::NotifyContentClicked(const FSimContentsView& InContent)
@@ -181,6 +144,7 @@ void SNPSimFrameContents::NotifyContentClicked(const FSimContentsView& InContent
 		}
 	};
 
+	
 	ContentsHBoxPtr->ClearChildren();
 	
 	// -------------------------------------------------------------
@@ -222,6 +186,64 @@ void SNPSimFrameContents::NotifyContentClicked(const FSimContentsView& InContent
 		]
 	];
 
+	// -------------------------------------------------------------
+	//	End of Frame
+	// -------------------------------------------------------------
+	uint64 EngineFrame = (InContent.SimTick ? InContent.SimTick->EngineFrame : (InContent.NetRecv ? InContent.NetRecv->EngineFrame : 0));
+	const FSimulationData::FEngineFrame* EOFState = nullptr;
+	if (EngineFrame > 0)
+	{
+		for (auto It = InContent.SimView->EOFState.GetIteratorFromEnd(); It; --It)
+		{
+			if (It->EngineFrame == EngineFrame)
+			{
+				EOFState = &*It;
+				break;
+			}
+		}
+
+		if (EOFState)
+		{
+			FString EOFText = FString::Format(TEXT(
+			"EngineFrame {0}\n"
+			"EngineFrameDeltaTime: {1}\n"
+			"BufferSize: {2}\n"
+			"PendingTickFrame: {3}\n"
+			"LatestInputFrame: {4}\n"
+			"MaxTickFrame: {5}\n"
+			"TotalSimTime: {6}\n"
+			"AllowedSimTime: {7}"),
+			{
+				EOFState->EngineFrame,
+				EOFState->EngineFrameDeltaTime,
+				EOFState->BufferSize,
+				EOFState->PendingTickFrame,
+				EOFState->LatestInputFrame,
+				EOFState->MaxTickFrame,
+				EOFState->TotalSimTime,
+				EOFState->AllowedSimTime
+			});
+
+			/*
+			for (const FSimulationData::FSystemFault& Fault : EOFState->SystemFaults)
+			{
+				EOFText.Append(Fault.Str);
+				EOFText.Append(TEXT("\n"));
+			}
+			*/
+
+			VertBox->AddSlot()
+			.AutoHeight()
+			.VAlign(VAlign_Top)
+			.HAlign(HAlign_Left)
+			.Padding(4.0f, 4.0f, 4.0f, 4.0f)
+			[
+				SNew(STextBlock)
+				.Text(FText::FromString(EOFText))
+			];
+		}
+	}
+
 	// -------------------------------------------------------------
 	//	Simulation Tick
 	// -------------------------------------------------------------
@@ -260,7 +282,10 @@ void SNPSimFrameContents::NotifyContentClicked(const FSimContentsView& InContent
 		[
 			SNew(STextBlock)
 			.Text(FText::FromString(TickText))
-		];		
+		];
+
+
+		// Input/Output states
 		
 		const int32 InputFrame = SimTick.OutputFrame - 1;
 
@@ -305,7 +330,7 @@ void SNPSimFrameContents::NotifyContentClicked(const FSimContentsView& InContent
 			NetRecv.SimTimeMS,
 			LexToString(NetRecv.Status),
 			NetRecv.EngineFrame
-		});
+		});		
 
 		ContentsHBoxPtr->AddSlot()
 		.AutoWidth()
@@ -325,7 +350,6 @@ void SNPSimFrameContents::NotifyContentClicked(const FSimContentsView& InContent
 			]
 		];
 
-
 		uint8 Mask = ~((uint8)ENP_UserStateSource::NetRecv | (uint8)ENP_UserStateSource::NetRecvCommit);
 		const FSimulationData::FUserState* RecvInputCmd = InContent.SimView->UserData.Get(ENP_UserState::Input, InContent.NetRecv->Frame, InContent.NetRecv->EngineFrame, Mask);
 		const FSimulationData::FUserState* RecvSyncState = InContent.SimView->UserData.Get(ENP_UserState::Sync, InContent.NetRecv->Frame, InContent.NetRecv->EngineFrame, Mask);
@@ -343,6 +367,57 @@ void SNPSimFrameContents::NotifyContentClicked(const FSimContentsView& InContent
 		);
 	}
 
+
+	// -------------------------------------
+	// System Faults
+	// -------------------------------------
+
+	SystemFaultsVBoxPtr->ClearChildren();
+
+	if (InContent.NetRecv && InContent.NetRecv->SystemFaults.Num() > 0)
+	{
+		FString FaultString = TEXT("NetRecv Faults:\n");
+
+		for (const FSimulationData::FSystemFault& Fault : InContent.NetRecv->SystemFaults)
+		{
+			FaultString.Append(Fault.Str);
+			FaultString.Append(TEXT("\n"));
+		}
+
+		SystemFaultsVBoxPtr->AddSlot()
+		.AutoHeight()
+		.VAlign(VAlign_Top)
+		.HAlign(HAlign_Left)
+		.Padding(4.0f, 4.0f, 4.0f, 4.0f)
+		[
+			SNew(SEditableText)
+			.IsReadOnly(true)
+			.Text(FText::FromString(FaultString))
+		];
+	}
+
+	if (EOFState && EOFState->SystemFaults.Num() > 0)
+	{
+		FString FaultString = TEXT("Engine Frame Faults:\n");
+
+		for (const FSimulationData::FSystemFault& Fault : EOFState->SystemFaults)
+		{
+			FaultString.Append(Fault.Str);
+			FaultString.Append(TEXT("\n"));
+		}
+
+		SystemFaultsVBoxPtr->AddSlot()
+		.AutoHeight()
+		.VAlign(VAlign_Top)
+		.HAlign(HAlign_Left)
+		.Padding(4.0f, 4.0f, 4.0f, 4.0f)
+		[
+			SNew(SEditableText)
+			.IsReadOnly(true)
+			.Text(FText::FromString(FaultString))
+		];
+	}
+
 }
 
 
diff --git a/Engine/Plugins/Runtime/NetworkPredictionInsights/Source/NetworkPredictionInsights/Private/UI/SNPSimFrameContents.h b/Engine/Plugins/Runtime/NetworkPredictionInsights/Source/NetworkPredictionInsights/Private/UI/SNPSimFrameContents.h
--- a/Engine/Plugins/Runtime/NetworkPredictionInsights/Source/NetworkPredictionInsights/Private/UI/SNPSimFrameContents.h
+++ b/Engine/Plugins/Runtime/NetworkPredictionInsights/Source/NetworkPredictionInsights/Private/UI/SNPSimFrameContents.h
@@ -27,6 +27,7 @@ public:
 	void NotifyContentClicked(const FSimContentsView& InView);
 
 	TSharedPtr<SHorizontalBox> ContentsHBoxPtr;
+	TSharedPtr<SVerticalBox> SystemFaultsVBoxPtr;
 
 	TSharedPtr<STextBlock> SimInfoTextBlock;
 	TSharedPtr<STextBlock> SimTickTextBlock;
diff --git a/Engine/Plugins/Runtime/NetworkPredictionInsights/Source/NetworkPredictionInsights/Private/UI/SNPSimFrameView.cpp b/Engine/Plugins/Runtime/NetworkPredictionInsights/Source/NetworkPredictionInsights/Private/UI/SNPSimFrameView.cpp
--- a/Engine/Plugins/Runtime/NetworkPredictionInsights/Source/NetworkPredictionInsights/Private/UI/SNPSimFrameView.cpp
+++ b/Engine/Plugins/Runtime/NetworkPredictionInsights/Source/NetworkPredictionInsights/Private/UI/SNPSimFrameView.cpp
@@ -428,7 +428,7 @@ void SNPSimFrameView::BuildSimulationView_ActorGroups(const FFilteredDataCollect
 	auto InitTrack = [&](FSimulationActorGroup& Group, const TSharedRef<FSimulationData::FRestrictedView>& View)
 	{
 		Group.SimulationTracks.Emplace( FSimulationTrack{View});
-		Group.MaxAllowedSimTime = FMath::Max(Group.MaxAllowedSimTime, View->EOFState.GetLast().TotalAllowedSimTimeMS);
+		Group.MaxAllowedSimTime = FMath::Max(Group.MaxAllowedSimTime, View->EOFState.GetLast().AllowedSimTime);
 		Group.MaxEngineFrame = FMath::Max(Group.MaxEngineFrame, View->EOFState.GetLast().EngineFrame);
 
 		SimulationFrameView.HeadEngineFrame = FMath::Max(SimulationFrameView.HeadEngineFrame, Group.MaxEngineFrame);
@@ -634,8 +634,8 @@ void SNPSimFrameView::BuildSimulationView_ActorGroups(const FFilteredDataCollect
 				check(BaseState && CurrentState);
 
 				// Compute difference in ProcessedSimulationTime at the end of this engine frame
-				const FSimTime BaseSimMS = BaseState->TotalProcessedSimTimeMS;
-				const FSimTime CurrentSimMS = CurrentState->TotalProcessedSimTimeMS;
+				const FSimTime BaseSimMS = BaseState->TotalSimTime;
+				const FSimTime CurrentSimMS = CurrentState->TotalSimTime;
 				const FSimTime Delta = BaseSimMS - CurrentSimMS;
 
 				// This is the final amount to offset this groups local sim time to align it properly with the primary simulation in this range
@@ -662,7 +662,7 @@ void SNPSimFrameView::BuildSimulationView_ActorGroups(const FFilteredDataCollect
 			MinOffset = FMath::Min(MinOffset, View.Group.OffsetSimTimeMS);
 		}
 
-		PrevSimTime = Range.Views[0].View->EOFState.GetLast().TotalProcessedSimTimeMS + 1000;
+		PrevSimTime = Range.Views[0].View->EOFState.GetLast().TotalSimTime + 1000;
 		if (MinOffset < 0)
 		{
 			PrevSimTime += -MinOffset;
@@ -674,7 +674,7 @@ void SNPSimFrameView::BuildSimulationView_ActorGroups(const FFilteredDataCollect
 	// Set presentable line, this is where the gray/light gray line is drawn
 	if (ActorGroups.Num() > 0)
 	{
-		SimulationFrameView.PresentableTimeMS = ActorGroups[0].OffsetSimTimeMS + ActorGroups[0].SimulationTracks[0].View->EOFState.GetLast().TotalProcessedSimTimeMS;
+		SimulationFrameView.PresentableTimeMS = ActorGroups[0].OffsetSimTimeMS + ActorGroups[0].SimulationTracks[0].View->EOFState.GetLast().TotalSimTime;
 	}
 
 	// Calc ViewportMaxSimTimeMS
@@ -702,8 +702,8 @@ for (FSimulationActorGroup& Group : ActorGroups)
 	{
 		if (Track.View->ConstData.GameInstanceId == PresentableGameInstanceId)
 		{
-			FSimTime LastProcessedSimTimeMS = Track.View->EOFState.GetLast().TotalProcessedSimTimeMS;
-			Group.PresentableSimTimeMS = Track.View->EOFState.GetLast().TotalProcessedSimTimeMS;
+			FSimTime LastProcessedSimTimeMS = Track.View->EOFState.GetLast().TotalSimTime;
+			Group.PresentableSimTimeMS = Track.View->EOFState.GetLast().TotalSimTime;
 			SimulationFrameView.PresentableTimeMS = FMath::Max(SimulationFrameView.PresentableTimeMS, Group.PresentableSimTimeMS);
 			break;
 		}
diff --git a/Engine/Plugins/Runtime/NetworkPredictionInsights/Source/NetworkPredictionInsights/Public/INetworkPredictionProvider.h b/Engine/Plugins/Runtime/NetworkPredictionInsights/Source/NetworkPredictionInsights/Public/INetworkPredictionProvider.h
--- a/Engine/Plugins/Runtime/NetworkPredictionInsights/Source/NetworkPredictionInsights/Public/INetworkPredictionProvider.h
+++ b/Engine/Plugins/Runtime/NetworkPredictionInsights/Source/NetworkPredictionInsights/Public/INetworkPredictionProvider.h
@@ -332,6 +332,11 @@ struct FSimulationData
 
 	struct FTick;
 
+	struct FSystemFault
+	{
+		const TCHAR* Str; // Traced Fmt string of fault
+	};
+
 	struct FNetSerializeRecv
 	{
 		// Traced Data:
@@ -344,6 +349,7 @@ struct FSimulationData
 		bool bOrphan = true; // Hasn't been attached to FTick::StartNetRecv
 
 		FSimulationData::FTick* NextTick = nullptr;
+		TArray<FSystemFault> SystemFaults; // System faults encountered during recv
 	};
 
 	// Data for a single simulation tick
@@ -373,14 +379,16 @@ struct FSimulationData
 	struct FEngineFrame
 	{
 		uint64 EngineFrame;
-		double EngineFrameDeltaTime;
-		double EngineCurrentTime;
+		float EngineFrameDeltaTime;
 
-		FSimTime	TotalProcessedSimTimeMS;
-		FSimTime	TotalAllowedSimTimeMS;
-		
-		int32		LastSentKeyframe;
-		int32		LastReceivedKeyframe;
+		int32 BufferSize;
+		int32 PendingTickFrame;
+		int32 LatestInputFrame;
+		int32 MaxTickFrame;
+		FSimTime TotalSimTime;
+		FSimTime AllowedSimTime;
+
+		TArray<FSystemFault> SystemFaults;	// System faults encountered during this tick
 	};
 
 	// Data that changes rarely over time about the simulation
@@ -403,8 +411,6 @@ struct FSimulationData
 		FString DebugName;
 	};
 
-	// -----------------------------------------------------------
-
 	struct FUserState
 	{
 		uint64 EngineFrame;	// EngineFrame this was written
@@ -583,6 +589,7 @@ struct FSimulationData
 
 		ENP_UserStateSource PendingUserStateSource = ENP_UserStateSource::Unknown;
 		TArray<FUserState*> PendingCommitUserStates; // NetRecv'd state that hasn't been commited
+		TArray<FSystemFault> PendingSystemFaults;
 
 	} Analysis;
 };
