NetworkPrediction Plugin
Few more minor renames/cleanup

#rnx

diff --git a/Engine/Plugins/Runtime/NetworkPrediction/Source/NetworkPrediction/Public/NetworkedSimulationModelBuffer.h b/Engine/Plugins/Runtime/NetworkPrediction/Source/NetworkPrediction/Public/NetworkedSimulationModelBuffer.h
--- a/Engine/Plugins/Runtime/NetworkPrediction/Source/NetworkPrediction/Public/NetworkedSimulationModelBuffer.h
+++ b/Engine/Plugins/Runtime/NetworkPrediction/Source/NetworkPrediction/Public/NetworkedSimulationModelBuffer.h
@@ -4,7 +4,7 @@
 #include "Containers/UnrealString.h"
 #include "Containers/ArrayView.h"
 
-// Helper struct to encapsulate optional, delayed writing of new element to buffer
+// Helper struct to encapsulate optional, delayed writing of new element to TNetworkSimAuxBuffer buffer
 template<typename ElementType>
 struct TNetSimLazyWriterFunc
 {
@@ -39,9 +39,7 @@ struct TNetSimLazyWriter
 	TFunctionRef<void*()> GetFunc;
 };
 
-
-// ---------------------------------------------------------------------------------------
-// ---------------------------------------------------------------------------------------
+// Sparse buffer: each element has explicit Frame value
 template<typename ElementType> 
 struct TNetworkSimAuxBuffer
 {
@@ -133,9 +131,7 @@ private:
 	TArray<TInternal> Elements;
 };
 
-// ---------------------------------------------------------------------------------------
-// ---------------------------------------------------------------------------------------
-
+// Circular buffer for storing continuous frame data
 template<typename ElementType> 
 struct TNetworkSimFrameBuffer
 {
diff --git a/Engine/Plugins/Runtime/NetworkPrediction/Source/NetworkPrediction/Public/NetworkedSimulationModelDebugger.h b/Engine/Plugins/Runtime/NetworkPrediction/Source/NetworkPrediction/Public/NetworkedSimulationModelDebugger.h
--- a/Engine/Plugins/Runtime/NetworkPrediction/Source/NetworkPrediction/Public/NetworkedSimulationModelDebugger.h
+++ b/Engine/Plugins/Runtime/NetworkPrediction/Source/NetworkPrediction/Public/NetworkedSimulationModelDebugger.h
@@ -364,7 +364,7 @@ struct TNetworkSimulationModelDebugger : public INetworkSimulationModelDebugger
 
 			for (int32 PredFrame = LastSerializedFrame+1; PredFrame <= NetworkSim->State.GetPendingTickFrame(); ++PredFrame)
 			{
-				if (auto* Cmd = NetworkSim->State.GetFrameState(PredFrame))
+				if (auto* Cmd = NetworkSim->State.ReadFrame(PredFrame))
 				{
 					PredictedMS += Cmd->FrameDeltaTime.Get();
 				}
@@ -406,11 +406,11 @@ struct TNetworkSimulationModelDebugger : public INetworkSimulationModelDebugger
 
 		{
 			const int32 HeadFrame = NetworkSim->State.GetPendingTickFrame();
-			if (auto* LatestFrameState = NetworkSim->State.GetFrameState(HeadFrame))
+			if (auto* LatestFrameState = NetworkSim->State.ReadFrame(HeadFrame))
 			{
 				auto* LatestSync = &LatestFrameState->SyncState;
 				FVisualLoggingParameters VLogParams(EVisualLoggingContext::LastPredicted, HeadFrame, EVisualLoggingLifetime::Transient);
-				NetworkSim->Driver->InvokeVisualLog(&LatestFrameState->InputCmd, LatestSync, NetworkSim->State.GetAuxState(HeadFrame), VLogParams);
+				NetworkSim->Driver->InvokeVisualLog(&LatestFrameState->InputCmd, LatestSync, NetworkSim->State.ReadAux(HeadFrame), VLogParams);
 			}
 		}
 
@@ -419,10 +419,10 @@ struct TNetworkSimulationModelDebugger : public INetworkSimulationModelDebugger
 			if (ServerPIEStuff.NetworkSim)
 			{
 				const int32 HeadFrame = ServerPIEStuff.NetworkSim->State.GetPendingTickFrame();
-				if (auto* LatestFrameState = NetworkSim->State.GetFrameState(HeadFrame))
+				if (auto* LatestFrameState = NetworkSim->State.ReadFrame(HeadFrame))
 				{
 					FVisualLoggingParameters VLogParams(EVisualLoggingContext::CurrentServerPIE, HeadFrame, EVisualLoggingLifetime::Transient);
-					NetworkSim->Driver->InvokeVisualLog(&LatestFrameState->InputCmd, &LatestFrameState->SyncState, ServerPIEStuff.NetworkSim->State.GetAuxState(HeadFrame), VLogParams);
+					NetworkSim->Driver->InvokeVisualLog(&LatestFrameState->InputCmd, &LatestFrameState->SyncState, ServerPIEStuff.NetworkSim->State.ReadAux(HeadFrame), VLogParams);
 				}
 			}
 		}
@@ -431,11 +431,11 @@ struct TNetworkSimulationModelDebugger : public INetworkSimulationModelDebugger
 		{
 			for (int32 Frame = NetworkSim->RepProxy_Autonomous.GetLastSerializedFrame(); Frame < NetworkSim->State.GetPendingTickFrame(); ++Frame)
 			{
-				if (auto* FrameState = NetworkSim->State.GetFrameState(Frame))
+				if (auto* FrameState = NetworkSim->State.ReadFrame(Frame))
 				{
 					const EVisualLoggingContext Context = (Frame == NetworkSim->RepProxy_Autonomous.GetLastSerializedFrame()) ? EVisualLoggingContext::LastConfirmed : EVisualLoggingContext::OtherPredicted;
 					FVisualLoggingParameters VLogParams(Context, Frame, EVisualLoggingLifetime::Transient);
-					NetworkSim->Driver->InvokeVisualLog(&FrameState->InputCmd, &FrameState->SyncState, NetworkSim->State.GetAuxState(Frame), VLogParams);
+					NetworkSim->Driver->InvokeVisualLog(&FrameState->InputCmd, &FrameState->SyncState, NetworkSim->State.ReadAux(Frame), VLogParams);
 				}
 			}
 		}
@@ -444,8 +444,8 @@ struct TNetworkSimulationModelDebugger : public INetworkSimulationModelDebugger
 			{
 				int32 HeadFrame = NetworkSim->State.GetPendingTickFrame();
 				FVisualLoggingParameters VLogParams(EVisualLoggingContext::LastPredicted, HeadFrame, EVisualLoggingLifetime::Transient);
-				auto* FrameState = NetworkSim->State.GetFrameState(HeadFrame);
-				NetworkSim->Driver->InvokeVisualLog(&FrameState->InputCmd, &FrameState->SyncState, NetworkSim->State.GetAuxState(HeadFrame), VLogParams);
+				auto* FrameState = NetworkSim->State.ReadFrame(HeadFrame);
+				NetworkSim->Driver->InvokeVisualLog(&FrameState->InputCmd, &FrameState->SyncState, NetworkSim->State.ReadAux(HeadFrame), VLogParams);
 			}
 
 			{
diff --git a/Engine/Plugins/Runtime/NetworkPrediction/Source/NetworkPrediction/Public/NetworkedSimulationModelInterpolator.h b/Engine/Plugins/Runtime/NetworkPrediction/Source/NetworkPrediction/Public/NetworkedSimulationModelInterpolator.h
--- a/Engine/Plugins/Runtime/NetworkPrediction/Source/NetworkPrediction/Public/NetworkedSimulationModelInterpolator.h
+++ b/Engine/Plugins/Runtime/NetworkPrediction/Source/NetworkPrediction/Public/NetworkedSimulationModelInterpolator.h
@@ -88,8 +88,8 @@ struct TNetSimInterpolator
 			const int32 HeadFrame = State.GetPendingTickFrame();
 			if (HeadFrame >= 0)
 			{
-				const TFrameState* FrameState = State.GetFrameState(HeadFrame);
-				const TAuxState* AuxState = State.GetAuxState(HeadFrame);
+				const TFrameState* FrameState = State.ReadFrame(HeadFrame);
+				const TAuxState* AuxState = State.ReadAux(HeadFrame);
 				check(FrameState && AuxState);
 				Driver->FinalizeFrame(FrameState->SyncState, *AuxState);
 			}
@@ -97,8 +97,8 @@ struct TNetSimInterpolator
 			return State.GetTotalProcessedSimulationTime().ToRealTimeSeconds();
 		}
 
-		const TFrameState* HeadFrameState = State.GetFrameState(State.GetPendingTickFrame());
-		const TFrameState* TailFrameState = State.GetFrameState(State.GetConfirmedFrame());
+		const TFrameState* HeadFrameState = State.ReadFrame(State.GetPendingTickFrame());
+		const TFrameState* TailFrameState = State.ReadFrame(State.GetConfirmedFrame());
 
 		FRealTime HeadRealTimeSeconds = HeadFrameState->TotalSimTime.ToRealTimeSeconds();
 		FRealTime TailRealTimeSeconds = TailFrameState->TotalSimTime.ToRealTimeSeconds();
@@ -110,8 +110,8 @@ struct TNetSimInterpolator
 			InterpolationFrame = State.GetConfirmedFrame();
 
 			auto& FromState = GetFromInterpolationState();
-			FromState.Sync = State.GetFrameState(InterpolationFrame)->SyncState;
-			FromState.Aux = *State.GetAuxState(InterpolationFrame);
+			FromState.Sync = State.ReadFrame(InterpolationFrame)->SyncState;
+			FromState.Aux = *State.ReadAux(InterpolationFrame);
 		}
 
 		EVisualLoggingContext LoggingContext = EVisualLoggingContext::InterpolationLatest;
@@ -192,13 +192,13 @@ struct TNetSimInterpolator
 
 		for (int32 Frame = State.GetConfirmedFrame(); Frame <= State.GetPendingTickFrame(); ++Frame)
 		{
-			FRealTime FrameRealTime = State.GetFrameState(Frame)->TotalSimTime.ToRealTimeSeconds();
+			FRealTime FrameRealTime = State.ReadFrame(Frame)->TotalSimTime.ToRealTimeSeconds();
 			if (NewInterpolationTime <= FrameRealTime)
 			{
 				InterpolationFrame = Frame;
 				ToTime = FrameRealTime;
-				ToState = &State.GetFrameState(Frame)->SyncState;
-				ToAuxState = State.GetAuxState(Frame);
+				ToState = &State.ReadFrame(Frame)->SyncState;
+				ToAuxState = State.ReadAux(Frame);
 				break;
 			}
 		}
@@ -251,9 +251,9 @@ struct TNetSimInterpolator
 				auto VLogHelper = [&](int32 Frame, EVisualLoggingContext Context, const FString& DebugStr)
 				{
 					FVisualLoggingParameters VLogParams(Context, Frame, EVisualLoggingLifetime::Transient, DebugStr);
-					const TFrameState* FrameState = State.GetFrameState(Frame);
+					const TFrameState* FrameState = State.ReadFrame(Frame);
 
-					Driver->InvokeVisualLog(&FrameState->InputCmd, &FrameState->SyncState, State.GetAuxState(Frame), VLogParams);
+					Driver->InvokeVisualLog(&FrameState->InputCmd, &FrameState->SyncState, State.ReadAux(Frame), VLogParams);
 				};
 
 				VLogHelper(State.GetConfirmedFrame(), EVisualLoggingContext::InterpolationBufferTail, LexToString(TailFrameState->TotalSimTime.ToRealTimeMS()));
@@ -261,17 +261,17 @@ struct TNetSimInterpolator
 
 				{
 					FVisualLoggingParameters VLogParams(EVisualLoggingContext::InterpolationFrom, InterpolationFrame-1, EVisualLoggingLifetime::Transient, LexToString(FromRealTime));
-					Driver->InvokeVisualLog(&State.GetFrameState(InterpolationFrame-1)->InputCmd, &FromState.Sync, &FromState.Aux, VLogParams);
+					Driver->InvokeVisualLog(&State.ReadFrame(InterpolationFrame-1)->InputCmd, &FromState.Sync, &FromState.Aux, VLogParams);
 				}
 
 				{
 					FVisualLoggingParameters VLogParams(EVisualLoggingContext::InterpolationTo, InterpolationFrame, EVisualLoggingLifetime::Transient, LexToString(ToRealTime));
-					Driver->InvokeVisualLog(&State.GetFrameState(InterpolationFrame)->InputCmd, ToState, ToAuxState, VLogParams);
+					Driver->InvokeVisualLog(&State.ReadFrame(InterpolationFrame)->InputCmd, ToState, ToAuxState, VLogParams);
 				}
 
 				{
 					FVisualLoggingParameters VLogParams(LoggingContext, InterpolationFrame, EVisualLoggingLifetime::Transient, LexToString(NewInterpolationTime));
-					Driver->InvokeVisualLog(&State.GetFrameState(InterpolationFrame)->InputCmd, &OutputState.Sync, &OutputState.Aux, VLogParams);
+					Driver->InvokeVisualLog(&State.ReadFrame(InterpolationFrame)->InputCmd, &OutputState.Sync, &OutputState.Aux, VLogParams);
 				}
 			}
 				
diff --git a/Engine/Plugins/Runtime/NetworkPrediction/Source/NetworkPrediction/Public/NetworkedSimulationModelRepControllers.h b/Engine/Plugins/Runtime/NetworkPrediction/Source/NetworkPrediction/Public/NetworkedSimulationModelRepControllers.h
--- a/Engine/Plugins/Runtime/NetworkPrediction/Source/NetworkPrediction/Public/NetworkedSimulationModelRepControllers.h
+++ b/Engine/Plugins/Runtime/NetworkPrediction/Source/NetworkPrediction/Public/NetworkedSimulationModelRepControllers.h
@@ -130,7 +130,7 @@ void TryGenerateLocalInput(TNetworkedSimulationModelDriver<typename Model::Buffe
 		FrameState->SetFlag(ESimulationFrameStateFlags::InputWritten);
 		State.SetLatestInputFrame(InputFrame);
 
-		npEnsureMsgfSlow(State.GetTotalAllowedSimulationTime() >= State.GetTotalInputSimulationTime(), TEXT("Overflowed input %s"), *State.GetBasicDebugStr());
+		npEnsureMsgfSlow(State.GetTotalAllowedSimulationTime() >= State.GetTotalInputSimulationTime(), TEXT("Overflowed input %s"), *State.DebugString());
 		MissingInputTime -= FrameState->FrameDeltaTime.Get();
 
 		if (bProduceInputViaDriver)
@@ -149,7 +149,7 @@ void TryGenerateLocalInput(TNetworkedSimulationModelDriver<typename Model::Buffe
 	// Advance MaxTickFrame
 	for (int32 Frame = State.GetLatestInputFrame(); Frame > State.GetMaxTickFrame(); --Frame)
 	{
-		auto* FrameState = State.GetFrameState(Frame);
+		auto* FrameState = State.ReadFrame(Frame);
 		if (npEnsure(FrameState) && FrameState->HasFlag(ESimulationFrameStateFlags::InputWritten) && FrameState->TotalSimTime <= AllowedSimTime)
 		{
 			State.SetMaxTickFrame(Frame);
@@ -169,7 +169,7 @@ void DoSimulationTicks(typename Model::Simulation& Simulation, TNetworkedSimulat
 			
 		// Find InputFrame. We want it to be PendingFrame but if no input was written at this frame, it can be skipped if there are valid inputs at later frames (gaps in network recvs)
 		int32 InputFrame = InputSyncFrame;
-		TFrameState* InputFrameState = State.GetFrameState(InputFrame);
+		TFrameState* InputFrameState = State.ReadFrame(InputFrame);
 		while (InputFrameState && !InputFrameState->HasFlag(ESimulationFrameStateFlags::InputWritten))
 		{
 			if (++InputFrame > State.GetMaxTickFrame())
@@ -177,7 +177,7 @@ void DoSimulationTicks(typename Model::Simulation& Simulation, TNetworkedSimulat
 				InputFrameState = nullptr;
 				break;
 			}
-			InputFrameState = State.GetFrameState(InputFrame);
+			InputFrameState = State.ReadFrame(InputFrame);
 		}
 
 		// No valid input to be processed between PendingFrame -> MaxAllowedFrame
@@ -189,7 +189,7 @@ void DoSimulationTicks(typename Model::Simulation& Simulation, TNetworkedSimulat
 		// If we skipped frames due to lack of input, we need to copy the sync state over and set PendingTickFrame
 		if (InputFrame > InputSyncFrame)
 		{
-			TSimulationFrameState<Model>* InputSyncState = State.GetFrameState(InputSyncFrame);
+			TSimulationFrameState<Model>* InputSyncState = State.ReadFrame(InputSyncFrame);
 			InputFrameState->SyncState = InputSyncState->SyncState;
 			InputFrameState->TotalSimTime = InputSyncState->TotalSimTime;
 			State.SetPendingTickFrame(InputFrame);
@@ -242,8 +242,10 @@ struct TRepController_Server : public TBase
 	using typename TBase::TSimulation;
 	using typename TBase::TTickSettings;
 	using typename TBase::TDriver;
-	using typename TBase::TInputCmd;
 	using typename TBase::TFrameState;
+	using typename TBase::TInputCmd;
+	using typename TBase::TSyncState;
+	using typename TBase::TAuxState;
 	
 	// -------------------------------------------------------------------------------------------------
 	// TRepController_Server::NetSerialize
@@ -278,7 +280,7 @@ struct TRepController_Server : public TBase
 
 			for (int32 Frame = StartInputFrame; Frame <= HeadInputFrame; ++Frame)
 			{
-				TFrameState* FrameState = State.GetFrameState(Frame);
+				TFrameState* FrameState = State.ReadFrame(Frame);
 				if (FrameState && FrameState->HasFlag(ESimulationFrameStateFlags::InputWritten))
 				{
 					SerializedNumElements++;
@@ -316,7 +318,7 @@ struct TRepController_Server : public TBase
 			const int32 StartInputFrame = FMath::Max(ReceivedHeadInputFrame - SerializedNumElements + 1, 0);
 
 			const int32 MinAcceptableFrame = State.GetNextInputFrame();
-			const int32 MaxAcceptableFrame = State.GetPendingTickFrame() + State.Max() - 1;
+			const int32 MaxAcceptableFrame = State.GetPendingTickFrame() + State.Capacity() - 1;
 
 			for (int32 Frame = StartInputFrame; Frame <= ReceivedHeadInputFrame; ++Frame)
 			{
@@ -334,7 +336,7 @@ struct TRepController_Server : public TBase
 					LastSerializedFrame = Frame;
 				}
 
-				if (TFrameState* ExistingFrame = State.GetFrameState(Frame))
+				if (TFrameState* ExistingFrame = State.ReadFrame(Frame))
 				{
 					if (npEnsure(!ExistingFrame->HasFlag(ESimulationFrameStateFlags::InputWritten)))
 					{
@@ -397,7 +399,7 @@ struct TRepController_Server : public TBase
 
 		// Sync to latest frame if there is any
 		const int32 PendingFrame = State.GetPendingTickFrame();
-		Driver->FinalizeFrame(State.GetFrameState(PendingFrame)->SyncState, *State.GetAuxState(PendingFrame));
+		Driver->FinalizeFrame(State.ReadFrame(PendingFrame)->SyncState, *State.ReadAux(PendingFrame));
 	}
 
 	int32 GetLastSerializedFrame() const { return LastSerializedFrame; }
@@ -418,13 +420,11 @@ struct TRepController_Autonomous: public TBase
 	using typename TBase::TSimulation;
 	using typename TBase::TTickSettings;
 	using typename TBase::TDriver;
-	using typename TBase::TBufferTypes;
+	using typename TBase::TFrameState;
 	using typename TBase::TInputCmd;
 	using typename TBase::TSyncState;
 	using typename TBase::TAuxState;
 
-	using TFrameState = TSimulationFrameState<Model>;
-
 	int32 GetLastSerializedFrame() const { return SerializedFrame; }
 	bool IsReconcileFaultDetected() const { return bReconcileFaultDetected; }
 	const FNetworkSimTime& GetLastSerializedSimTime() const { return SerializedTime; }
@@ -448,8 +448,8 @@ struct TRepController_Autonomous: public TBase
 		if (Ar.IsSaving())
 		{
 			// Server serialize the latest state
-			State.GetFrameState(SerializedFrame)->SyncState.NetSerialize(Ar);
-			State.GetAuxState(SerializedFrame)->NetSerialize(Ar);
+			State.ReadFrame(SerializedFrame)->SyncState.NetSerialize(Ar);
+			State.ReadAux(SerializedFrame)->NetSerialize(Ar);
 			State.CueDispatcher.NetSerializeSavedCues(Ar, ENetSimCueReplicationTarget::AutoProxy);
 		}
 		else
@@ -464,8 +464,8 @@ struct TRepController_Autonomous: public TBase
 			
 			bPendingReconciliation = false;
 
-			TFrameState* ClientFrameState = State.GetFrameState(SerializedFrame);
-			TAuxState* ClientAuxState = State.GetAuxState(SerializedFrame);
+			TFrameState* ClientFrameState = State.ReadFrame(SerializedFrame);
+			TAuxState* ClientAuxState = State.ReadAux(SerializedFrame);
 
 			// The state the client predicted that corresponds to the state the server just serialized to us
 			if (ClientFrameState && ClientAuxState)
@@ -503,12 +503,12 @@ struct TRepController_Autonomous: public TBase
 		bPendingReconciliation = false;
 		bDependentSimulationNeedsReconcile = false;
 
-		TFrameState* ClientFrameState = State.GetFrameState(SerializedFrame);
+		TFrameState* ClientFrameState = State.ReadFrame(SerializedFrame);
 		if (!ClientFrameState)
 		{
 			// We received state that we don't have in our local buffer.
 			// We can't resimulate here, just reset to the serialized state
-			UE_NP_TRACE_SYSTEM_FAULT("Could not reconcile SerializedFrame %d with local state. %s", SerializedFrame, *State.GetBasicDebugStr());
+			UE_NP_TRACE_SYSTEM_FAULT("Could not reconcile SerializedFrame %d with local state. %s", SerializedFrame, *State.DebugString());
 			ResetToSerializedState(State);
 			return;
 		}
@@ -577,26 +577,26 @@ struct TRepController_Autonomous: public TBase
 			NetworkSimulationModelCVars::SetForceAutonomousFrameJump(0);
 			
 			// Copy pending state to jumped frame
-			TFrameState PendingFrameState = *State.GetFrameState(State.GetPendingTickFrame());
+			TFrameState PendingFrameState = *State.ReadFrame(State.GetPendingTickFrame());
 			const int32 NewFrame = State.GetPendingTickFrame() + FrameJump;
 			TFrameState* NewFrameState =  State.WriteFrameSequential(NewFrame);
 			*NewFrameState = PendingFrameState;
 
 			State.Reset(NewFrame);
 
-			UE_NP_TRACE_SYSTEM_FAULT("NetworkSimulationModelCVars::ForceAutonomousFrameJump invoked. Jumped to: %d. Prev: %s", FrameJump, *State.GetBasicDebugStr());
+			UE_NP_TRACE_SYSTEM_FAULT("NetworkSimulationModelCVars::ForceAutonomousFrameJump invoked. Jumped to: %d. Prev: %s", FrameJump, *State.DebugString());
 		}
 
 		// The MaxInputFrame we can write without getting too far ahead of server
 		// We pass this into TryGenerateLocalInput to ensure that it doesn't.
 		// If we some how get in a state where we are past MaxInputFrame, we need to reset
-		const int32 MaxInputFrame = (SerializedFrame + State.Max() - 1);
+		const int32 MaxInputFrame = (SerializedFrame + State.Capacity() - 1);
 
 		if (State.GetNextInputFrame() > MaxInputFrame)
 		{
 			// We somehow went over the limit and must reset.
 			ResetToSerializedState(State);
-			UE_NP_TRACE_SYSTEM_FAULT("TRepController_Autonomous::PreSimTick input overrun. MaxInputFrame: %d. %s", MaxInputFrame, *State.GetBasicDebugStr());
+			UE_NP_TRACE_SYSTEM_FAULT("TRepController_Autonomous::PreSimTick input overrun. MaxInputFrame: %d. %s", MaxInputFrame, *State.DebugString());
 			return;
 		}
 
@@ -604,7 +604,7 @@ struct TRepController_Autonomous: public TBase
 		{
 			// We are at the limit but not over it. We are essentially pausing the simulation until we recover. Set this bool to reflect this.
 			bReconcileFaultDetected = true;
-			UE_NP_TRACE_SYSTEM_FAULT("TRepController_Autonomous::PreSimTick Input Buffer Maxed. Not processing local time. MaxInputFrame: %d. %s", MaxInputFrame, *State.GetBasicDebugStr());
+			UE_NP_TRACE_SYSTEM_FAULT("TRepController_Autonomous::PreSimTick Input Buffer Maxed. Not processing local time. MaxInputFrame: %d. %s", MaxInputFrame, *State.DebugString());
 			return;
 		}
 
@@ -638,7 +638,7 @@ struct TRepController_Autonomous: public TBase
 		State.CueDispatcher.SetConfirmedTime(SerializedTime);
 
 		const int32 PendingFrame = State.GetPendingTickFrame();
-		Driver->FinalizeFrame(State.GetFrameState(PendingFrame)->SyncState, *State.GetAuxState(PendingFrame));
+		Driver->FinalizeFrame(State.ReadFrame(PendingFrame)->SyncState, *State.ReadAux(PendingFrame));
 	}
 
 private:
@@ -660,11 +660,11 @@ private:
 		{
 			for (int32 Frame = StartFrame; Frame <= EndFrame; ++Frame)
 			{
-				TFrameState* FrameState = State.GetFrameState(Frame);
-				if (npEnsureMsgf(FrameState, TEXT("Invalid Frame %d. %s"), Frame, *State.GetBasicDebugStr()))
+				TFrameState* FrameState = State.ReadFrame(Frame);
+				if (npEnsureMsgf(FrameState, TEXT("Invalid Frame %d. %s"), Frame, *State.DebugString()))
 				{
 					FVisualLoggingParameters VLogParameters(Context, Frame, EVisualLoggingLifetime::Persistent, Str);
-					Driver->InvokeVisualLog(&FrameState->InputCmd, &FrameState->SyncState, State.GetAuxState(Frame), VLogParameters);
+					Driver->InvokeVisualLog(&FrameState->InputCmd, &FrameState->SyncState, State.ReadAux(Frame), VLogParameters);
 				}
 			}
 		}
@@ -695,13 +695,11 @@ struct TRepController_Simulated : public TBase
 	using typename TBase::TSimulation;
 	using typename TBase::TTickSettings;
 	using typename TBase::TDriver;
-	using typename TBase::TBufferTypes;
+	using typename TBase::TFrameState;
 	using typename TBase::TInputCmd;
 	using typename TBase::TSyncState;
 	using typename TBase::TAuxState;
 
-	using TFrameState = TSimulationFrameState<Model>;
-
 	// Parent Simulation. If this is set, this simulation will forward predict in sync with this parent sim. The parent sim should be an autonomous proxy driven simulation
 	INetworkedSimulationModel* ParentSimulation = nullptr;
 
@@ -744,13 +742,13 @@ struct TRepController_Simulated : public TBase
 		{
 			LastSerializedKeyframe = State.GetPendingTickFrame();
 			
-			TFrameState* FrameState = State.GetFrameState(LastSerializedKeyframe);
+			TFrameState* FrameState = State.ReadFrame(LastSerializedKeyframe);
 			npCheckSlow(FrameState);
 
-			TFrameState* InputFrameState = State.GetFrameState(State.GetLatestInputFrame());
-			npCheckfSlow(InputFrameState, TEXT("Invalid InputFrame. %s"), *State.GetBasicDebugStr());
+			TFrameState* InputFrameState = State.ReadFrame(State.GetLatestInputFrame());
+			npCheckfSlow(InputFrameState, TEXT("Invalid InputFrame. %s"), *State.DebugString());
 
-			TAuxState* AuxState = State.GetAuxState(LastSerializedKeyframe);
+			TAuxState* AuxState = State.ReadAux(LastSerializedKeyframe);
 			npCheckSlow(AuxState);
 			
 			FrameState->TotalSimTime.NetSerialize(P.Ar);	// 1. Time
@@ -880,7 +878,7 @@ struct TRepController_Simulated : public TBase
 			const int32 PendingFrame = State.GetPendingTickFrame();
 			if (PendingFrame >= 0)
 			{
-				Driver->FinalizeFrame(State.GetFrameState(PendingFrame)->SyncState, *State.GetAuxState(PendingFrame));
+				Driver->FinalizeFrame(State.ReadFrame(PendingFrame)->SyncState, *State.ReadAux(PendingFrame));
 			}
 
 			State.CueDispatcher.ClearMaxDispatchTime();
@@ -924,7 +922,7 @@ private:
 		const int32 OutputFrame = InputFrame + 1;
 
 		// Override the FrameDeltaTime
-		TFrameState* InputFrameState = State.GetFrameState(InputFrame);
+		TFrameState* InputFrameState = State.ReadFrame(InputFrame);
 		npCheckSlow(InputFrameState);
 		npEnsureSlow(InputFrameState->HasFlag(ESimulationFrameStateFlags::InputWritten));
 		InputFrameState->FrameDeltaTime.Set(DeltaSimTime);
diff --git a/Engine/Plugins/Runtime/NetworkPrediction/Source/NetworkPrediction/Public/NetworkedSimulationModelTypes.h b/Engine/Plugins/Runtime/NetworkPrediction/Source/NetworkPrediction/Public/NetworkedSimulationModelTypes.h
--- a/Engine/Plugins/Runtime/NetworkPrediction/Source/NetworkPrediction/Public/NetworkedSimulationModelTypes.h
+++ b/Engine/Plugins/Runtime/NetworkPrediction/Source/NetworkPrediction/Public/NetworkedSimulationModelTypes.h
@@ -47,9 +47,9 @@ struct TNetworkedSimulationState
 	{
 	}
 
-	int32 Max() const { return FrameBuffer.Capacity(); }
-	FString GetBasicDebugStr() const { return FString::Printf(TEXT("PendingTick: %d MaxTick: %d LatestInput: %d NextIntput: %d Size: %d. PT: %s. IT: %s. MT: %s"), 
-		PendingTickFrame, MaxTickFrame, LatestInputFrame, NextInputFrame, Max(), *GetTotalProcessedSimulationTime().ToString(), *GetTotalInputSimulationTime().ToString(), *GetTotalAllowedSimulationTime().ToString()); }
+	int32 Capacity() const { return FrameBuffer.Capacity(); }
+	FString DebugString() const { return FString::Printf(TEXT("PendingTick: %d MaxTick: %d LatestInput: %d NextIntput: %d Size: %d. PT: %s. IT: %s. MT: %s"), 
+		PendingTickFrame, MaxTickFrame, LatestInputFrame, NextInputFrame, Capacity(), *GetTotalProcessedSimulationTime().ToString(), *GetTotalInputSimulationTime().ToString(), *GetTotalAllowedSimulationTime().ToString()); }
 	
 	// ---------------------------------------------------------------------
 	// Frame State
@@ -66,10 +66,10 @@ struct TNetworkedSimulationState
 		}
 		
 		const int32 Delta = NewFrame - PendingTickFrame;
-		if (Delta >= Max())
+		if (Delta >= Capacity())
 		{
 			// PendingTickFrame must always be preserved
-			UE_NP_TRACE_SYSTEM_FAULT("WriteFrameWithTime frame %d outside of buffer range. Reseting contents. %s", PendingTickFrame, *GetBasicDebugStr());
+			UE_NP_TRACE_SYSTEM_FAULT("WriteFrameWithTime frame %d outside of buffer range. Reseting contents. %s", PendingTickFrame, *DebugString());
 				
 			TSimulationFrameState<Model>& NewFrameState = FrameBuffer[NewFrame];
 			NewFrameState = FrameBuffer[PendingTickFrame]; // Init with previous state
@@ -111,10 +111,10 @@ struct TNetworkedSimulationState
 		}
 				
 		const int32 Delta = NewFrame - PendingTickFrame;
-		if (Delta >= Max())
+		if (Delta >= Capacity())
 		{
 			// PendingTickFrame must always be preserved
-			UE_NP_TRACE_SYSTEM_FAULT("WriteFrameSequential frame %d outside of buffer range. Reseting contents. %s", PendingTickFrame, *GetBasicDebugStr());
+			UE_NP_TRACE_SYSTEM_FAULT("WriteFrameSequential frame %d outside of buffer range. Reseting contents. %s", PendingTickFrame, *DebugString());
 
 			TSimulationFrameState<Model>& NewFrameState = FrameBuffer[NewFrame];
 			NewFrameState = FrameBuffer[PendingTickFrame]; // Init with previous state
@@ -168,15 +168,15 @@ struct TNetworkedSimulationState
 		return &FrameBuffer[0];
 	}
 
-	TSimulationFrameState<Model>* GetFrameState(int32 Frame) { return GetValidFrame(Frame); }
-	const TSimulationFrameState<Model>* GetFrameState(int32 Frame) const { return GetValidFrame(Frame); }
+	TSimulationFrameState<Model>* ReadFrame(int32 Frame) { return GetValidFrame(Frame); }
+	const TSimulationFrameState<Model>* ReadFrame(int32 Frame) const { return GetValidFrame(Frame); }
 
 	// ---------------------------------------------------------------------
 	// Aux State
 	// ---------------------------------------------------------------------
 	
-	TAuxState* GetAuxState(int32 Frame) { return AuxBuffer[Frame]; }
-	const TAuxState* GetAuxState(int32 Frame) const { return AuxBuffer[Frame]; }
+	TAuxState* ReadAux(int32 Frame) { return AuxBuffer[Frame]; }
+	const TAuxState* ReadAux(int32 Frame) const { return AuxBuffer[Frame]; }
 
 	TNetSimLazyWriterFunc<TAuxState> GetAuxStateLazyWriter(int32 Frame) { return AuxBuffer.LazyWriter(Frame); }
 	TAuxState* WriteAuxState(int32 Frame) { return AuxBuffer.WriteAtFrame(Frame); }
@@ -259,7 +259,7 @@ struct TNetworkedSimulationState
 
 	void CheckInvariants()
 	{
-		npEnsureMsgf(IsValidFrame(PendingTickFrame), TEXT("PendingTickFrame not valid. %s"), *GetBasicDebugStr());
+		npEnsureMsgf(IsValidFrame(PendingTickFrame), TEXT("PendingTickFrame not valid. %s"), *DebugString());
 		npEnsureMsgf(ConfirmedFrame <= PendingTickFrame, TEXT("ConfirmedFrame %d <= PendingTickFrame"), ConfirmedFrame, PendingTickFrame);
 
 		if (LatestInputFrame != INDEX_NONE)
@@ -267,7 +267,7 @@ struct TNetworkedSimulationState
 			npEnsureMsgf(MaxTickFrame <= LatestInputFrame, TEXT("MaxTickFrame %d <= LatestInputFrame %d"), MaxTickFrame, LatestInputFrame);
 			npEnsureMsgf(PendingTickFrame <= LatestInputFrame+1, TEXT("PendingTickFrame %d <= (LatestInputFrame+1) %d"), PendingTickFrame, LatestInputFrame+1);
 			npEnsureMsgf(LatestInputFrame == NextInputFrame-1, TEXT("LatestInputFrame %d != (NextInputFrame+1) %d"),  LatestInputFrame, NextInputFrame+1);			
-			npEnsureMsgf(IsValidFrame(LatestInputFrame) && FrameBuffer[LatestInputFrame].HasFlag(ESimulationFrameStateFlags::InputWritten), TEXT("LatestInputFrame not valid. %s"), *GetBasicDebugStr());
+			npEnsureMsgf(IsValidFrame(LatestInputFrame) && FrameBuffer[LatestInputFrame].HasFlag(ESimulationFrameStateFlags::InputWritten), TEXT("LatestInputFrame not valid. %s"), *DebugString());
 		}
 
 		if (MaxTickFrame != INDEX_NONE)
@@ -356,7 +356,7 @@ struct TNetSimAccessorHelper
 	{
 		TAccessFunc Func = [NetSimModel](bool bForWrite, TState*& OutState, bool& OutSafe)
 		{
-			OutState = &NetSimModel->State.GetFrameState(NetSimModel->State.GetPendingTickFrame())->SyncState;
+			OutState = &NetSimModel->State.ReadFrame(NetSimModel->State.GetPendingTickFrame())->SyncState;
 			OutSafe = NetSimModel->State.GetTickInProgress();
 		};
 		return Func;
@@ -371,7 +371,7 @@ struct TNetSimAccessorHelper<TNetworkSimModel, TState, false>
 	{
 		TAccessFunc Func = [NetSimModel](bool bForWrite, TState*& OutState, bool& OutSafe)
 		{
-			OutState = NetSimModel->State.GetAuxState(NetSimModel->State.GetPendingTickFrame());
+			OutState = NetSimModel->State.ReadAux(NetSimModel->State.GetPendingTickFrame());
 			OutSafe = NetSimModel->State.GetTickInProgress();
 		};
 		return Func;
@@ -483,8 +483,8 @@ struct TSimulationDoTickFunc
 		const int32 OutputFrame = State.GetPendingTickFrame()+1;
 
 		// Get inputs. These all need to be present to tick (caller's responsibility)
-		TFrameState* InFrameState = State.GetFrameState(InputFrame);
-		const TAuxState* InAuxState = State.GetAuxState(InputFrame);
+		TFrameState* InFrameState = State.ReadFrame(InputFrame);
+		const TAuxState* InAuxState = State.ReadAux(InputFrame);
 		
 		check(InFrameState);
 		check(InAuxState);
@@ -510,7 +510,7 @@ struct TSimulationDoTickFunc
 
 		// Post
 		UE_NP_TRACE_USER_STATE_SYNC(OutFrameState->SyncState, OutputFrame);
-		const TAuxState* AuxHead = State.GetAuxState(OutputFrame);
+		const TAuxState* AuxHead = State.ReadAux(OutputFrame);
 		if (AuxHead != InAuxState)
 		{
 			UE_NP_TRACE_USER_STATE_AUX(*AuxHead, OutputFrame);
diff --git a/Engine/Plugins/Runtime/NetworkPrediction/Source/NetworkPrediction/Public/Trace/NetworkPredictionTrace.h b/Engine/Plugins/Runtime/NetworkPrediction/Source/NetworkPrediction/Public/Trace/NetworkPredictionTrace.h
--- a/Engine/Plugins/Runtime/NetworkPrediction/Source/NetworkPrediction/Public/Trace/NetworkPredictionTrace.h
+++ b/Engine/Plugins/Runtime/NetworkPrediction/Source/NetworkPrediction/Public/Trace/NetworkPredictionTrace.h
@@ -136,7 +136,7 @@ public:
 	template<typename Model>
 	static void TraceSimulationEOF(const TNetworkedSimulationState<Model>& State)
 	{
-		TraceEOF_Internal(State.Max(), State.GetPendingTickFrame(), State.GetLatestInputFrame(), State.GetMaxTickFrame(), State.GetTotalProcessedSimulationTime(), State.GetTotalAllowedSimulationTime());
+		TraceEOF_Internal(State.Capacity(), State.GetPendingTickFrame(), State.GetLatestInputFrame(), State.GetMaxTickFrame(), State.GetTotalProcessedSimulationTime(), State.GetTotalAllowedSimulationTime());
 	}
 
 private:
