NetworkPrediction
-Trace EngineFrameNumber one time per frame instead of all over the place
-Remove some code that wasn't meant to be checked in

#rnx

diff --git a/Engine/Plugins/Runtime/NetworkPrediction/Source/NetworkPrediction/Private/NetworkedSimulationGlobalManager.cpp b/Engine/Plugins/Runtime/NetworkPrediction/Source/NetworkPrediction/Private/NetworkedSimulationGlobalManager.cpp
--- a/Engine/Plugins/Runtime/NetworkPrediction/Source/NetworkPrediction/Private/NetworkedSimulationGlobalManager.cpp
+++ b/Engine/Plugins/Runtime/NetworkPrediction/Source/NetworkPrediction/Private/NetworkedSimulationGlobalManager.cpp
@@ -2,6 +2,7 @@
 
 #include "NetworkedSimulationGlobalManager.h"
 #include "Engine/World.h"
+#include "Trace/NetworkPredictionTrace.h"
 
 DEFINE_LOG_CATEGORY_STATIC(NetworkSimulationGlobalManager, Log, All);
 
@@ -15,6 +16,7 @@ void UNetworkSimulationGlobalManager::Initialize(FSubsystemCollectionBase& Colle
 	UWorld* World = GetWorld();
 	check(World);
 
+	PreTickDispatchHandle = FWorldDelegates::OnWorldTickStart.AddUObject(this, &UNetworkSimulationGlobalManager::OnWorldPreTick);
 	PostTickDispatchHandle = World->OnPostTickDispatch().AddUObject(this, &UNetworkSimulationGlobalManager::ReconcileSimulationsPostNetworkUpdate);
 	PreWorldActorTickHandle = FWorldDelegates::OnWorldPreActorTick.AddUObject(this, &UNetworkSimulationGlobalManager::BeginNewSimulationFrame);
 }
@@ -25,10 +27,21 @@ void UNetworkSimulationGlobalManager::Deinitialize()
 	{
 		World->OnPostTickDispatch().Remove(PostTickDispatchHandle);
 	}
-
+	
+	FWorldDelegates::OnWorldTickStart.Remove(PreTickDispatchHandle);
 	FWorldDelegates::OnWorldPreActorTick.Remove(PreWorldActorTickHandle);
 }
 
+void UNetworkSimulationGlobalManager::OnWorldPreTick(UWorld* InWorld, ELevelTick InLevelTick, float InDeltaSeconds)
+{
+	if (InWorld != GetWorld() || !InWorld->HasBegunPlay())
+	{
+		return;
+	}
+
+	UE_NP_TRACE_WORLD_FRAME_START(InDeltaSeconds);
+}
+
 void UNetworkSimulationGlobalManager::ReconcileSimulationsPostNetworkUpdate()
 {
 	for (auto It = SimulationGroupMap.CreateIterator(); It; ++It)
diff --git a/Engine/Plugins/Runtime/NetworkPrediction/Source/NetworkPrediction/Private/Trace/NetworkPredictionTrace.cpp b/Engine/Plugins/Runtime/NetworkPrediction/Source/NetworkPrediction/Private/Trace/NetworkPredictionTrace.cpp
--- a/Engine/Plugins/Runtime/NetworkPrediction/Source/NetworkPrediction/Private/Trace/NetworkPredictionTrace.cpp
+++ b/Engine/Plugins/Runtime/NetworkPrediction/Source/NetworkPrediction/Private/Trace/NetworkPredictionTrace.cpp
@@ -30,7 +30,6 @@ UE_TRACE_EVENT_END()
 
 // Trace a simulation's local netrole, which can change during the sim's liftetime.
 UE_TRACE_EVENT_BEGIN(NetworkPrediction, SimulationNetRole)
-	UE_TRACE_EVENT_FIELD(uint64, EngineFrameNumber)
 	UE_TRACE_EVENT_FIELD(uint32, SimulationId)
 	UE_TRACE_EVENT_FIELD(uint8, NetRole)
 UE_TRACE_EVENT_END()
@@ -44,7 +43,6 @@ UE_TRACE_EVENT_END()
 // Traces a simulation tick in any context
 UE_TRACE_EVENT_BEGIN(NetworkPrediction, SimulationTick)
 	UE_TRACE_EVENT_FIELD(uint32, SimulationId)
-	UE_TRACE_EVENT_FIELD(uint64, EngineFrameNumber)
 	UE_TRACE_EVENT_FIELD(int32, StartMS)
 	UE_TRACE_EVENT_FIELD(int32, EndMS)
 	UE_TRACE_EVENT_FIELD(int32, OutputFrame)	
@@ -53,7 +51,6 @@ UE_TRACE_EVENT_END()
 
 UE_TRACE_EVENT_BEGIN(NetworkPrediction, SimulationEOF)
 	UE_TRACE_EVENT_FIELD(uint32, SimulationId)
-	UE_TRACE_EVENT_FIELD(uint64, EngineFrameNumber)
 	UE_TRACE_EVENT_FIELD(float, EngineFrameDeltaTime)
 	UE_TRACE_EVENT_FIELD(int32, BufferSize)
 	UE_TRACE_EVENT_FIELD(int32, PendingTickFrame)
@@ -65,7 +62,6 @@ UE_TRACE_EVENT_END()
 
 UE_TRACE_EVENT_BEGIN(NetworkPrediction, NetSerializeRecv)
 	UE_TRACE_EVENT_FIELD(uint32, SimulationId)
-	UE_TRACE_EVENT_FIELD(uint64, EngineFrameNumber)
 	UE_TRACE_EVENT_FIELD(int32, ReceivedSimTimeMS)
 	UE_TRACE_EVENT_FIELD(int32, ReceivedFrame)
 UE_TRACE_EVENT_END()
@@ -73,19 +69,16 @@ UE_TRACE_EVENT_END()
 UE_TRACE_EVENT_BEGIN(NetworkPrediction, InputCmd)
 	UE_TRACE_EVENT_FIELD(uint32, SimulationId)
 	UE_TRACE_EVENT_FIELD(int32, SimulationFrame)
-	UE_TRACE_EVENT_FIELD(uint64, EngineFrameNumber)
 UE_TRACE_EVENT_END()
 
 UE_TRACE_EVENT_BEGIN(NetworkPrediction, SyncState)
 	UE_TRACE_EVENT_FIELD(uint32, SimulationId)
 	UE_TRACE_EVENT_FIELD(int32, SimulationFrame)
-	UE_TRACE_EVENT_FIELD(uint64, EngineFrameNumber)
 UE_TRACE_EVENT_END()
 
 UE_TRACE_EVENT_BEGIN(NetworkPrediction, AuxState)
 	UE_TRACE_EVENT_FIELD(uint32, SimulationId)
 	UE_TRACE_EVENT_FIELD(int32, SimulationFrame)
-	UE_TRACE_EVENT_FIELD(uint64, EngineFrameNumber)
 UE_TRACE_EVENT_END()
 
 UE_TRACE_EVENT_BEGIN(NetworkPrediction, NetSerializeCommit)
@@ -95,7 +88,6 @@ UE_TRACE_EVENT_END()
 // General system fault. Log message is in attachment
 UE_TRACE_EVENT_BEGIN(NetworkPrediction, SystemFault)
 	UE_TRACE_EVENT_FIELD(uint32, SimulationId)
-	UE_TRACE_EVENT_FIELD(uint64, EngineFrameNumber)
 UE_TRACE_EVENT_END()
 
 UE_TRACE_EVENT_BEGIN(NetworkPrediction, OOBStateMod)
@@ -114,6 +106,11 @@ UE_TRACE_EVENT_BEGIN(NetworkPrediction, PieBegin)
 	UE_TRACE_EVENT_FIELD(uint32, DummyData)
 UE_TRACE_EVENT_END()
 
+UE_TRACE_EVENT_BEGIN(NetworkPrediction, WorldFrameStart)
+	UE_TRACE_EVENT_FIELD(uint64, EngineFrameNumber)
+	UE_TRACE_EVENT_FIELD(float, DeltaSeconds)
+UE_TRACE_EVENT_END()
+
 
 // ----------------------------------------------------------------------------
 //	Internal stack to track which simulation is active. 
@@ -197,7 +194,6 @@ void FNetworkPredictionTrace::TraceSimulationCreated(const AActor* OwningActor,
 void FNetworkPredictionTrace::TraceSimulationNetRole(uint32 SimulationId, ENetRole NetRole)
 {
 	UE_TRACE_LOG(NetworkPrediction, SimulationNetRole, NetworkPredictionChannel)
-		<< SimulationNetRole.EngineFrameNumber(GFrameNumber)
 		<< SimulationNetRole.SimulationId(SimulationId)
 		<< SimulationNetRole.NetRole((uint8)NetRole);
 }
@@ -206,7 +202,6 @@ void FNetworkPredictionTrace::TraceSimulationTick(int32 OutputFrame, const FNetw
 {
 	UE_TRACE_LOG(NetworkPrediction, SimulationTick, NetworkPredictionChannel)
 		<< SimulationTick.SimulationId(PeakSimulationIdChecked())
-		<< SimulationTick.EngineFrameNumber(GFrameNumber)
 		<< SimulationTick.StartMS(TimeStep.TotalSimulationTime)
 		<< SimulationTick.EndMS(TimeStep.TotalSimulationTime + TimeStep.StepMS)
 		<< SimulationTick.OutputFrame(OutputFrame);
@@ -218,7 +213,6 @@ void FNetworkPredictionTrace::TraceEOF_Internal(int32 BufferSize, int32 PendingT
 
 	UE_TRACE_LOG(NetworkPrediction, SimulationEOF, NetworkPredictionChannel)
 		<< SimulationEOF.SimulationId(SimulationId)
-		<< SimulationEOF.EngineFrameNumber(GFrameNumber)
 		<< SimulationEOF.EngineFrameDeltaTime(FApp::GetDeltaTime())
 		<< SimulationEOF.BufferSize(BufferSize)
 		<< SimulationEOF.PendingTickFrame(PendingTickFrame)
@@ -261,7 +255,6 @@ void FNetworkPredictionTrace::TraceNetSerializeRecv(const FNetworkSimTime& Recei
 
 	UE_TRACE_LOG(NetworkPrediction, NetSerializeRecv, NetworkPredictionChannel)
 		<< NetSerializeRecv.SimulationId(SimulationId)
-		<< NetSerializeRecv.EngineFrameNumber(GFrameNumber)
 		<< NetSerializeRecv.ReceivedSimTimeMS(ReceivedTime)
 		<< NetSerializeRecv.ReceivedFrame(ReceivedFrame);
 }
@@ -310,7 +303,6 @@ void FNetworkPredictionTrace::TraceUserState_Internal(int32 Frame, ETraceUserSta
 			UE_TRACE_LOG(NetworkPrediction, InputCmd, NetworkPredictionChannel, AttachmentSize)
 				<< InputCmd.SimulationId(SimulationId)
 				<< InputCmd.SimulationFrame(Frame)
-				<< InputCmd.EngineFrameNumber(GFrameNumber)
 				<< InputCmd.Attachment(*StrOut, AttachmentSize);
 			break;
 		}
@@ -319,7 +311,6 @@ void FNetworkPredictionTrace::TraceUserState_Internal(int32 Frame, ETraceUserSta
 			UE_TRACE_LOG(NetworkPrediction, SyncState, NetworkPredictionChannel, AttachmentSize)
 				<< SyncState.SimulationId(SimulationId)
 				<< SyncState.SimulationFrame(Frame)
-				<< SyncState.EngineFrameNumber(GFrameNumber)
 				<< SyncState.Attachment(*StrOut, AttachmentSize);
 			break;
 		}
@@ -328,7 +319,6 @@ void FNetworkPredictionTrace::TraceUserState_Internal(int32 Frame, ETraceUserSta
 			UE_TRACE_LOG(NetworkPrediction, AuxState, NetworkPredictionChannel, AttachmentSize)
 				<< AuxState.SimulationId(SimulationId)
 				<< AuxState.SimulationFrame(Frame)
-				<< AuxState.EngineFrameNumber(GFrameNumber)
 				<< AuxState.Attachment(*StrOut, AttachmentSize);
 			break;
 		}
@@ -341,6 +331,13 @@ void FNetworkPredictionTrace::TracePIEStart()
 		<< PieBegin.DummyData(0); // temp to quiet clang
 }
 
+void FNetworkPredictionTrace::TraceWorldFrameStart(float DeltaSeconds)
+{
+	UE_TRACE_LOG(NetworkPrediction, WorldFrameStart, NetworkPredictionChannel)
+		<< WorldFrameStart.EngineFrameNumber(GFrameNumber)
+		<< WorldFrameStart.DeltaSeconds(DeltaSeconds);
+}
+
 #include "CoreTypes.h"
 #include "Misc/VarArgs.h"
 #include "HAL/UnrealMemory.h"
@@ -390,7 +387,6 @@ void FNetworkPredictionTrace::TraceSystemFault(const TCHAR* Fmt, ...)
 
 		UE_TRACE_LOG(NetworkPrediction, SystemFault, NetworkPredictionChannel, AttachmentSize)
 			<< SystemFault.SimulationId(SimulationId)
-			<< SystemFault.EngineFrameNumber(GFrameNumber)
 			<< SystemFault.Attachment(Buffer, AttachmentSize);
 
 		UE_LOG(LogNetworkSim, Warning, TEXT("SystemFault: %s"), Buffer);
diff --git a/Engine/Plugins/Runtime/NetworkPrediction/Source/NetworkPrediction/Public/NetworkedSimulationGlobalManager.h b/Engine/Plugins/Runtime/NetworkPrediction/Source/NetworkPrediction/Public/NetworkedSimulationGlobalManager.h
--- a/Engine/Plugins/Runtime/NetworkPrediction/Source/NetworkPrediction/Public/NetworkedSimulationGlobalManager.h
+++ b/Engine/Plugins/Runtime/NetworkPrediction/Source/NetworkPrediction/Public/NetworkedSimulationGlobalManager.h
@@ -61,6 +61,7 @@ public:
 
 private:
 
+	void OnWorldPreTick(UWorld* InWorld, ELevelTick InLevelTick, float InDeltaSeconds);
 	void ReconcileSimulationsPostNetworkUpdate();
 	void BeginNewSimulationFrame(UWorld* InWorld, ELevelTick InLevelTick, float InDeltaSeconds);
 	
@@ -82,6 +83,7 @@ private:
 
 	TSortedMap<FName, FSimulationGroup, FDefaultAllocator, FNameFastLess> SimulationGroupMap;
 
+	FDelegateHandle PreTickDispatchHandle;
 	FDelegateHandle PostTickDispatchHandle;
 	FDelegateHandle PreWorldActorTickHandle;
 };
\ No newline at end of file
diff --git a/Engine/Plugins/Runtime/NetworkPrediction/Source/NetworkPrediction/Public/NetworkedSimulationModelCVars.h b/Engine/Plugins/Runtime/NetworkPrediction/Source/NetworkPrediction/Public/NetworkedSimulationModelCVars.h
--- a/Engine/Plugins/Runtime/NetworkPrediction/Source/NetworkPrediction/Public/NetworkedSimulationModelCVars.h
+++ b/Engine/Plugins/Runtime/NetworkPrediction/Source/NetworkPrediction/Public/NetworkedSimulationModelCVars.h
@@ -13,7 +13,6 @@
 
 inline IConsoleVariable* FindConsoleVarHelper(const TCHAR* VarName)
 {
-	UE_LOG(LogTemp, Warning, TEXT("RATTI Finding... %s"), VarName);
 	return IConsoleManager::Get().FindConsoleVariable(VarName, false);
 }
 
@@ -26,20 +25,16 @@ struct FConditionalAutoConsoleRegister
 {
 	FConditionalAutoConsoleRegister(const TCHAR* VarName, int32 Value, const TCHAR* Help)
 	{
-		UE_LOG(LogTemp, Warning, TEXT("RATTI FConditionalAutoConsoleRegister() %s"), VarName);
 		if (!IConsoleManager::Get().FindConsoleVariable(VarName, false))
 		{
-			UE_LOG(LogTemp, Warning, TEXT("RATTI RegisterConsoleVariable() %s"), VarName);
 			IConsoleManager::Get().RegisterConsoleVariable(VarName, Value, Help, ECVF_Cheat);
 		}
 	}
 
 	FConditionalAutoConsoleRegister(const TCHAR* VarName, float Value, const TCHAR* Help)
 	{
-		UE_LOG(LogTemp, Warning, TEXT("RATTI FConditionalAutoConsoleRegister() %s"), VarName);
 		if (!IConsoleManager::Get().FindConsoleVariable(VarName, false))
 		{
-			UE_LOG(LogTemp, Warning, TEXT("RATTI RegisterConsoleVariable() %s"), VarName);
 			IConsoleManager::Get().RegisterConsoleVariable(VarName, Value, Help, ECVF_Cheat);
 		}
 	}
diff --git a/Engine/Plugins/Runtime/NetworkPrediction/Source/NetworkPrediction/Public/NetworkedSimulationModelRepControllers.h b/Engine/Plugins/Runtime/NetworkPrediction/Source/NetworkPrediction/Public/NetworkedSimulationModelRepControllers.h
--- a/Engine/Plugins/Runtime/NetworkPrediction/Source/NetworkPrediction/Public/NetworkedSimulationModelRepControllers.h
+++ b/Engine/Plugins/Runtime/NetworkPrediction/Source/NetworkPrediction/Public/NetworkedSimulationModelRepControllers.h
@@ -255,9 +255,6 @@ struct TRepController_Server : public TBase
 	{
 		NETSIM_CHECKSUM(P.Ar);
 
-		TSimulationFrameState<Model> Test;
-		
-
 		FArchive& Ar = P.Ar;
 
 		// Serialize Input Cmds redundantly. This could be better:
diff --git a/Engine/Plugins/Runtime/NetworkPrediction/Source/NetworkPrediction/Public/Trace/NetworkPredictionTrace.h b/Engine/Plugins/Runtime/NetworkPrediction/Source/NetworkPrediction/Public/Trace/NetworkPredictionTrace.h
--- a/Engine/Plugins/Runtime/NetworkPrediction/Source/NetworkPrediction/Public/Trace/NetworkPredictionTrace.h
+++ b/Engine/Plugins/Runtime/NetworkPrediction/Source/NetworkPrediction/Public/Trace/NetworkPredictionTrace.h
@@ -59,6 +59,8 @@
 
 #define UE_NP_TRACE_SYSTEM_FAULT(Format, ...) FNetworkPredictionTrace::TraceSystemFault(TEXT(Format), ##__VA_ARGS__)
 
+#define UE_NP_TRACE_WORLD_FRAME_START(DeltaSeconds) FNetworkPredictionTrace::TraceWorldFrameStart(DeltaSeconds)
+
 #else
 
 // Compiled out
@@ -78,6 +80,7 @@
 #define UE_NP_TRACE_OOB_STATE_MOD(...)
 #define UE_NP_TRACE_PIE_START(...)
 #define UE_NP_TRACE_SYSTEM_FAULT(Format, ...) UE_LOG(LogNetworkSim, Warning, TEXT(Format), ##__VA_ARGS__);
+#define UE_NP_TRACE_WORLD_FRAME_START(...)
 
 #endif // UE_NP_TRACE_ENABLED
 
@@ -109,6 +112,8 @@ public:
 	static void TracePIEStart();
 	static void TraceSystemFault(const TCHAR* Fmt, ...);
 
+	static void TraceWorldFrameStart(float DeltaSeconds);
+
 	enum ETraceUserState
 	{
 		Input,
diff --git a/Engine/Plugins/Runtime/NetworkPrediction/readme.txt b/Engine/Plugins/Runtime/NetworkPrediction/readme.txt
--- a/Engine/Plugins/Runtime/NetworkPrediction/readme.txt
+++ b/Engine/Plugins/Runtime/NetworkPrediction/readme.txt
@@ -105,7 +105,7 @@ Simulation: should always refer to use code. The simulation is what we are netwo
 System: should always refer to the TNetworkedSimulationModel code: the generic code that is doing the networking. It is distinct from the simulation.
 
 Reconcile: When a predicting client receives authoritative data from the server that conflicts with what they previously predicted: sorting this out is called "reconciling". Usually means you will resimulate.
-Synthesize: In general means "making shit up". For example, "synthesizing user commands" would mean "creating user commands that were not generated by a player, but instead guessing what it might be".
+Synthesize: In general means "making stuff up". For example, "synthesizing user commands" would mean "creating user commands that were not generated by a player, but instead guessing what it might be".
  
  
 Client Prediction: using new, locally created information (input cmd) to evolve the game state ahead of the authoritative server.
diff --git a/Engine/Plugins/Runtime/NetworkPredictionInsights/Source/NetworkPredictionInsights/Private/NetworkPredictionAnalyzer.cpp b/Engine/Plugins/Runtime/NetworkPredictionInsights/Source/NetworkPredictionInsights/Private/NetworkPredictionAnalyzer.cpp
--- a/Engine/Plugins/Runtime/NetworkPredictionInsights/Source/NetworkPredictionInsights/Private/NetworkPredictionAnalyzer.cpp
+++ b/Engine/Plugins/Runtime/NetworkPredictionInsights/Source/NetworkPredictionInsights/Private/NetworkPredictionAnalyzer.cpp
@@ -19,6 +19,7 @@ void FNetworkPredictionAnalyzer::OnAnalysisBegin(const FOnAnalysisContext& Conte
 	Builder.RouteEvent(RouteId_SimulationNetRole, "NetworkPrediction", "SimulationNetRole");
 	Builder.RouteEvent(RouteId_SimulationNetGUID, "NetworkPrediction", "SimulationNetGUID");	
 	Builder.RouteEvent(RouteId_SimulationTick, "NetworkPrediction", "SimulationTick");
+	Builder.RouteEvent(RouteId_WorldFrameStart, "NetworkPrediction", "WorldFrameStart");
 	Builder.RouteEvent(RouteId_OOBStateMod, "NetworkPrediction", "OOBStateMod");
 	Builder.RouteEvent(RouteId_ProduceInput, "NetworkPrediction", "ProduceInput");
 	Builder.RouteEvent(RouteId_SynthInput, "NetworkPrediction", "SynthInput");
@@ -45,7 +46,7 @@ bool FNetworkPredictionAnalyzer::OnEvent(uint16 RouteId, const FOnEventContext&
 	{
 		int32 SimulationId = EventData.GetValue<uint32>("SimulationId");			 
 		int32 Frame = EventData.GetValue<int32>("SimulationFrame");
-		uint64 EngineFrame = EventData.GetValue<uint64>("EngineFrameNumber");
+		uint64 EngineFrame = this->EngineFrameNumber;
 
 		NetworkPredictionProvider.WriteUserState(SimulationId, Frame, EngineFrame, Type, Session.StoreString(reinterpret_cast<const TCHAR*>(EventData.GetAttachment())));
 	};
@@ -64,6 +65,14 @@ bool FNetworkPredictionAnalyzer::OnEvent(uint16 RouteId, const FOnEventContext&
 			// EventData.GetValue<bool>("IsServer")
 			break;
 		}
+
+		case RouteId_WorldFrameStart:
+		{
+			EngineFrameNumber = EventData.GetValue<uint64>("EngineFrameNumber");
+			DeltaTimeSeconds = EventData.GetValue<float>("DeltaTimeSeconds");
+			break;
+		}
+
 		case RouteId_SimulationCreated:
 		{
 			FSimulationData::FConst& ConstData = NetworkPredictionProvider.WriteSimulationCreated(EventData.GetValue<uint32>("SimulationId"));
@@ -75,7 +84,7 @@ bool FNetworkPredictionAnalyzer::OnEvent(uint16 RouteId, const FOnEventContext&
 		case RouteId_SimulationNetRole:
 		{
 			NetworkPredictionProvider.WriteSimulationNetRole(	EventData.GetValue<uint32>("SimulationId"),
-																EventData.GetValue<uint64>("EngineFrameNumber"),
+																EngineFrameNumber,
 																(ENP_NetRole)EventData.GetValue<uint8>("NetRole"));
 			break;
 		}
@@ -89,7 +98,7 @@ bool FNetworkPredictionAnalyzer::OnEvent(uint16 RouteId, const FOnEventContext&
 		case RouteId_SimulationTick:
 		{
 			FSimulationData::FTick TickData;
-			TickData.EngineFrame = EventData.GetValue<uint64>("EngineFrameNumber");
+			TickData.EngineFrame = EngineFrameNumber;
 			TickData.StartMS = EventData.GetValue<int32>("StartMS");
 			TickData.EndMS = EventData.GetValue<int32>("EndMS");
 			TickData.OutputFrame = EventData.GetValue<int32>("OutputFrame");
@@ -100,7 +109,7 @@ bool FNetworkPredictionAnalyzer::OnEvent(uint16 RouteId, const FOnEventContext&
 		case RouteId_SimulationEOF:
 		{
 			FSimulationData::FEngineFrame& FrameData = NetworkPredictionProvider.WriteSimulationEOF(EventData.GetValue<uint32>("SimulationId"));
-			FrameData.EngineFrame = EventData.GetValue<uint64>("EngineFrameNumber");
+			FrameData.EngineFrame = EngineFrameNumber;
 			FrameData.EngineFrameDeltaTime = EventData.GetValue<float>("EngineFrameDeltaTime");
 			
 
@@ -116,7 +125,7 @@ bool FNetworkPredictionAnalyzer::OnEvent(uint16 RouteId, const FOnEventContext&
 		case RouteId_NetSerializeRecv:
 		{
 			FSimulationData::FNetSerializeRecv NetRecv;
-			NetRecv.EngineFrame = EventData.GetValue<uint64>("EngineFrameNumber");
+			NetRecv.EngineFrame = EngineFrameNumber;
 			NetRecv.SimTimeMS = EventData.GetValue<int32>("ReceivedSimTimeMS");
 			NetRecv.Frame = EventData.GetValue<int32>("ReceivedFrame");
 			
@@ -176,7 +185,7 @@ bool FNetworkPredictionAnalyzer::OnEvent(uint16 RouteId, const FOnEventContext&
 		{
 			NetworkPredictionProvider.WriteSystemFault(
 				EventData.GetValue<uint32>("SimulationId"),
-				EventData.GetValue<uint64>("EngineFrameNumber"),
+				EngineFrameNumber,
 				Session.StoreString(reinterpret_cast<const TCHAR*>(EventData.GetAttachment()))
 			);
 			break;
diff --git a/Engine/Plugins/Runtime/NetworkPredictionInsights/Source/NetworkPredictionInsights/Private/NetworkPredictionAnalyzer.h b/Engine/Plugins/Runtime/NetworkPredictionInsights/Source/NetworkPredictionInsights/Private/NetworkPredictionAnalyzer.h
--- a/Engine/Plugins/Runtime/NetworkPredictionInsights/Source/NetworkPredictionInsights/Private/NetworkPredictionAnalyzer.h
+++ b/Engine/Plugins/Runtime/NetworkPredictionInsights/Source/NetworkPredictionInsights/Private/NetworkPredictionAnalyzer.h
@@ -30,6 +30,7 @@ private:
 	enum : uint16
 	{
 		RouteId_GameInstanceRegister,
+		RouteId_WorldFrameStart,
 		RouteId_SimulationCreated,
 		RouteId_SimulationNetRole,
 		RouteId_SimulationNetGUID,
@@ -50,4 +51,7 @@ private:
 
 	Trace::IAnalysisSession& Session;
 	FNetworkPredictionProvider& NetworkPredictionProvider;
+
+	uint64 EngineFrameNumber;
+	float DeltaTimeSeconds;
 };
\ No newline at end of file
