Network Prediction
-Fix ensure with automation based AP clients
-Better support for OOB mod tracing: aux state changes are now properly traced. Optonal "User String" is now supported along with the mod, to track who made the change.
-Search highlight border size increasd to make it a little more visible.

#rnx

diff --git a/Engine/Plugins/Runtime/NetworkPrediction/Source/NetworkPrediction/Private/Trace/NetworkPredictionTrace.cpp b/Engine/Plugins/Runtime/NetworkPrediction/Source/NetworkPrediction/Private/Trace/NetworkPredictionTrace.cpp
--- a/Engine/Plugins/Runtime/NetworkPrediction/Source/NetworkPrediction/Private/Trace/NetworkPredictionTrace.cpp
+++ b/Engine/Plugins/Runtime/NetworkPrediction/Source/NetworkPrediction/Private/Trace/NetworkPredictionTrace.cpp
@@ -111,6 +111,15 @@ UE_TRACE_EVENT_BEGIN(NetworkPrediction, WorldFrameStart)
 	UE_TRACE_EVENT_FIELD(float, DeltaSeconds)
 UE_TRACE_EVENT_END()
 
+UE_TRACE_EVENT_BEGIN(NetworkPrediction, OOBStateModStrSync)
+	UE_TRACE_EVENT_FIELD(uint32, SimulationId)
+UE_TRACE_EVENT_END()
+
+UE_TRACE_EVENT_BEGIN(NetworkPrediction, OOBStateModStrAux)
+	UE_TRACE_EVENT_FIELD(uint32, SimulationId)
+UE_TRACE_EVENT_END()
+
+
 
 // ----------------------------------------------------------------------------
 //	Internal stack to track which simulation is active. 
@@ -283,10 +292,8 @@ void FNetworkPredictionTrace::TraceSynthInput()
 		<< SynthInput.SimulationId(SimulationId);
 }
 
-void FNetworkPredictionTrace::TraceOOBStateMod()
+void FNetworkPredictionTrace::TraceOOBStateMod(int32 SimulationId)
 {
-	const uint32 SimulationId = PeakSimulationIdChecked();
-
 	UE_TRACE_LOG(NetworkPrediction, OOBStateMod, NetworkPredictionChannel)
 		<< OOBStateMod.SimulationId(SimulationId);
 }
@@ -338,6 +345,33 @@ void FNetworkPredictionTrace::TraceWorldFrameStart(float DeltaSeconds)
 		<< WorldFrameStart.DeltaSeconds(DeltaSeconds);
 }
 
+void FNetworkPredictionTrace::TraceOOBStrInternal(int32 SimulationId, ETraceUserState StateType, const TCHAR* Fmt)
+{
+	const uint16 AttachmentSize = (FCString::Strlen(Fmt)+1) * sizeof(TCHAR);
+	switch(StateType)
+	{
+		case ETraceUserState::Input:
+		{
+			ensure(false); // unexpected. We don't do OOB mods to input state
+			break;
+		}
+		case ETraceUserState::Sync:
+		{
+			UE_TRACE_LOG(NetworkPrediction, OOBStateModStrSync, NetworkPredictionChannel, AttachmentSize)
+				<< OOBStateModStrSync.SimulationId(SimulationId)	
+				<< OOBStateModStrSync.Attachment(Fmt, AttachmentSize);
+			break;
+		}
+		case ETraceUserState::Aux:
+		{
+			UE_TRACE_LOG(NetworkPrediction, OOBStateModStrAux, NetworkPredictionChannel, AttachmentSize)
+				<< OOBStateModStrAux.SimulationId(SimulationId)	
+				<< OOBStateModStrAux.Attachment(Fmt, AttachmentSize);
+			break;
+		}
+	}
+}
+
 #include "CoreTypes.h"
 #include "Misc/VarArgs.h"
 #include "HAL/UnrealMemory.h"
diff --git a/Engine/Plugins/Runtime/NetworkPrediction/Source/NetworkPrediction/Public/NetworkedSimulationModel.h b/Engine/Plugins/Runtime/NetworkPrediction/Source/NetworkPrediction/Public/NetworkedSimulationModel.h
--- a/Engine/Plugins/Runtime/NetworkPrediction/Source/NetworkPrediction/Public/NetworkedSimulationModel.h
+++ b/Engine/Plugins/Runtime/NetworkPrediction/Source/NetworkPrediction/Public/NetworkedSimulationModel.h
@@ -127,7 +127,7 @@
 	{
 		check(Simulation.IsValid() && InDriver);
 		Driver = InDriver;
-		State.WriteFrameInitial()->SyncState = InitialSyncState;
+		State.WriteFramePending()->SyncState = InitialSyncState;
 		*State.WriteAuxState(0) = InitialAuxState;
 		DO_NETSIM_MODEL_DEBUG(FNetworkSimulationModelDebuggerManager::Get().RegisterNetworkSimulationModel(this, Driver->GetVLogOwner()));
 
@@ -141,9 +141,12 @@
 		SimulationId = this->GetNextSimulationId();
 		UE_NP_TRACE_SET_SCOPE_SIM(SimulationId);
 		UE_NP_TRACE_SIM_CREATED(Driver->GetVLogOwner(), GetSimulationGroupName());
-		UE_NP_TRACE_OOB_STATE_MOD();
+		UE_NP_TRACE_OOB_STATE_MOD(SimulationId);
 		UE_NP_TRACE_USER_STATE_SYNC(InitialSyncState, 0);
 		UE_NP_TRACE_USER_STATE_AUX(InitialAuxState, 0);
+
+		bPendingOOBAuxWrite = false;
+		bPendingOOBSyncWrite = false;
 	}
 
 	virtual ~TNetworkedSimulationModel()
@@ -152,31 +155,34 @@
 		ClearAllDependentSimulations();
 	}
 
-	void Tick(const FNetSimTickParameters& Parameters) final override
+	void NetSerializeProxy(EReplicationProxyTarget Target, const FNetSerializeParams& Params) final override
 	{
 		UE_NP_TRACE_SET_SCOPE_SIM(SimulationId);
-		
-		switch (Parameters.Role)
+		FlushPendingOOBWrites();
+
+		switch(Target)
 		{
-			case ROLE_Authority:
-				RepProxy_ServerRPC.Tick(Driver, Simulation.Get(), State, Parameters);
+		case EReplicationProxyTarget::ServerRPC:
+			RepProxy_ServerRPC.NetSerialize(Params, State);
 			break;
-
-			case ROLE_AutonomousProxy:
-				RepProxy_Autonomous.Tick(Driver, Simulation.Get(), State, Parameters);
+		case EReplicationProxyTarget::AutonomousProxy:
+			RepProxy_Autonomous.NetSerialize(Params, State);
 			break;
-
-			case ROLE_SimulatedProxy:
-				RepProxy_Simulated.Tick(Driver, Simulation.Get(), State, Parameters);
+		case EReplicationProxyTarget::SimulatedProxy:
+			RepProxy_Simulated.NetSerialize(Params, State);
 			break;
-		}		
-
-		UE_NP_TRACE_SIM_EOF(State);
+		case EReplicationProxyTarget::Replay:
+			RepProxy_Replay.NetSerialize(Params, State);
+			break;
+		default:
+			checkf(false, TEXT("Unknown: %d"), (int32)Target);
+		};
 	}
 
 	virtual void Reconcile(const ENetRole Role) final override
 	{
 		UE_NP_TRACE_SET_SCOPE_SIM(SimulationId);
+		FlushPendingOOBWrites();
 
 		// --------------------------------------------------------------------------------------------------------------------------
 		//	Reconcile
@@ -198,6 +204,29 @@
 			break;
 		}
 	}
+
+	void Tick(const FNetSimTickParameters& Parameters) final override
+	{
+		UE_NP_TRACE_SET_SCOPE_SIM(SimulationId);
+		FlushPendingOOBWrites();
+		
+		switch (Parameters.Role)
+		{
+			case ROLE_Authority:
+				RepProxy_ServerRPC.Tick(Driver, Simulation.Get(), State, Parameters);
+			break;
+
+			case ROLE_AutonomousProxy:
+				RepProxy_Autonomous.Tick(Driver, Simulation.Get(), State, Parameters);
+			break;
+
+			case ROLE_SimulatedProxy:
+				RepProxy_Simulated.Tick(Driver, Simulation.Get(), State, Parameters);
+			break;
+		}		
+
+		UE_NP_TRACE_SIM_EOF(State);
+	}
 	
 	void InitializeForNetworkRole(const ENetRole Role, const FNetworkSimulationModelInitParameters& Parameters) final override
 	{
@@ -208,30 +237,7 @@
 		//Buffers.Input.SetBufferSize(Parameters.InputBufferSize);
 		//Buffers.Sync.SetBufferSize(Parameters.SyncedBufferSize);
 		//Buffers.Aux.SetBufferSize(Parameters.AuxBufferSize);
-	}
+	}	
 
-	void NetSerializeProxy(EReplicationProxyTarget Target, const FNetSerializeParams& Params) final override
-	{
-		UE_NP_TRACE_SET_SCOPE_SIM(SimulationId);
-
-		switch(Target)
-		{
-		case EReplicationProxyTarget::ServerRPC:
-			RepProxy_ServerRPC.NetSerialize(Params, State);
-			break;
-		case EReplicationProxyTarget::AutonomousProxy:
-			RepProxy_Autonomous.NetSerialize(Params, State);
-			break;
-		case EReplicationProxyTarget::SimulatedProxy:
-			RepProxy_Simulated.NetSerialize(Params, State);
-			break;
-		case EReplicationProxyTarget::Replay:
-			RepProxy_Replay.NetSerialize(Params, State);
-			break;
-		default:
-			checkf(false, TEXT("Unknown: %d"), (int32)Target);
-		};
-	}
-
 	int32 GetProxyDirtyCount(EReplicationProxyTarget Target) final override
 	{
@@ -251,5 +257,48 @@
 	}
 
 	// ------------------------------------------------------------------------------------------------------
+	//	OOB State access
+	//		This is for outside game code to read/write to the sync&aux states.
+	//		(Input should only be written via the ProduceInput driver callback)
+	//		These accessors always operate on the PendingTickFrame's state: the 
+	//		next frame that will be processed in a ::SimulationTick.
+	// ------------------------------------------------------------------------------------------------------
+
+	const TAuxState* ReadAuxState() const
+	{
+		const TAuxState* PendingAux = State.ReadAux(State.GetPendingTickFrame());
+		npCheckSlow(PendingAux);
+		return PendingAux; 
+	}
+	const TSyncState* ReadSyncState() const
+	{
+		const TFrameState* PendingFrameState = State.ReadFrame(State.GetPendingTickFrame());
+		npCheckSlow(PendingFrameState);
+		return &PendingFrameState->SyncState; 
+	}
+
+	// Returns Aux State for writing. TraceStr is an optional debug message that will traced.
+	TAuxState* WriteAuxState(const TCHAR* TraceStr=nullptr)
+	{
+		FlushPendingOOBWriteAux();
+		TAuxState* PendingAux = State.WriteAuxState(State.GetPendingTickFrame());
+		npCheckSlow(PendingAux);
+		bPendingOOBAuxWrite = true;
+		UE_NP_TRACE_OOB_STR_AUX(SimulationId, TraceStr);
+		return PendingAux; 
+	}
+
+	// Returns Sync State for writing. TraceStr is an optional debug message that will traced.
+	TSyncState* WriteSyncState(const TCHAR* TraceStr=nullptr)
+	{
+		FlushPendingOOBWriteSync();
+		TFrameState* PendingFrameState = State.WriteFramePending();
+		npCheckSlow(PendingFrameState);
+		bPendingOOBSyncWrite = true;
+		UE_NP_TRACE_OOB_STR_SYNC(SimulationId, TraceStr);
+		return &PendingFrameState->SyncState; 
+	}
+
+	// ------------------------------------------------------------------------------------------------------
 	//	Dependent Simulations
 	// ------------------------------------------------------------------------------------------------------
@@ -351,7 +400,38 @@
 	// Unique, global, per process Id of simulation. Used for trace debugging purposes only. Not replicated.
 	uint32 GetSimulationId() const final override { return SimulationId; }
 
-private:
+private:	
+
+	void FlushPendingOOBWrites()
+	{
+		FlushPendingOOBWriteSync();
+		FlushPendingOOBWriteAux();
+	}
+
+	void FlushPendingOOBWriteSync()
+	{
+		if (bPendingOOBSyncWrite)
+		{
+			UE_NP_TRACE_OOB_STATE_MOD(SimulationId);
+			UE_NP_TRACE_USER_STATE_SYNC(*ReadSyncState(), State.GetPendingTickFrame());
+			bPendingOOBSyncWrite = false;
+		}
+	}
+
+	void FlushPendingOOBWriteAux()
+	{
+		if (bPendingOOBAuxWrite)
+		{
+			UE_NP_TRACE_OOB_STATE_MOD(SimulationId);
+			UE_NP_TRACE_USER_STATE_AUX(*ReadAuxState(), State.GetPendingTickFrame());
+			bPendingOOBAuxWrite = false;
+		}
+		
+	}
+
+	bool bPendingOOBAuxWrite:1;
+	bool bPendingOOBSyncWrite:1;
+
 	float ServerRPCAccumulatedTimeSeconds = 0.f;
 	float ServerRPCThresholdTimeSeconds = 1.f / 999.f; // Default is to send at a max of 999hz. This part of the system needs to be build out more (better handling of super high FPS clients and fixed rate servers)
 
diff --git a/Engine/Plugins/Runtime/NetworkPrediction/Source/NetworkPrediction/Public/NetworkedSimulationModelRepControllers.h b/Engine/Plugins/Runtime/NetworkPrediction/Source/NetworkPrediction/Public/NetworkedSimulationModelRepControllers.h
--- a/Engine/Plugins/Runtime/NetworkPrediction/Source/NetworkPrediction/Public/NetworkedSimulationModelRepControllers.h
+++ b/Engine/Plugins/Runtime/NetworkPrediction/Source/NetworkPrediction/Public/NetworkedSimulationModelRepControllers.h
@@ -350,7 +350,7 @@ struct TRepController_Server : public TBase
 				else
 				{
 					// Brand new frame
-					auto* NewFrameState = State.WriteFrameSequential(Frame);
+					TFrameState* NewFrameState = State.WriteFrameSequential(Frame);
 					
 					NewFrameState->FrameDeltaTime.NetSerialize(P.Ar);
 					NewFrameState->InputCmd.NetSerialize(P);
@@ -481,7 +481,14 @@ struct TRepController_Autonomous: public TBase
 				bPendingReconciliation =  true;
 			}
 
-			State.SetConfirmedFrame(SerializedFrame);
+			if (State.GetPendingTickFrame() < SerializedFrame)
+			{
+				State.Reset(SerializedFrame);
+			}
+			else
+			{
+				State.SetConfirmedFrame(SerializedFrame);
+			}
 		}
 		NETSIM_CHECKSUM(P.Ar);
 	}
diff --git a/Engine/Plugins/Runtime/NetworkPrediction/Source/NetworkPrediction/Public/NetworkedSimulationModelTypes.h b/Engine/Plugins/Runtime/NetworkPrediction/Source/NetworkPrediction/Public/NetworkedSimulationModelTypes.h
--- a/Engine/Plugins/Runtime/NetworkPrediction/Source/NetworkPrediction/Public/NetworkedSimulationModelTypes.h
+++ b/Engine/Plugins/Runtime/NetworkPrediction/Source/NetworkPrediction/Public/NetworkedSimulationModelTypes.h
@@ -163,9 +163,9 @@ struct TNetworkedSimulationState
 		return &NewFrameState;
 	}
 
-	TSimulationFrameState<Model>* WriteFrameInitial()
+	TSimulationFrameState<Model>* WriteFramePending()
 	{
-		return &FrameBuffer[0];
+		return &FrameBuffer[PendingTickFrame];
 	}
 
 	TSimulationFrameState<Model>* ReadFrame(int32 Frame) { return GetValidFrame(Frame); }
@@ -347,16 +347,25 @@ private:
 //	you may roll back the predicted change to then later receive the authority change.
 // ----------------------------------------------------------------------------------------------------------------------------------------------
 
+// Helper to specialize on Sync vs Aux state access
 template<typename TNetworkSimModel, typename TState, bool IsSyncState=TIsDerivedFrom<typename TNetworkSimModel::TSyncState, TState>::Value>
 struct TNetSimAccessorHelper
 {
-	using TAccessFunc = TUniqueFunction<void(bool bForWrite, TState*& OutState, bool& OutSafe)>;
-	
+	// Sync State
+	using TAccessFunc = TUniqueFunction<void(bool bForWrite, TState*& OutState, bool& OutSafe, const TCHAR* TraceStr)>;
 	static TAccessFunc GetStateFunc(TNetworkSimModel* NetSimModel)
 	{
-		TAccessFunc Func = [NetSimModel](bool bForWrite, TState*& OutState, bool& OutSafe)
+		TAccessFunc Func = [NetSimModel](bool bForWrite, TState*& OutState, bool& OutSafe, const TCHAR* TraceStr)
 		{
-			OutState = &NetSimModel->State.ReadFrame(NetSimModel->State.GetPendingTickFrame())->SyncState;
+			if (bForWrite)
+			{
+				OutState = NetSimModel->WriteSyncState(TraceStr);
+			}
+			else
+			{
+				OutState = const_cast<TState*>((TState*)NetSimModel->ReadSyncState());
+			}
+			
 			OutSafe = NetSimModel->State.GetTickInProgress();
 		};
 		return Func;
@@ -366,19 +375,26 @@ struct TNetSimAccessorHelper
 template<typename TNetworkSimModel, typename TState>
 struct TNetSimAccessorHelper<TNetworkSimModel, TState, false>
 {
-	using TAccessFunc = TUniqueFunction<void(bool bForWrite, TState*& OutState, bool& OutSafe)>;
+	// Aux State
+	using TAccessFunc = TUniqueFunction<void(bool bForWrite, TState*& OutState, bool& OutSafe, const TCHAR* TraceStr)>;
 	static TAccessFunc GetStateFunc(TNetworkSimModel* NetSimModel)
 	{
-		TAccessFunc Func = [NetSimModel](bool bForWrite, TState*& OutState, bool& OutSafe)
+		TAccessFunc Func = [NetSimModel](bool bForWrite, TState*& OutState, bool& OutSafe, const TCHAR* TraceStr)
 		{
-			OutState = NetSimModel->State.ReadAux(NetSimModel->State.GetPendingTickFrame());
+			if (bForWrite)
+			{
+				OutState = NetSimModel->WriteAuxState(TraceStr);
+			}
+			else
+			{
+				OutState = const_cast<TState*>((TState*)NetSimModel->ReadAuxState());
+			}
 			OutSafe = NetSimModel->State.GetTickInProgress();
 		};
 		return Func;
 	}
 };
 
-
 template<typename TState>
 struct TNetSimStateAccessor
 {
@@ -402,28 +418,29 @@ struct TNetSimStateAccessor
 		bool bSafe = false;
 		if (GetStateFunc)
 		{
-			GetStateFunc(false, State, bSafe);
+			GetStateFunc(false, State, bSafe, nullptr);
 		}
 		return State;
 	}
 
 	/** Gets the current (PendingFrame) state for writing. This is expected to fail outside of the core update loop when bHasAuthority=false. (E.g, it is not safe to predict writes) */
-	TState* GetStateWrite(bool bHasAuthority) const
+	TState* GetStateWrite(bool bHasAuthority, const TCHAR* TraceStr=nullptr) const
 	{
 		TState* State = nullptr;
 		bool bSafe = false;
 		if (GetStateFunc)
 		{
-			GetStateFunc(true, State, bSafe);
+			GetStateFunc(true, State, bSafe, TraceStr);
 		}
 		return (bHasAuthority || bSafe) ? State : nullptr;
 	}
 
 private:
 
-	TUniqueFunction<void(bool bForWrite, TState*& OutState, bool& OutSafe)> GetStateFunc;
+	TUniqueFunction<void(bool bForWrite, TState*& OutState, bool& OutSafe, const TCHAR* TraceStr)> GetStateFunc;
 };
 
+
 // ----------------------------------------------------------------------------------------------------------------------------------------------
 //	Simulation Input/Output parameters. These are the data structures passed into the simulation code each frame
 // ----------------------------------------------------------------------------------------------------------------------------------------------
diff --git a/Engine/Plugins/Runtime/NetworkPrediction/Source/NetworkPrediction/Public/Trace/NetworkPredictionTrace.h b/Engine/Plugins/Runtime/NetworkPrediction/Source/NetworkPrediction/Public/Trace/NetworkPredictionTrace.h
--- a/Engine/Plugins/Runtime/NetworkPrediction/Source/NetworkPrediction/Public/Trace/NetworkPredictionTrace.h
+++ b/Engine/Plugins/Runtime/NetworkPrediction/Source/NetworkPrediction/Public/Trace/NetworkPredictionTrace.h
@@ -52,7 +52,7 @@
 #define UE_NP_TRACE_SYNTH_INPUT() FNetworkPredictionTrace::TraceSynthInput()
 
 // Called to indicate we are about to write state to the buffers outside of the normal simulation tick/netrecive. TODO: add char* identifier to debug where the mod came from
-#define UE_NP_TRACE_OOB_STATE_MOD() FNetworkPredictionTrace::TraceOOBStateMod()
+#define UE_NP_TRACE_OOB_STATE_MOD(SimulationId) FNetworkPredictionTrace::TraceOOBStateMod(SimulationId)
 
 // Called when a PIE session is started. This is so we can keep our *sets* of worlds/simulations seperate.
 #define UE_NP_TRACE_PIE_START() FNetworkPredictionTrace::TracePIEStart()
@@ -61,6 +61,9 @@
 
 #define UE_NP_TRACE_WORLD_FRAME_START(DeltaSeconds) FNetworkPredictionTrace::TraceWorldFrameStart(DeltaSeconds)
 
+#define UE_NP_TRACE_OOB_STR_SYNC(SimulationId, Str) FNetworkPredictionTrace::TraceOOBStr(SimulationId, FNetworkPredictionTrace::ETraceUserState::Sync, Str);
+#define UE_NP_TRACE_OOB_STR_AUX(SimulationId, Str) FNetworkPredictionTrace::TraceOOBStr(SimulationId, FNetworkPredictionTrace::ETraceUserState::Aux, Str);
+
 #else
 
 // Compiled out
@@ -81,6 +84,8 @@
 #define UE_NP_TRACE_PIE_START(...)
 #define UE_NP_TRACE_SYSTEM_FAULT(Format, ...) UE_LOG(LogNetworkSim, Warning, TEXT(Format), ##__VA_ARGS__);
 #define UE_NP_TRACE_WORLD_FRAME_START(...)
+#define UE_NP_TRACE_OOB_STR_SYNC(...)
+#define UE_NP_TRACE_OOB_STR_AUX(...)
 
 #endif // UE_NP_TRACE_ENABLED
 
@@ -91,7 +96,6 @@ class AActor;
 template<typename Model>
 struct TNetworkedSimulationState;
 
-// Trace is what the system+user code calls to do traces (could still be wrapped around [my own] macros, etc. This will call into the generic UE_TRACE_LOG macros)
 class NETWORKPREDICTION_API FNetworkPredictionTrace
 {
 public:
@@ -108,10 +112,9 @@ public:
 	static void TraceNetSerializeCommit();
 	static void TraceProduceInput();
 	static void TraceSynthInput();
-	static void TraceOOBStateMod();
+	static void TraceOOBStateMod(int32 SimulationId);
 	static void TracePIEStart();
 	static void TraceSystemFault(const TCHAR* Fmt, ...);
-
 	static void TraceWorldFrameStart(float DeltaSeconds);
 
 	enum ETraceUserState
@@ -120,6 +123,14 @@ public:
 		Sync,
 		Aux
 	};
+
+	static void TraceOOBStr(int32 SimulationId, ETraceUserState StateType, const TCHAR* Fmt)
+	{
+		if (Fmt)
+		{
+			TraceOOBStrInternal(SimulationId, StateType, Fmt);
+		}
+	}
 	
 	template<typename TUserState>
 	static void TraceUserState(const TUserState& State, int32 Frame, ETraceUserState StateType)
@@ -151,6 +162,7 @@ private:
 
 	static void TraceUserState_Internal(int32 Frame, ETraceUserState StateType);
 	static void TraceEOF_Internal(int32 BufferSize, int32 PendingTickFrame, int32 LatestInputFrame, int32 MaxTickFrame, FNetworkSimTime TotalSimTime, FNetworkSimTime AllowedSimTime);
+	static void TraceOOBStrInternal(int32 SimulationId, ETraceUserState StateType, const TCHAR* Fmt);
 
 	static FStringOutputDevice StrOut;
 
diff --git a/Engine/Plugins/Runtime/NetworkPredictionExtras/Source/NetworkPredictionExtras/Private/NetworkPredictionExtrasFlyingPawn.cpp b/Engine/Plugins/Runtime/NetworkPredictionExtras/Source/NetworkPredictionExtras/Private/NetworkPredictionExtrasFlyingPawn.cpp
--- a/Engine/Plugins/Runtime/NetworkPredictionExtras/Source/NetworkPredictionExtras/Private/NetworkPredictionExtrasFlyingPawn.cpp
+++ b/Engine/Plugins/Runtime/NetworkPredictionExtras/Source/NetworkPredictionExtras/Private/NetworkPredictionExtrasFlyingPawn.cpp
@@ -219,7 +219,7 @@ float ANetworkPredictionExtrasFlyingPawn::GetMaxMoveSpeed() const
 
 void ANetworkPredictionExtrasFlyingPawn::SetMaxMoveSpeed(float NewMaxMoveSpeed)
 {
-	if (FFlyingMovementAuxState* AuxState = GetAuxStateWrite())
+	if (FFlyingMovementAuxState* AuxState = GetAuxStateWrite(TEXT("SetMaxMoveSpeed")))
 	{
 		AuxState->MaxSpeed = NewMaxMoveSpeed;
 	}
@@ -227,7 +227,7 @@ void ANetworkPredictionExtrasFlyingPawn::SetMaxMoveSpeed(float NewMaxMoveSpeed)
 
 void ANetworkPredictionExtrasFlyingPawn::AddMaxMoveSpeed(float AdditiveMaxMoveSpeed)
 {
-	if (FFlyingMovementAuxState* AuxState = GetAuxStateWrite())
+	if (FFlyingMovementAuxState* AuxState = GetAuxStateWrite(TEXT("AddMaxMoveSpeed")))
 	{
 		//UE_LOG(LogTemp, Warning, TEXT("MaxSpeed: %.2f += %.2f"), AuxState->MaxSpeed, AdditiveMaxMoveSpeed);
 		AuxState->MaxSpeed += AdditiveMaxMoveSpeed;
@@ -243,27 +243,15 @@ const FFlyingMovementAuxState* ANetworkPredictionExtrasFlyingPawn::GetAuxStateRe
 	return nullptr;
 }
 
-FFlyingMovementAuxState* ANetworkPredictionExtrasFlyingPawn::GetAuxStateWrite()
+FFlyingMovementAuxState* ANetworkPredictionExtrasFlyingPawn::GetAuxStateWrite(const TCHAR* TraceStr)
 {
 	if (ensure(FlyingMovementComponent))
 	{
-		return FlyingMovementComponent->MovementAuxState.GetStateWrite(GetLocalRole() == ROLE_Authority);
+		return FlyingMovementComponent->MovementAuxState.GetStateWrite(GetLocalRole() == ROLE_Authority, TraceStr);
 	}
 	return nullptr;
 }
 
-int32 ANetworkPredictionExtrasFlyingPawn::GetPendingFrame() const
-{
-	// TODO
-	/*
-	if (ensure(FlyingMovementComponent))
-	{
-		return FlyingMovementComponent->MovementSyncState->GetPendingFrame();
-	}
-	*/
-	return 0;
-}
-
 void ANetworkPredictionExtrasFlyingPawn::ProduceInput(const FNetworkSimTime SimTime, FFlyingMovementInputCmd& Cmd)
 {
 	// Generate user commands. Called right before the flying movement simulation will tick (for a locally controlled pawn)
diff --git a/Engine/Plugins/Runtime/NetworkPredictionExtras/Source/NetworkPredictionExtras/Public/NetworkPredictionExtrasFlyingPawn.h b/Engine/Plugins/Runtime/NetworkPredictionExtras/Source/NetworkPredictionExtras/Public/NetworkPredictionExtrasFlyingPawn.h
--- a/Engine/Plugins/Runtime/NetworkPredictionExtras/Source/NetworkPredictionExtras/Public/NetworkPredictionExtrasFlyingPawn.h
+++ b/Engine/Plugins/Runtime/NetworkPredictionExtras/Source/NetworkPredictionExtras/Public/NetworkPredictionExtrasFlyingPawn.h
@@ -76,14 +76,10 @@ public:
 	UFUNCTION(BlueprintCallable, Category="Gameplay")
 	void AddMaxMoveSpeed(float AdditiveMaxMoveSpeed);
 
-	// Only intended for debugging in test map examples. Not really intended to be useful for general game code use.
-	UFUNCTION(BlueprintCallable, Category="Gameplay")
-	int32 GetPendingFrame() const;
-
 protected:
 
 	const FFlyingMovementAuxState* GetAuxStateRead() const;
-	FFlyingMovementAuxState* GetAuxStateWrite();
+	FFlyingMovementAuxState* GetAuxStateWrite(const TCHAR* TraceStr=nullptr);
 
 	void ProduceInput(const FNetworkSimTime SimTime, FFlyingMovementInputCmd& Cmd);
 
diff --git a/Engine/Plugins/Runtime/NetworkPredictionInsights/Source/NetworkPredictionInsights/Private/NetworkPredictionAnalyzer.cpp b/Engine/Plugins/Runtime/NetworkPredictionInsights/Source/NetworkPredictionInsights/Private/NetworkPredictionAnalyzer.cpp
--- a/Engine/Plugins/Runtime/NetworkPredictionInsights/Source/NetworkPredictionInsights/Private/NetworkPredictionAnalyzer.cpp
+++ b/Engine/Plugins/Runtime/NetworkPredictionInsights/Source/NetworkPredictionInsights/Private/NetworkPredictionAnalyzer.cpp
@@ -21,6 +21,8 @@ void FNetworkPredictionAnalyzer::OnAnalysisBegin(const FOnAnalysisContext& Conte
 	Builder.RouteEvent(RouteId_SimulationTick, "NetworkPrediction", "SimulationTick");
 	Builder.RouteEvent(RouteId_WorldFrameStart, "NetworkPrediction", "WorldFrameStart");
 	Builder.RouteEvent(RouteId_OOBStateMod, "NetworkPrediction", "OOBStateMod");
+	Builder.RouteEvent(RouteId_OOBStateModStrSync, "NetworkPrediction", "OOBStateModStrSync");
+	Builder.RouteEvent(RouteId_OOBStateModStrAux, "NetworkPrediction", "OOBStateModStrAux");
 	Builder.RouteEvent(RouteId_ProduceInput, "NetworkPrediction", "ProduceInput");
 	Builder.RouteEvent(RouteId_SynthInput, "NetworkPrediction", "SynthInput");
 	Builder.RouteEvent(RouteId_SimulationEOF, "NetworkPrediction", "SimulationEOF");
@@ -191,6 +193,18 @@ bool FNetworkPredictionAnalyzer::OnEvent(uint16 RouteId, const FOnEventContext&
 			break;
 		}
 
+		case RouteId_OOBStateModStrSync:
+		{
+			NetworkPredictionProvider.WriteOOBStateModStrSync(EventData.GetValue<uint32>("SimulationId"),Session.StoreString(reinterpret_cast<const TCHAR*>(EventData.GetAttachment())));
+			break;
+		}
+
+		case RouteId_OOBStateModStrAux:
+		{
+			NetworkPredictionProvider.WriteOOBStateModStrAux(EventData.GetValue<uint32>("SimulationId"),Session.StoreString(reinterpret_cast<const TCHAR*>(EventData.GetAttachment())));
+			break;
+		}
+
 		default:
 			bUnhandled = true;
 			break;
diff --git a/Engine/Plugins/Runtime/NetworkPredictionInsights/Source/NetworkPredictionInsights/Private/NetworkPredictionAnalyzer.h b/Engine/Plugins/Runtime/NetworkPredictionInsights/Source/NetworkPredictionInsights/Private/NetworkPredictionAnalyzer.h
--- a/Engine/Plugins/Runtime/NetworkPredictionInsights/Source/NetworkPredictionInsights/Private/NetworkPredictionAnalyzer.h
+++ b/Engine/Plugins/Runtime/NetworkPredictionInsights/Source/NetworkPredictionInsights/Private/NetworkPredictionAnalyzer.h
@@ -36,6 +36,8 @@ private:
 		RouteId_SimulationNetGUID,
 		RouteId_SimulationTick,
 		RouteId_OOBStateMod,
+		RouteId_OOBStateModStrSync,
+		RouteId_OOBStateModStrAux,
 		RouteId_ProduceInput,
 		RouteId_SynthInput,
 		RouteId_SimulationEOF,
diff --git a/Engine/Plugins/Runtime/NetworkPredictionInsights/Source/NetworkPredictionInsights/Private/NetworkPredictionProvider.cpp b/Engine/Plugins/Runtime/NetworkPredictionInsights/Source/NetworkPredictionInsights/Private/NetworkPredictionProvider.cpp
--- a/Engine/Plugins/Runtime/NetworkPredictionInsights/Source/NetworkPredictionInsights/Private/NetworkPredictionProvider.cpp
+++ b/Engine/Plugins/Runtime/NetworkPredictionInsights/Source/NetworkPredictionInsights/Private/NetworkPredictionProvider.cpp
@@ -256,11 +256,23 @@ void FNetworkPredictionProvider::WriteSystemFault(uint32 SimulationId, uint64 En
 
 void FNetworkPredictionProvider::WriteOOBStateMod(uint32 SimulationId)
 {
+	// Signals that the next user states traced will be OOB mods
 	FSimulationData& SimulationData = FindOrAdd(SimulationId).Get();
 	SimulationData.Analysis.PendingCommitUserStates.Reset();
 	SimulationData.Analysis.PendingUserStateSource = ENP_UserStateSource::OOB;
 }
 
+void FNetworkPredictionProvider::WriteOOBStateModStrSync(uint32 SimulationId, const TCHAR* Fmt)
+{
+	FSimulationData& SimulationData = FindOrAdd(SimulationId).Get();
+	SimulationData.Analysis.PendingOOBSyncStr = Fmt;
+}
+void FNetworkPredictionProvider::WriteOOBStateModStrAux(uint32 SimulationId, const TCHAR* Fmt)
+{
+	FSimulationData& SimulationData = FindOrAdd(SimulationId).Get();
+	SimulationData.Analysis.PendingOOBAuxStr = Fmt;
+}
+
 void FNetworkPredictionProvider::WriteProduceInput(uint32 SimulationId)
 {
 	FSimulationData& SimulationData = FindOrAdd(SimulationId).Get();
@@ -298,6 +310,19 @@ void FNetworkPredictionProvider::WriteUserState(uint32 SimulationId, int32 Frame
 	{
 		SimulationData.Analysis.PendingCommitUserStates.Add(&NewUserState);
 	}
+	else if (SimulationData.Analysis.PendingUserStateSource == ENP_UserStateSource::OOB)
+	{
+		if (Type == ENP_UserState::Sync)
+		{
+			NewUserState.OOBStr = SimulationData.Analysis.PendingOOBSyncStr;
+			SimulationData.Analysis.PendingOOBSyncStr = nullptr;
+		}
+		else if (Type == ENP_UserState::Aux)
+		{
+			NewUserState.OOBStr = SimulationData.Analysis.PendingOOBAuxStr;
+			SimulationData.Analysis.PendingOOBAuxStr = nullptr;
+		}
+	}
 }
 
 void FNetworkPredictionProvider::WritePIEStart()
diff --git a/Engine/Plugins/Runtime/NetworkPredictionInsights/Source/NetworkPredictionInsights/Private/NetworkPredictionProvider.h b/Engine/Plugins/Runtime/NetworkPredictionInsights/Source/NetworkPredictionInsights/Private/NetworkPredictionProvider.h
--- a/Engine/Plugins/Runtime/NetworkPredictionInsights/Source/NetworkPredictionInsights/Private/NetworkPredictionProvider.h
+++ b/Engine/Plugins/Runtime/NetworkPredictionInsights/Source/NetworkPredictionInsights/Private/NetworkPredictionProvider.h
@@ -34,6 +34,8 @@ public:
 	void WriteNetCommit(uint32 SimulationId);
 	void WriteSystemFault(uint32 SimulationId, uint64 EngineFrameNumber, const TCHAR* Fmt);
 	void WriteOOBStateMod(uint32 SimulationId);
+	void WriteOOBStateModStrSync(uint32 SimulationId, const TCHAR* Fmt);
+	void WriteOOBStateModStrAux(uint32 SimulationId, const TCHAR* Fmt);
 	void WriteProduceInput(uint32 SimulationId);
 	void WriteSynthInput(uint32 SimulationId);
 	void WriteUserState(uint32 SimulationId, int32 Frame, uint64 EngineFrame, ENP_UserState Type, const TCHAR* UserStr);
diff --git a/Engine/Plugins/Runtime/NetworkPredictionInsights/Source/NetworkPredictionInsights/Private/UI/SNPSimFrameContents.cpp b/Engine/Plugins/Runtime/NetworkPredictionInsights/Source/NetworkPredictionInsights/Private/UI/SNPSimFrameContents.cpp
--- a/Engine/Plugins/Runtime/NetworkPredictionInsights/Source/NetworkPredictionInsights/Private/UI/SNPSimFrameContents.cpp
+++ b/Engine/Plugins/Runtime/NetworkPredictionInsights/Source/NetworkPredictionInsights/Private/UI/SNPSimFrameContents.cpp
@@ -65,13 +65,18 @@ void SNPSimFrameContents::NotifyContentClicked(const FSimContentsView& InContent
 	auto MakeUserStateWidget = [&](const FUserStateWidgetInfo& Info)
 	{
 		FString HeaderText;
-		const TCHAR* UserStr = nullptr;
+		FString UserStr;
 		uint64 EngineFrameHyperLink = 0;
 		if (Info.State)
 		{
 			// style=\"Hyperlink\"
 			HeaderText = FString::Printf(TEXT("%s [Sim: %d. <a id=\"engine\">Engine: %d</>. %s]"), Info.Heading, Info.State->SimFrame, Info.State->EngineFrame, LexToString(Info.State->Source));
 			UserStr = Info.State->UserStr;
+			if (Info.State->OOBStr)
+			{
+				UserStr += FString::Printf(TEXT("\nOOB Mod:\n%s"), Info.State->OOBStr);
+			}
+
 			EngineFrameHyperLink = Info.State->EngineFrame;
 		}
 		else
diff --git a/Engine/Plugins/Runtime/NetworkPredictionInsights/Source/NetworkPredictionInsights/Private/UI/SNPSimFrameView.cpp b/Engine/Plugins/Runtime/NetworkPredictionInsights/Source/NetworkPredictionInsights/Private/UI/SNPSimFrameView.cpp
--- a/Engine/Plugins/Runtime/NetworkPredictionInsights/Source/NetworkPredictionInsights/Private/UI/SNPSimFrameView.cpp
+++ b/Engine/Plugins/Runtime/NetworkPredictionInsights/Source/NetworkPredictionInsights/Private/UI/SNPSimFrameView.cpp
@@ -113,7 +113,7 @@ void FSimFrameViewDrawHelper::DrawCached(const FSimulationFrameView& View, float
 	const FLinearColor SelectedColor(1.0f, 0.5f, 0.0f);
 	const FLinearColor SelectedPulseColor(1.0f * PulseFactor, 0.5f * PulseFactor, 0.0f);
 
-	const FLinearColor SearchHighlightColor(1.0f, 1.0f, 0.25f);
+	const FLinearColor SearchHighlightColor(1.0f, 1.0f, 0.15f);
 	const FLinearColor SearchHighlightPulseColor(1.0f * PulseFactor, 1.0f * PulseFactor, 0.25f);
 
 	auto SelectBorderColor = [&](const bool bPulse, const bool bSelected, const bool bSearchHighlighted, const FLinearColor& ColorFill)
@@ -144,7 +144,7 @@ void FSimFrameViewDrawHelper::DrawCached(const FSimulationFrameView& View, float
 			DrawContext.DrawBox(X + 1.0f, Y + 1.0f, W - 2.0f, H - 2.0f, WhiteBrush, ColorFill);
 
 			// Draw border.
-			const float B = (Source.bPulse || Source.bSelected) ? 4.f : 2.f;
+			const float B = Source.bSearchHighlighted ? 6.f : ((Source.bPulse || Source.bSelected) ? 4.f : 2.f);
 			const float B2 = B/2.f;
 						
 			DrawContext.DrawBox(X, Y, B2, H, WhiteBrush, ColorBorder);
diff --git a/Engine/Plugins/Runtime/NetworkPredictionInsights/Source/NetworkPredictionInsights/Public/INetworkPredictionProvider.h b/Engine/Plugins/Runtime/NetworkPredictionInsights/Source/NetworkPredictionInsights/Public/INetworkPredictionProvider.h
--- a/Engine/Plugins/Runtime/NetworkPredictionInsights/Source/NetworkPredictionInsights/Public/INetworkPredictionProvider.h
+++ b/Engine/Plugins/Runtime/NetworkPredictionInsights/Source/NetworkPredictionInsights/Public/INetworkPredictionProvider.h
@@ -374,7 +374,6 @@ struct FSimulationData
 	};
 
 	// End of (Engine) Frame state of the simulation
-	// FIXME: it would be much nicer if we had a single global emission of GFrameNumber and EngineFrameDeltaTime
 	// It may be better to eventually just trace the events that change this simulation specific state (NetSerialize, Tick, etc)
 	struct FEngineFrame
 	{
@@ -418,6 +417,7 @@ struct FSimulationData
 		FUserState* Prev = nullptr; // Previously written user state at this SimFrame number. Prev->EngineFrame < this->EngineFrame.
 		const TCHAR* UserStr = nullptr; // stable pointer to traced string representation of user state
 		ENP_UserStateSource Source = ENP_UserStateSource::Unknown; // Analyzed source of this user state
+		const TCHAR* OOBStr = nullptr; // if OOB mod, this is the user traced string (telling us who did it)
 	};
 
 	struct FUserStateStore
@@ -590,6 +590,8 @@ struct FSimulationData
 		ENP_UserStateSource PendingUserStateSource = ENP_UserStateSource::Unknown;
 		TArray<FUserState*> PendingCommitUserStates; // NetRecv'd state that hasn't been commited
 		TArray<FSystemFault> PendingSystemFaults;
+		const TCHAR* PendingOOBSyncStr = nullptr;
+		const TCHAR* PendingOOBAuxStr = nullptr;
 
 	} Analysis;
 };
