NetworkPrediction
-Fix several issues with previous changes
-Moved MaxTickFrame, NextInputFrame to RepControllers instead of being a central state buffer concept.

#rnx

diff --git a/Engine/Plugins/Runtime/NetworkPrediction/Source/NetworkPrediction/Private/Trace/NetworkPredictionTrace.cpp b/Engine/Plugins/Runtime/NetworkPrediction/Source/NetworkPrediction/Private/Trace/NetworkPredictionTrace.cpp
--- a/Engine/Plugins/Runtime/NetworkPrediction/Source/NetworkPrediction/Private/Trace/NetworkPredictionTrace.cpp
+++ b/Engine/Plugins/Runtime/NetworkPrediction/Source/NetworkPrediction/Private/Trace/NetworkPredictionTrace.cpp
@@ -55,7 +55,6 @@ UE_TRACE_EVENT_BEGIN(NetworkPrediction, SimulationEOF)
 	UE_TRACE_EVENT_FIELD(int32, BufferSize)
 	UE_TRACE_EVENT_FIELD(int32, PendingTickFrame)
 	UE_TRACE_EVENT_FIELD(int32, LatestInputFrame)
-	UE_TRACE_EVENT_FIELD(int32, MaxTickFrame)
 	UE_TRACE_EVENT_FIELD(int32, TotalSimTime)
 	UE_TRACE_EVENT_FIELD(int32, AllowedSimTime)
 UE_TRACE_EVENT_END()
@@ -216,7 +215,7 @@ void FNetworkPredictionTrace::TraceSimulationTick(int32 OutputFrame, const FNetw
 		<< SimulationTick.OutputFrame(OutputFrame);
 }
 
-void FNetworkPredictionTrace::TraceEOF_Internal(int32 BufferSize, int32 PendingTickFrame, int32 LatestInputFrame, int32 MaxTickFrame, FNetworkSimTime TotalSimTime, FNetworkSimTime AllowedSimTime)
+void FNetworkPredictionTrace::TraceEOF_Internal(int32 BufferSize, int32 PendingTickFrame, int32 LatestInputFrame, FNetworkSimTime TotalSimTime, FNetworkSimTime AllowedSimTime)
 {
 	const uint32 SimulationId = PeakSimulationIdChecked();
 
@@ -226,7 +225,6 @@ void FNetworkPredictionTrace::TraceEOF_Internal(int32 BufferSize, int32 PendingT
 		<< SimulationEOF.BufferSize(BufferSize)
 		<< SimulationEOF.PendingTickFrame(PendingTickFrame)
 		<< SimulationEOF.LatestInputFrame(LatestInputFrame)
-		<< SimulationEOF.MaxTickFrame(MaxTickFrame)
 		<< SimulationEOF.TotalSimTime(TotalSimTime)
 		<< SimulationEOF.AllowedSimTime(AllowedSimTime);
 
diff --git a/Engine/Plugins/Runtime/NetworkPrediction/Source/NetworkPrediction/Public/NetworkedSimulationModelBuffer.h b/Engine/Plugins/Runtime/NetworkPrediction/Source/NetworkPrediction/Public/NetworkedSimulationModelBuffer.h
--- a/Engine/Plugins/Runtime/NetworkPrediction/Source/NetworkPrediction/Public/NetworkedSimulationModelBuffer.h
+++ b/Engine/Plugins/Runtime/NetworkPrediction/Source/NetworkPrediction/Public/NetworkedSimulationModelBuffer.h
@@ -89,6 +89,22 @@ struct TNetworkSimAuxBuffer
 		return TNetSimLazyWriterFunc<ElementType>(*this, Frame);
 	}
 
+	// Resizes while preserving contents from head position
+	void Resize(int32 NewSize)
+	{
+		check(NewSize > 0);
+		TArray<TInternal> PrevElements = MoveTemp(Elements);
+		int32 PrevMask = IndexMask;
+
+		Init(NewSize);
+
+		int32 StartCopyFrom = HeadPosition - FMath::Min(PrevElements.Num(), NewSize) + 1;
+		for (int32 i=StartCopyFrom; i <= HeadPosition; ++i)
+		{
+			Elements[i & IndexMask] = PrevElements[i & PrevMask];
+		}
+	}
+
 private:
 
 	struct TInternal
@@ -115,7 +131,6 @@ private:
 			{
 				return &Data.Element;
 			}
-			Pos--;
 		} while(--Pos >= TailPos);
 
 		return nullptr;
@@ -157,7 +172,6 @@ struct TNetworkSimFrameBuffer
 	}
 
 	// Resizes while preserving contents from head position
-	// Could be faster but doesn't really matter, this will be rare
 	void Resize(int32 NewSize, int32 PrevHead)
 	{
 		check(NewSize > 0);
@@ -166,7 +180,7 @@ struct TNetworkSimFrameBuffer
 
 		Init(NewSize);
 
-		int32 StartCopyFrom = PrevHead - FMath::Min(PrevElements.Num(), NewSize);
+		int32 StartCopyFrom = PrevHead - FMath::Min(PrevElements.Num(), NewSize) + 1;
 		for (int32 i=StartCopyFrom; i <= PrevHead; ++i)
 		{
 			Elements[i & IndexMask] = PrevElements[i & PrevMask];
diff --git a/Engine/Plugins/Runtime/NetworkPrediction/Source/NetworkPrediction/Public/NetworkedSimulationModelDebugger.h b/Engine/Plugins/Runtime/NetworkPrediction/Source/NetworkPrediction/Public/NetworkedSimulationModelDebugger.h
--- a/Engine/Plugins/Runtime/NetworkPrediction/Source/NetworkPrediction/Public/NetworkedSimulationModelDebugger.h
+++ b/Engine/Plugins/Runtime/NetworkPrediction/Source/NetworkPrediction/Public/NetworkedSimulationModelDebugger.h
@@ -383,7 +383,7 @@ struct TNetworkSimulationModelDebugger : public INetworkSimulationModelDebugger
 				*NetworkSim->RepProxy_Autonomous.GetLastSerializedSimTime().ToString(), *(NetworkSim->State.GetTotalProcessedSimulationTime() - NetworkSim->RepProxy_Autonomous.GetLastSerializedSimTime()).ToString());
 			Out.Emit(MoveTemp(SimulationTimeString), Color);
 
-			FString AllowedSimulationTimeString = FString::Printf(TEXT("Frame: %d/%d/%d"), NetworkSim->State.GetMaxTickFrame(), NetworkSim->State.GetPendingTickFrame(), NetworkSim->State.GetLatestInputFrame());
+			FString AllowedSimulationTimeString = FString::Printf(TEXT("%s"), *NetworkSim->State.DebugString());
 			Out.Emit(MoveTemp(AllowedSimulationTimeString), Color);
 		}
 		else if (Owner->GetLocalRole() == ROLE_SimulatedProxy)
@@ -499,4 +499,4 @@ void FNetworkSimulationModelDebuggerManager::RegisterNetworkSimulationModel(T* N
 		}
 		return typename TNetworkSimulationModelDebugger<T>::FStuff();
 	};
-}
\ No newline at end of file
+}
diff --git a/Engine/Plugins/Runtime/NetworkPrediction/Source/NetworkPrediction/Public/NetworkedSimulationModelRepControllers.h b/Engine/Plugins/Runtime/NetworkPrediction/Source/NetworkPrediction/Public/NetworkedSimulationModelRepControllers.h
--- a/Engine/Plugins/Runtime/NetworkPrediction/Source/NetworkPrediction/Public/NetworkedSimulationModelRepControllers.h
+++ b/Engine/Plugins/Runtime/NetworkPrediction/Source/NetworkPrediction/Public/NetworkedSimulationModelRepControllers.h
@@ -14,8 +14,8 @@
 //	The RepControllers are the pieces of the TNetworkedSimulationModel that make up the role-specific functionality (Server, Autonomous Client, and Simulated Client).
 //	Mainly they NetSerialize, Reconcile, and Tick. Essentially, they are what drives the Networked Simulation frame to frame based on networking role.
 //
-//	Multiple controllers can be used on a single TNetworkedSimulationModel instance. There is a 1:1 mapping for Tick, Reconcile, and Receiving. But to send you must 
-//	use the RepController of the receiving side.
+//	Multiple controllers can be used on a single TNetworkedSimulationModel instance. There is a 1:1 mapping for Tick, Reconcile, and Receiving. But to send, the 
+//	RepController of the receiving side is used.
 //
 //	Tick, Reconcile, NetSerialize: Receive
 //		Server    --> TRepController_Server
@@ -35,7 +35,6 @@ namespace NetworkSimulationModelCVars
 	NETSIM_DEVCVAR_SHIPCONST_INT(EnableSimulatedExtrapolation, 1, "ns.EnableSimulatedExtrapolation", "Toggle simulated proxy extrapolation.");
 	NETSIM_DEVCVAR_SHIPCONST_INT(ForceReconcile, 0, "ns.ForceReconcile", "Forces reconcile even if state does not differ. E.g, force resimulation after every netupdate.");
 	NETSIM_DEVCVAR_SHIPCONST_INT(ForceReconcileSingle, 0, "ns.ForceReconcileSingle", "Forces a since reconcile to happen on the next frame");
-	NETSIM_DEVCVAR_SHIPCONST_INT(ForceAutonomousFrameJump, 0, "ns.ForceAutonomousFrameJump", "Causes AP client to jump 1000 frames ahead in its local frame #");
 }
 
 #define NETSIM_ENABLE_CHECKSUMS !(UE_BUILD_SHIPPING || UE_BUILD_TEST)
@@ -88,11 +87,14 @@ void GenerateLocalInputCmdAtFrame(TNetworkedSimulationModelDriver<typename Model
 {
 	// Write the new frame
 	const int32 InputFrame = State.GetNextInputFrame();
-	TSimulationFrameState<Model>* FrameState = State.WriteFrameSequential(InputFrame);
-	
-	npEnsure(FrameState->HasFlag(ESimulationFrameStateFlags::InputWritten) == false);
+	if (InputFrame > State.GetMaxWritableFrame())
+	{
+		return;
+	}
+
+	TSimulationFrameState<Model>* FrameState = State.WriteInputFrame(InputFrame);
+	npCheckSlow(FrameState);	
 	FrameState->FrameDeltaTime.Set(DeltaSimTime);
-	FrameState->SetFlag(ESimulationFrameStateFlags::InputWritten);
 
 	if (bProduceInputViaDriver)
 	{
@@ -103,32 +105,30 @@ void GenerateLocalInputCmdAtFrame(TNetworkedSimulationModelDriver<typename Model
 	{
 		UE_NP_TRACE_SYNTH_INPUT();
 	}
-
-	State.SetLatestInputFrame(InputFrame);
+	
 	UE_NP_TRACE_USER_STATE_INPUT(FrameState->InputCmd, InputFrame);
 }
 
 /** Helper to generate a local input cmd if we have simulation time to spend and advance the simulation's MaxAllowedFrame so that it can be processed. */
 template<typename Model>
-void TryGenerateLocalInput(TNetworkedSimulationModelDriver<typename Model::BufferTypes>* Driver, TNetworkedSimulationState<Model>& State, bool bProduceInputViaDriver, int32 MaxInputFrame=INDEX_NONE)
+void TryGenerateLocalInput(TNetworkedSimulationModelDriver<typename Model::BufferTypes>* Driver, TNetworkedSimulationState<Model>& State, bool bProduceInputViaDriver)
 {	
 	const FNetworkSimTime AllowedSimTime = State.GetTotalAllowedSimulationTime();
+	const int32 MaxInputFrame = State.GetMaxWritableFrame();
 	
 	// Submit new input if necessary
 	FNetworkSimTime MissingInputTime = AllowedSimTime - State.GetTotalInputSimulationTime();
 	while (MissingInputTime.IsPositive())
 	{
 		const int32 InputFrame = State.GetNextInputFrame();
-		if (MaxInputFrame != INDEX_NONE && InputFrame >= MaxInputFrame)
+		if (InputFrame >= MaxInputFrame)
 		{
 			break;
 		}
 		
-		TSimulationFrameState<Model>* FrameState = State.WriteFrameSequential(InputFrame);
-		npEnsure(FrameState->HasFlag(ESimulationFrameStateFlags::InputWritten) == false);
+		TSimulationFrameState<Model>* FrameState = State.WriteInputFrame(InputFrame);
+		npCheckSlow(FrameState);		
 		FrameState->FrameDeltaTime.Set(MissingInputTime);
-		FrameState->SetFlag(ESimulationFrameStateFlags::InputWritten);
-		State.SetLatestInputFrame(InputFrame);
 
 		npEnsureMsgfSlow(State.GetTotalAllowedSimulationTime() >= State.GetTotalInputSimulationTime(), TEXT("Overflowed input %s"), *State.DebugString());
 		MissingInputTime -= FrameState->FrameDeltaTime.Get();
@@ -145,25 +145,14 @@ void TryGenerateLocalInput(TNetworkedSimulationModelDriver<typename Model::Buffe
 		
 		UE_NP_TRACE_USER_STATE_INPUT(FrameState->InputCmd, InputFrame);
 	}
-
-	// Advance MaxTickFrame
-	for (int32 Frame = State.GetLatestInputFrame(); Frame > State.GetMaxTickFrame(); --Frame)
-	{
-		auto* FrameState = State.ReadFrame(Frame);
-		if (npEnsure(FrameState) && FrameState->HasFlag(ESimulationFrameStateFlags::InputWritten) && FrameState->TotalSimTime <= AllowedSimTime)
-		{
-			State.SetMaxTickFrame(Frame);
-			break;
-		}
-	}
 }
 
 template<typename Model>
-void DoSimulationTicks(typename Model::Simulation& Simulation, TNetworkedSimulationState<Model>& State, ESimulationTickContext Context)
+void DoSimulationTicks(typename Model::Simulation& Simulation, TNetworkedSimulationState<Model>& State, ESimulationTickContext Context, const int32 MaxTickFrame)
 {
 	using TFrameState = TSimulationFrameState<Model>;
 
-	while (State.GetPendingTickFrame() <= State.GetMaxTickFrame())
+	while (State.GetPendingTickFrame() <= MaxTickFrame)
 	{
 		const int32 InputSyncFrame = State.GetPendingTickFrame(); // Where Sync/Aux state have to come from. PendingFrame is where OOB mods are made, so we cannot skip ahead.
 			
@@ -172,7 +161,7 @@ void DoSimulationTicks(typename Model::Simulation& Simulation, TNetworkedSimulat
 		TFrameState* InputFrameState = State.ReadFrame(InputFrame);
 		while (InputFrameState && !InputFrameState->HasFlag(ESimulationFrameStateFlags::InputWritten))
 		{
-			if (++InputFrame > State.GetMaxTickFrame())
+			if (++InputFrame > State.GetLatestInputFrame())
 			{
 				InputFrameState = nullptr;
 				break;
@@ -315,7 +304,7 @@ struct TRepController_Server : public TBase
 			const int32 StartInputFrame = FMath::Max(ReceivedHeadInputFrame - SerializedNumElements + 1, 0);
 
 			const int32 MinAcceptableFrame = State.GetNextInputFrame();
-			const int32 MaxAcceptableFrame = State.GetPendingTickFrame() + State.Capacity() - 1;
+			const int32 MaxAcceptableFrame = State.GetMaxWritableFrame();
 
 			for (int32 Frame = StartInputFrame; Frame <= ReceivedHeadInputFrame; ++Frame)
 			{
@@ -350,13 +339,11 @@ struct TRepController_Server : public TBase
 				else
 				{
 					// Brand new frame
-					TFrameState* NewFrameState = State.WriteFrameSequential(Frame);
+					TFrameState* NewFrameState = State.WriteInputFrame(Frame);
+					npCheckSlow(NewFrameState);
 					
 					NewFrameState->FrameDeltaTime.NetSerialize(P.Ar);
-					NewFrameState->InputCmd.NetSerialize(P);
-					NewFrameState->SetFlag(ESimulationFrameStateFlags::InputWritten);
-
-					State.SetLatestInputFrame(Frame);
+					NewFrameState->InputCmd.NetSerialize(P);					
 
 					TraceNewInputCmd(NewFrameState->InputCmd, NewFrameState->TotalSimTime, Frame);
 				}
@@ -372,12 +359,7 @@ struct TRepController_Server : public TBase
 	// -------------------------------------------------------------------------------------------------
 	void Reconcile(TSimulation* Simulation, TDriver* Driver, TNetworkedSimulationState<Model>& State)
 	{
-		// After receiving input, server may process up to the latest received frames.
-		// (If we needed to buffer input server side for whatever reason, we would do it here)
-		// (also note that we will implicitly guard against speed hacks in the core update loop by not processing cmds past what we have been "allowed")
-		State.SetMaxTickFrame(FMath::Max(LastSerializedFrame, State.GetMaxTickFrame()));
-
-		
+		// Nothing to actually reconcile server side
 	}
 	
 	// -------------------------------------------------------------------------------------------------
@@ -392,7 +374,7 @@ struct TRepController_Server : public TBase
 			TryGenerateLocalInput(Driver, State, true);
 		}
 
-		DoSimulationTicks(*Simulation, State, ESimulationTickContext::Authority);
+		DoSimulationTicks(*Simulation, State, ESimulationTickContext::Authority, State.GetLatestInputFrame());
 
 		// Sync to latest frame if there is any
 		const int32 PendingFrame = State.GetPendingTickFrame();
@@ -479,13 +461,10 @@ struct TRepController_Autonomous: public TBase
 			{
 				// We don't have this frame locally anymore, will sort it out in Reconcile
 				bPendingReconciliation =  true;
+				UE_NP_TRACE_SYSTEM_FAULT("TRepController_Autonomous::NetSerialize Unable to compare SerializedState. Received: %d. %s", SerializedFrame, *State.DebugString());
 			}
 
-			if (State.GetPendingTickFrame() < SerializedFrame)
-			{
-				State.Reset(SerializedFrame);
-			}
-			else
+			if (!bPendingReconciliation)
 			{
 				State.SetConfirmedFrame(SerializedFrame);
 			}
@@ -540,7 +519,6 @@ struct TRepController_Autonomous: public TBase
 		State.CueDispatcher.NotifyRollback(SerializedTime);
 
 		const int32 StartingPendingTickFrame = State.GetPendingTickFrame();
-		const int32 ExpectedEndPendingFrame = FMath::Min(StartingPendingTickFrame, State.GetMaxTickFrame());
 		const int32 LastResimulateFrame = StartingPendingTickFrame-1;
 		if (SerializedFrame < StartingPendingTickFrame)
 		{
@@ -551,9 +529,8 @@ struct TRepController_Autonomous: public TBase
 			npEnsureMsgf(!StartingRemainingTime.IsNegative(), TEXT("Negative remaining SimTime %s %s"), *State.GetTotalAllowedSimulationTime().ToString(), *State.GetTotalProcessedSimulationTime().ToString());
 
 			State.SetPendingTickFrame(SerializedFrame);
-			State.SetMaxTickFrame(StartingPendingTickFrame-1);
 
-			DoSimulationTicks(*Simulation, State, ESimulationTickContext::Resimulate);
+			DoSimulationTicks(*Simulation, State, ESimulationTickContext::Resimulate, StartingPendingTickFrame-1);
 			
 			State.SetTotalAllowedSimulationTime(State.GetTotalProcessedSimulationTime() + StartingRemainingTime);
 
@@ -562,8 +539,9 @@ struct TRepController_Autonomous: public TBase
 		}
 		else
 		{
-			// We have no frames to resimulate, essentially we are reset to the SerializedFrame
-			State.Reset(SerializedFrame);
+			// We have no frames to resimulate, essentially we are reset to the SerializedFrame	
+			State.ResetToFrame(SerializedFrame, ClientFrameState->HasFlag(ESimulationFrameStateFlags::InputWritten));
+			State.SetTotalAllowedSimulationTime(ClientFrameState->TotalSimTime);
 		}
 
 		State.CheckInvariants();
@@ -574,27 +552,15 @@ struct TRepController_Autonomous: public TBase
 	// --------------------------------------------------------------------
 	void Tick(TDriver* Driver, TSimulation* Simulation, TNetworkedSimulationState<Model>& State, const FNetSimTickParameters& TickParameters)
 	{
-		if (NetworkSimulationModelCVars::ForceAutonomousFrameJump() > 0)
+		if (SerializedFrame == INDEX_NONE)
 		{
-			// Testing CVar to mimic grave mishaps
-			const int32 FrameJump = NetworkSimulationModelCVars::ForceAutonomousFrameJump();
-			NetworkSimulationModelCVars::SetForceAutonomousFrameJump(0);
-			
-			// Copy pending state to jumped frame
-			TFrameState PendingFrameState = *State.ReadFrame(State.GetPendingTickFrame());
-			const int32 NewFrame = State.GetPendingTickFrame() + FrameJump;
-			TFrameState* NewFrameState =  State.WriteFrameSequential(NewFrame);
-			*NewFrameState = PendingFrameState;
-
-			State.Reset(NewFrame);
-
-			UE_NP_TRACE_SYSTEM_FAULT("NetworkSimulationModelCVars::ForceAutonomousFrameJump invoked. Jumped to: %d. Prev: %s", FrameJump, *State.DebugString());
+			State.SetConfirmedFrame(0);
 		}
 
 		// The MaxInputFrame we can write without getting too far ahead of server
 		// We pass this into TryGenerateLocalInput to ensure that it doesn't.
 		// If we some how get in a state where we are past MaxInputFrame, we need to reset
-		const int32 MaxInputFrame = (SerializedFrame + State.Capacity() - 1);
+		const int32 MaxInputFrame = State.GetMaxWritableFrame();
 
 		if (State.GetNextInputFrame() > MaxInputFrame)
 		{
@@ -620,7 +586,7 @@ struct TRepController_Autonomous: public TBase
 			{
 				// Prediction: add simulation time and generate new commands
 				State.GiveSimulationTime(TickParameters.LocalDeltaTimeSeconds);
-				TryGenerateLocalInput<Model>(Driver, State, true, MaxInputFrame);
+				TryGenerateLocalInput<Model>(Driver, State, true);
 			}
 			else
 			{
@@ -636,7 +602,7 @@ struct TRepController_Autonomous: public TBase
 			}
 		}
 
-		DoSimulationTicks(*Simulation, State, ESimulationTickContext::Predict);
+		DoSimulationTicks(*Simulation, State, ESimulationTickContext::Predict, State.GetLatestInputFrame());
 
 		State.CueDispatcher.ClearMaxDispatchTime();
 		State.CueDispatcher.SetConfirmedTime(SerializedTime);
@@ -652,9 +618,14 @@ private:
 	{
 		if (SerializedFrame != INDEX_NONE)
 		{
-			State.WriteFrameWithTime(SerializedFrame, SerializedTime)->SyncState = SerializedSyncState;
-			*State.WriteAuxState(SerializedFrame) = SerializedAuxState;
-			State.Reset(SerializedFrame);
+			TFrameState* Frame = State.WriteFrameDirect(SerializedFrame);
+			Frame->TotalSimTime = SerializedTime;
+			Frame->SyncState = SerializedSyncState;
+
+			State.ResetToFrame(SerializedFrame, Frame->HasFlag(ESimulationFrameStateFlags::InputWritten));
+			State.SetTotalAllowedSimulationTime(Frame->TotalSimTime);
+
+			*State.WriteAuxState(SerializedFrame) = SerializedAuxState;			
 		}
 	}
 
@@ -677,7 +648,7 @@ private:
 	TSyncState SerializedSyncState;
 	TAuxState SerializedAuxState;
 	FNetworkSimTime SerializedTime;
-	int32 SerializedFrame = -1;
+	int32 SerializedFrame = INDEX_NONE;
 
 	bool bPendingReconciliation = false;	// Reconciliation is pending: we need to reconcile state from the server that differs from the locally predicted state
 	bool bReconcileFaultDetected = false;	// A fault was detected: we received state from the server that we are unable to reconcile with locally predicted state
@@ -771,44 +742,37 @@ struct TRepController_Simulated : public TBase
 			{
 				ReconcileSimulationTime = State.GetTotalProcessedSimulationTime();
 			}
-			
-			LastSerializedSimulationTime.NetSerialize(P.Ar); // 1. Time
 
 			LastSerializedKeyframe = State.GetPendingTickFrame()+1;
-			
-			TFrameState* FrameState = State.WriteFrameWithTime(LastSerializedKeyframe, LastSerializedSimulationTime);
+
+			TFrameState* FrameState = State.WriteFrameDirect(LastSerializedKeyframe);
 			npCheckSlow(FrameState);
+			FrameState->TotalSimTime.NetSerialize(P.Ar); // 1. Time
+			FrameState->SyncState.NetSerialize(P.Ar); // 2. Sync
+			FrameState->InputCmd.NetSerialize(P.Ar); // 3. Input
+			FrameState->SetFlag(ESimulationFrameStateFlags::InputWritten);
 
 			TAuxState* AuxState = State.WriteAuxState(LastSerializedKeyframe);
 			npCheckSlow(AuxState);
-
-			FrameState->SyncState.NetSerialize(P.Ar);	// 2. Sync
-			FrameState->InputCmd.NetSerialize(P.Ar);	// 3. Input
-			FrameState->SetFlag(ESimulationFrameStateFlags::InputWritten);
-			
 			AuxState->NetSerialize(P.Ar); // 4. Aux
 
 			State.CueDispatcher.NetSerializeSavedCues(Ar, GetSimulatedUpdateMode() == ESimulatedUpdateMode::Interpolate ? ENetSimCueReplicationTarget::Interpolators : ENetSimCueReplicationTarget::SimulatedProxy); // 5. Cues
-
-			State.SetLatestInputFrame(LastSerializedKeyframe);
-			State.SetMaxTickFrame(LastSerializedKeyframe);
-			State.SetPendingTickFrame(LastSerializedKeyframe);
-			State.SetConfirmedFrame(LastSerializedKeyframe);
 			
+			State.ResetToFrame(LastSerializedKeyframe, true);
 			if (State.GetTotalAllowedSimulationTime() < FrameState->TotalSimTime)
 			{
 				State.SetTotalAllowedSimulationTime(FrameState->TotalSimTime);
 			}
 
-			UE_NP_TRACE_NET_SERIALIZE_RECV(LastSerializedSimulationTime, LastSerializedKeyframe);
+			UE_NP_TRACE_NET_SERIALIZE_RECV(FrameState->TotalSimTime, LastSerializedKeyframe);
 
 			UE_NP_TRACE_USER_STATE_INPUT(FrameState->InputCmd, LastSerializedKeyframe);
 			UE_NP_TRACE_USER_STATE_SYNC(FrameState->SyncState, LastSerializedKeyframe);
 			UE_NP_TRACE_USER_STATE_AUX(*AuxState, LastSerializedKeyframe);
-
-			// Temp: this will go away
-			LastSerializedInputCmd = FrameState->InputCmd;
+			
+			LastSerializedSimulationTime = FrameState->TotalSimTime;
 			LastSerializedSyncState = FrameState->SyncState;
+			LastSerializedInputCmd = FrameState->InputCmd;
 			LastSerializedAuxState = *AuxState;
 		}
 
@@ -866,7 +830,7 @@ struct TRepController_Simulated : public TBase
 			TryGenerateLocalInput<Model>(Driver, State, false);
 		}
 
-		DoSimulationTicks(*Simulation, State, ESimulationTickContext::SimExtrapolate);
+		DoSimulationTicks(*Simulation, State, ESimulationTickContext::SimExtrapolate, State.GetLatestInputFrame());
 
 		if (GetSimulatedUpdateMode() == ESimulatedUpdateMode::Interpolate)
 		{
@@ -896,15 +860,13 @@ struct TRepController_Simulated : public TBase
 	void DependentRollbackBegin(TSimulation* Simulation, TDriver* Driver, TNetworkedSimulationState<Model>& State, const FNetworkSimTime& RollbackDeltaTime, const int32 ParentFrame)
 	{
 		// FIXME: this isn't using RollbackDeltaTime
-
-		// Rollback to the last serialized frame
-		// For now, we make the assumption that our last serialized state and time match the parent simulation.
-		// This would not always be the case with low frequency simulated proxies. But could be handled by replicating the simulations together (at the replication level)
-		// npEnsure(Buffers.IsValidFrame(LastSerializedSimulationTime));
-		npEnsure(State.WriteFrameSequential(LastSerializedKeyframe)); // fixme
-		
-		State.SetPendingTickFrame(LastSerializedKeyframe);
-		State.SetMaxTickFrame(LastSerializedKeyframe);
+		TFrameState* RollbackState = State.WriteFrameDirect(LastSerializedKeyframe);
+		RollbackState->SyncState = LastSerializedSyncState;
+		RollbackState->InputCmd = LastSerializedInputCmd;
+		RollbackState->TotalSimTime = LastSerializedSimulationTime;
+		RollbackState->SetFlag(ESimulationFrameStateFlags::InputWritten);
+		State.ResetToFrame(LastSerializedKeyframe, true);
+		State.SetTotalAllowedSimulationTime(LastSerializedSimulationTime);
 
 		State.CueDispatcher.NotifyRollback(LastSerializedSimulationTime);
 	}
@@ -932,7 +894,6 @@ private:
 		InputFrameState->FrameDeltaTime.Set(DeltaSimTime);
 
 		// Do the actual update
-		State.SetMaxTickFrame(InputFrame);
 		TSimulationDoTickFunc<Model>::DoTick(*Simulation, State, ESimulationTickContext::ResimExtrapolate);
 	}
 
diff --git a/Engine/Plugins/Runtime/NetworkPrediction/Source/NetworkPrediction/Public/NetworkedSimulationModelTypes.h b/Engine/Plugins/Runtime/NetworkPrediction/Source/NetworkPrediction/Public/NetworkedSimulationModelTypes.h
--- a/Engine/Plugins/Runtime/NetworkPrediction/Source/NetworkPrediction/Public/NetworkedSimulationModelTypes.h
+++ b/Engine/Plugins/Runtime/NetworkPrediction/Source/NetworkPrediction/Public/NetworkedSimulationModelTypes.h
@@ -18,21 +18,76 @@
 //		-FrameBuffer:	frame-to-frame state
 //		-AuxBuffer:		sparse frame based state
 //		-CueDispatcher: non simulation affecting event history, mappings
-//		-TimingInfo:	time keeper for advancing simulation ticks
+//		-TotalAllowedSimulationTime : time keeper for advancing simulation ticks
 //	
 //	-Frame Markers
 //		-Pending		Next frame that will be ticked. Latest data.
-//		-LatestInput	Latest frame to have input written to it.
-//		-NextInput		Where next input should be written to.
-//		-MaxTick		Max frame that is allowed to be ticked. To facilitate input buffering when necessary.
+//		-LatestInput	Latest frame to have input written to it. Can be ahead of pending (input buffering), equal to, or 1 frame behind.
 //		-Confirmed		Last valid frame. Frames prior to confirmed frame do not need to be kept around.
 //
+//	The head frame of the SimulationState is Max(LatestInput, PendingTickFrame). LatestInput will be higher if you are buffering input.
+//	Otherwise, PendingTickFrame will be up to one frame ahead of LatestInputFrame. They will be equal prior to running SimulatinTick.
+//
+//	This structure supports resizing up to a MaxCapacity size. It will resize in order to preserve the 3 frame markers (when valid).
+//	E.g, if ConfirmedFrame=1 and we try to write Frame(1+BufferSize), the buffer will be resized so that ConfirmedFrame is preserved.
+//	We cannot resize, we will crash! It is the calling code (RepController's) responsibility to handle this by checking with 
+//	::GetMaxWritableFrame prior to calling a WriteFrame function.
+//
+//
+//
+// Examples
+// Note that this is how the rep controllers currently work. But ultimately they are what drives the simulation forward
+// and can set their own rules and own internal frame markers. 
+//
+
+//  Autonomous Proxy prior to running SimulationTick
+//
+//	Flags   TotalSimTime   DeltaTime   InputCmd   SyncState
+// |------|--------------|-----------|----------|-----------|
+// |   In |            X |        16 |        1 |         1 |	<-- Confirmed
+// |   In |         X+16 |        20 |        1 |         1 |   
+// |   In |         X+36 |        14 |        1 |         1 |	
+// |   In |         X+50 |        17 |        1 |         1 |   <-- PendingTick + LatestInput
+// |------|--------------|-----------|----------|-----------|
+
+
+//	Autonomous Proxy after running SimulationTick:
+//
+//	Flags   TotalSimTime   DeltaTime   InputCmd   SyncState
+// |------|--------------|-----------|----------|-----------|
+// |   In |            X |        16 |        1 |         1 |	<-- Confirmed
+// |   In |         X+16 |        20 |        1 |         1 |   
+// |   In |         X+36 |        14 |        1 |         1 |
+// |   In |         X+50 |        17 |        1 |         1 |   <-- LatestInput
+// |    0 |         X+67 |         0 |        0 |         1 |   <-- PendingTick
+// |------|--------------|-----------|----------|-----------|
+
+
+// Authority w/ buffered input
+//
+//	Flags   TotalSimTime   DeltaTime   InputCmd   SyncState
+// |------|--------------|-----------|----------|-----------|
+// |   In |            X |        16 |        1 |         1 |	<-- PendingTick
+// |   In |         X+16 |        20 |        1 |         0 |   
+// |   In |         X+36 |        14 |        1 |         0 |	
+// |    0 |            0 |         0 |        0 |         0 |   (Note gap frames! High packet loss caused these inputs to be missed)
+// |    0 |            0 |         0 |        0 |         0 |  
+// |   In |         X+50 |        17 |        1 |         0 |   <-- LatestInput
+// |------|--------------|-----------|----------|-----------|
+
 // ---------------------------------------------------------------------------------------------------------------------
 
+// Specialize-able struct for setting initial/max size of frame buffer. This should be tied into network role eventually.
+template<typename Model>
+struct TNetworkedSimStateBufferSizes
+{
+	static constexpr int32 InitialSize() { return 16; }
+	static constexpr int32 MaxSize() { return 128; }
+};
+
 template<typename Model>
 struct TNetworkedSimulationState
 {
-	// Collection of types we were assigned
 	using TBufferTypes = typename Model::BufferTypes;
 	using TickSettings = typename Model::TickSettings;
 
@@ -43,56 +98,82 @@ struct TNetworkedSimulationState
 	TNetSimCueDispatcher<Model> CueDispatcher;
 
 	TNetworkedSimulationState()
-		: FrameBuffer(64), AuxBuffer(64)
+		: FrameBuffer(InitialCapactiy()), AuxBuffer(InitialCapactiy())
+	{
+	}
+
+	int32 InitialCapactiy() const { return TNetworkedSimStateBufferSizes<Model>::InitialSize(); }
+	int32 MaxCapacity() const { return TNetworkedSimStateBufferSizes<Model>::MaxSize(); }
+
+	int32 GetCapacity() const { return FrameBuffer.Capacity(); }
+	
+	FString DebugString() const { return FString::Printf(TEXT("PendingTick: %d LatestInput: %d Size: %d. PT: %s. IT: %s. MT: %s"), 
+		PendingTickFrame, LatestInputFrame, FrameBuffer.Capacity(), *GetTotalProcessedSimulationTime().ToString(), *GetTotalInputSimulationTime().ToString(), *GetTotalAllowedSimulationTime().ToString()); }
+
+	// Highest frame that can be written without causing a system fault
+	int32 GetMaxWritableFrame() const
+	{
+		return GetMaxWritableFrame(GetMinFrameMarker());
+	}
+
+	int32 GetMaxWritableFrame(const int32 MinFrameMarker) const
 	{
+		return MinFrameMarker + MaxCapacity() - 1;
+	}
+	
+	int32 GetMinFrameMarker() const
+	{
+		// Calculate the lowest frame marker. Do this as uint32 to make INDEX_NONE values max
+		return (int32)(FMath::Min<uint32>(FMath::Min<uint32>(ConfirmedFrame, PendingTickFrame), LatestInputFrame));
 	}
 
-	int32 Capacity() const { return FrameBuffer.Capacity(); }
-	FString DebugString() const { return FString::Printf(TEXT("PendingTick: %d MaxTick: %d LatestInput: %d NextIntput: %d Size: %d. PT: %s. IT: %s. MT: %s"), 
-		PendingTickFrame, MaxTickFrame, LatestInputFrame, NextInputFrame, Capacity(), *GetTotalProcessedSimulationTime().ToString(), *GetTotalInputSimulationTime().ToString(), *GetTotalAllowedSimulationTime().ToString()); }
+	// Resets the FrameMarkers to NewPendingTickFrame. 
+	void ResetToFrame(int32 NewPendingTickFrame, bool bHasInput)
+	{
+		PendingTickFrame = NewPendingTickFrame;
+		ConfirmedFrame = NewPendingTickFrame;
+		LatestInputFrame = bHasInput ? NewPendingTickFrame : INDEX_NONE;		
+	}
+
+	// Returns a frame for writing. This will always return a valid frame. Expected to be used with ResetToFrame when reseting the system.
+	TSimulationFrameState<Model>* WriteFrameDirect(int32 Frame)
+	{
+		return &FrameBuffer[Frame];
+	}
 	
 	// ---------------------------------------------------------------------
 	// Frame State
 	// ---------------------------------------------------------------------
-
-	// Writes a new frame to the buffer, with the given explicit TotalSimTime.
-	// This is used in places where we might break sim time continuity (NetRecvs) or are possibly resimulating over a frame and want to explicitly set the TotalSimTime.
-	TSimulationFrameState<Model>* WriteFrameWithTime(int32 NewFrame, FNetworkSimTime TotalSimTime)
+	
+	// Called to create the next frame for simulation ticking
+	// Will return PendingTickFrame+1. Does not update frame markers.
+	TSimulationFrameState<Model>* WriteNextTickFrame()
 	{
-		if (TSimulationFrameState<Model>* ExistingFrame = GetValidFrame(NewFrame))
+		npEnsureMsgf(LatestInputFrame >= PendingTickFrame, TEXT("WriteNextTickFrame called with no input. %s"), *DebugString());
+		const int32 NewFrame = PendingTickFrame+1;
+
+		// Its possible the frame was already created (input buffering or resimulating)
+		if (LatestInputFrame >= NewFrame)
 		{
-			ExistingFrame->TotalSimTime = TotalSimTime;
-			return ExistingFrame;
+			return &FrameBuffer[NewFrame];
 		}
 		
-		const int32 Delta = NewFrame - PendingTickFrame;
-		if (Delta >= Capacity())
-		{
-			// PendingTickFrame must always be preserved
-			UE_NP_TRACE_SYSTEM_FAULT("WriteFrameWithTime frame %d outside of buffer range. Reseting contents. %s", PendingTickFrame, *DebugString());
-				
-			TSimulationFrameState<Model>& NewFrameState = FrameBuffer[NewFrame];
-			NewFrameState = FrameBuffer[PendingTickFrame]; // Init with previous state
-			NewFrameState.TotalSimTime = TotalSimTime;
-			NewFrameState.ClearFlags();
-				
-			PendingTickFrame = NewFrame;
-			MaxTickFrame = NewFrame-1;
-			LatestInputFrame = INDEX_NONE;
-			NextInputFrame = NewFrame;
-			CheckInvariants();
-
-			return &NewFrameState;
-		}
-			
-		TSimulationFrameState<Model>& NewFrameState = FrameBuffer[NewFrame];
-		NewFrameState.TotalSimTime = TotalSimTime;
+		// It is the caller's responsibility to not write to an invalid frame. Use GetMaxWritableFrame() prior to calling.
+		// The caller should then decide to either A) not write the frame or B) call ResetToFrame.
+		const int32 MinFrameMarker = GetMinFrameMarker();
+		npCheckf(NewFrame <= GetMaxWritableFrame(MinFrameMarker), TEXT("Invalid Frame Write. NewFrame: %d. %s"), NewFrame, *DebugString());		
+		
+		// Resize: we may need to resize the buffer to keep everything in the window.
+		const int32 Delta = NewFrame - MinFrameMarker;
 
-		for (int32 Frame = LatestInputFrame+1; Frame <= NewFrame; ++Frame)
+		if (Delta >= FrameBuffer.Capacity())
 		{
-			FrameBuffer[Frame].ClearFlags();
+			const int32 NewSize = Delta + 2; // will be rounded to next pow(2) in resize
+			Resize(NewSize);
 		}
 		
+		TSimulationFrameState<Model>& NewFrameState = FrameBuffer[NewFrame];
+		NewFrameState.ClearFlags();
 		return &NewFrameState;
 	}
 
@@ -102,55 +183,45 @@ struct TNetworkedSimulationState
 	// If the frame already exists, will just return the frame
 	// If this causes a new frame to be created, its flags will be cleared and TotalSimTime will be set for you. But contents will otherwise be stale
 	// If a gap is created, those frames will also have flags cleared
-	TSimulationFrameState<Model>* WriteFrameSequential(int32 NewFrame)
-	{
-		if (TSimulationFrameState<Model>* ExistingFrame = GetValidFrame(NewFrame))
-		{
-			CheckInvariants();
-			return ExistingFrame;
-		}
-				
-		const int32 Delta = NewFrame - PendingTickFrame;
-		if (Delta >= Capacity())
-		{
-			// PendingTickFrame must always be preserved
-			UE_NP_TRACE_SYSTEM_FAULT("WriteFrameSequential frame %d outside of buffer range. Reseting contents. %s", PendingTickFrame, *DebugString());
 
-			TSimulationFrameState<Model>& NewFrameState = FrameBuffer[NewFrame];
-			NewFrameState = FrameBuffer[PendingTickFrame]; // Init with previous state
+	// Writes a new frame for input. 
 
-			PendingTickFrame = NewFrame;
-			MaxTickFrame = NewFrame-1;
+	TSimulationFrameState<Model>* WriteInputFrame(int32 NewFrame)
+	{
+		npEnsureMsgf(NewFrame >= PendingTickFrame, TEXT("Input Frame write prior to PendingTickFrame. NewFrame: %d. %s"), NewFrame, *DebugString());
+		npEnsureMsgf(NewFrame > LatestInputFrame, TEXT("Input Frame write prior to LatestInputFrame. NewFrame: %d. %s"), NewFrame, *DebugString());
+		
+		const int32 PrevFrame = FMath::Max(PendingTickFrame, LatestInputFrame);
 
-			if (NewFrameState.HasFlag(ESimulationFrameStateFlags::InputWritten))
+		if (NewFrame > PendingTickFrame)
+		{
+			// It is the caller's responsibility to not write to an invalid frame. Use GetMaxWritableFrame() prior to calling.
+			// The caller should then decide to either A) not write the frame or B) call ResetToFrame. We can't recover this deep so this is a check.
+			const int32 MinFrameMarker = GetMinFrameMarker();
+			npCheckf(NewFrame <= GetMaxWritableFrame(MinFrameMarker), TEXT("Invalid Frame Write. NewFrame: %d. %s"), NewFrame, *DebugString());	
+
+			// Resize: we may need to resize the buffer to keep everything in the window.
+			const int32 Delta = NewFrame - MinFrameMarker;
+			if (Delta >= FrameBuffer.Capacity())
 			{
-				LatestInputFrame = NewFrame;
-				NextInputFrame = NewFrame+1;
+				const int32 NewSize = Delta + 2; // will be rounded to next pow(2) in resize
+				Resize(NewSize);
 			}
-			else
+
+			// Clear gap frame flags so they don't get processed
+			for (int32 Frame = PrevFrame+1; Frame < NewFrame; ++Frame)
 			{
-				LatestInputFrame = INDEX_NONE;
-				NextInputFrame = NewFrame;
+				FrameBuffer[Frame].ClearFlags();
 			}
-			
-			CheckInvariants();
-			return &NewFrameState;
-		}
 
-		// Create the new frame
-		TSimulationFrameState<Model>& NewFrameState = FrameBuffer[NewFrame];
-		NewFrameState.ClearFlags();
-
-		// Clear gap frame flags so they don't get processed
-		const int32 PrevFrame = FMath::Max(PendingTickFrame, LatestInputFrame);
-		for (int32 Frame = PrevFrame+1; Frame < NewFrame; ++Frame)
-		{
-			FrameBuffer[Frame].ClearFlags();
+			// We are creating a new frame so clear the flags
+			FrameBuffer[NewFrame].ClearFlags();
 		}
 
 		// Look at the previous frame sequential frame and calculate the new Frame's TotalSimTime
 		TSimulationFrameState<Model>& PrevFrameState = FrameBuffer[PrevFrame];
-		
+		TSimulationFrameState<Model>& NewFrameState = FrameBuffer[NewFrame];
+
 		if (PrevFrameState.HasFlag(ESimulationFrameStateFlags::InputWritten))
 		{
 			NewFrameState.TotalSimTime = PrevFrameState.FrameDeltaTime.Get() + PrevFrameState.TotalSimTime;
@@ -159,10 +230,14 @@ struct TNetworkedSimulationState
 		{
 			NewFrameState.TotalSimTime = PrevFrameState.TotalSimTime;
 		}
-		CheckInvariants();
+
+		npEnsureMsgf(!NewFrameState.HasFlag(ESimulationFrameStateFlags::InputWritten), TEXT("Input already written to Frame %d. %s"), NewFrame, *DebugString());
+		NewFrameState.SetFlag(ESimulationFrameStateFlags::InputWritten);
+		LatestInputFrame = NewFrame;
+		
 		return &NewFrameState;
 	}
-
+	
 	TSimulationFrameState<Model>* WriteFramePending()
 	{
 		return &FrameBuffer[PendingTickFrame];
@@ -227,58 +302,31 @@ struct TNetworkedSimulationState
 	// On server, it is the tail frame of the buffer: the oldest valid frame.
 	int32 GetConfirmedFrame() const { return ConfirmedFrame == INDEX_NONE ? TailFrame() : ConfirmedFrame; }
 	void SetConfirmedFrame(int32 Frame) { ConfirmedFrame = Frame; CheckInvariants(); }
-
-	// Max frame we are allowed to process. This is how local buffer can be done.
-	int32 GetMaxTickFrame() const { return MaxTickFrame; }
-	void SetMaxTickFrame(int32 Frame) { MaxTickFrame = Frame; CheckInvariants(); }
-
+	
 	// The last frame that had input written to it. This can be ahead of PendingTickFrame but never more than 1 frame behind.
 	int32 GetLatestInputFrame() const { return LatestInputFrame; }
 	void SetLatestInputFrame(int32 Frame)
 	{
 		LatestInputFrame = Frame;
-		NextInputFrame = Frame == INDEX_NONE ? PendingTickFrame : (Frame+1);
 		CheckInvariants(); 
 	}
 
 	// Where input is expected to be written to next. Gaps in NetSerialization can cause this to not be the case though.
-	int32 GetNextInputFrame() const { return NextInputFrame; }
-
-	// Resets the system to the existing NewPendingTickFrame. This is used when we need to hard reset things to a given authoritative state
-	void Reset(int32 NewPendingTickFrame)
-	{
-		PendingTickFrame = NewPendingTickFrame;
-		ConfirmedFrame = NewPendingTickFrame;
-		MaxTickFrame = INDEX_NONE;
-		LatestInputFrame = INDEX_NONE;
-		NextInputFrame = NewPendingTickFrame;
-
-		TotalAllowedSimulationTime = FrameBuffer[PendingTickFrame].TotalSimTime;
-		CheckInvariants();
-	}
+	int32 GetNextInputFrame() const { return LatestInputFrame != INDEX_NONE ? LatestInputFrame+1 : PendingTickFrame; }
 
 	void CheckInvariants()
 	{
-		npEnsureMsgf(IsValidFrame(PendingTickFrame), TEXT("PendingTickFrame not valid. %s"), *DebugString());
 		npEnsureMsgf(ConfirmedFrame <= PendingTickFrame, TEXT("ConfirmedFrame %d <= PendingTickFrame"), ConfirmedFrame, PendingTickFrame);
-
 		if (LatestInputFrame != INDEX_NONE)
 		{
-			npEnsureMsgf(MaxTickFrame <= LatestInputFrame, TEXT("MaxTickFrame %d <= LatestInputFrame %d"), MaxTickFrame, LatestInputFrame);
 			npEnsureMsgf(PendingTickFrame <= LatestInputFrame+1, TEXT("PendingTickFrame %d <= (LatestInputFrame+1) %d"), PendingTickFrame, LatestInputFrame+1);
-			npEnsureMsgf(LatestInputFrame == NextInputFrame-1, TEXT("LatestInputFrame %d != (NextInputFrame+1) %d"),  LatestInputFrame, NextInputFrame+1);			
-			npEnsureMsgf(IsValidFrame(LatestInputFrame) && FrameBuffer[LatestInputFrame].HasFlag(ESimulationFrameStateFlags::InputWritten), TEXT("LatestInputFrame not valid. %s"), *DebugString());
-		}
-
-		if (MaxTickFrame != INDEX_NONE)
-		{
-			npEnsureMsgf(PendingTickFrame <= MaxTickFrame+1, TEXT("PendingTickFrame %d <= (MaxTickFrame+1) %d"), PendingTickFrame, MaxTickFrame+1);
+			npEnsureMsgf(IsValidFrame(LatestInputFrame), TEXT("LatestInputFrame not valid. %s"), *DebugString());
+			npEnsureMsgf(FrameBuffer[LatestInputFrame].HasFlag(ESimulationFrameStateFlags::InputWritten), TEXT("LatestInputFrame doesnt actually have input %s"), *DebugString());
 		}
 
 #if NP_CHECKS_AND_ENSURES_SLOW
 		// If PendingTickFrame is ahead of input, ensure that calculated TotalSimTime is what we expected.
-		// (If this fires, its likely LatestInputFrame is wrong. If we took a frame from a NetSerialize for example,
-		// latest InputFrame should be reset.
+		// If this fires, its likely LatestInputFrame is wrong. If we took a frame from a NetSerialize for example, latest InputFrame should be reset.
 		if (LatestInputFrame != INDEX_NONE)
 		{
 			if (LatestInputFrame < PendingTickFrame)
@@ -298,8 +346,14 @@ struct TNetworkedSimulationState
 private:
 
 	int32 HeadFrame() const { return FMath::Max(PendingTickFrame, LatestInputFrame); }
-	int32 TailFrame() const { return FMath::Max(ConfirmedFrame, HeadFrame() - FrameBuffer.Capacity() + 1); }
-	bool IsValidFrame(int32 Frame) const { return Frame >= TailFrame() && Frame <= HeadFrame(); }
+	int32 TailFrame() const { return HeadFrame() - FrameBuffer.Capacity() + 1; }
+	
+	bool IsValidFrame(int32 Frame) const
+	{
+		const int32 Head = FMath::Max(PendingTickFrame, LatestInputFrame);
+		const int32 Tail = Head - FrameBuffer.Capacity() + 1;
+		return Frame >= TailFrame() && Frame <= HeadFrame();
+	}
 
 	const TSimulationFrameState<Model>* GetValidFrame(int32 Frame) const
 	{
@@ -311,12 +365,15 @@ private:
 		return IsValidFrame(Frame) ? &FrameBuffer[Frame] : nullptr;
 	}
 
+	void Resize(int32 NewSize)
+	{
+		FrameBuffer.Resize(NewSize, HeadFrame());
+		AuxBuffer.Resize(NewSize);
+	}
+
 	bool bTickInProgress = false;
 
 	int32 PendingTickFrame = 0;
-	int32 NextInputFrame = 0;
-
-	int32 MaxTickFrame = INDEX_NONE;
 	int32 LatestInputFrame = INDEX_NONE;
 	int32 ConfirmedFrame = INDEX_NONE;
 	
@@ -475,8 +532,9 @@ struct TNetSimOutput
 };
 
 // Helper to actual tick a simulation. This does all the setup, runs ::SimulationTick, and finalizes the frame.
-// Caller is responsible to have the input buffers in the correct state: Input/Sync/Aux buffers must have valid data in OutputFrame-1.
-// DeltaTime comes from the InputCmd->GetFrameDeltaTime().
+// Caller is responsible for ensuring State.PendingTickFrame() has valid input and that State.PendingTickFrame()+1 <= State.GetMaxWritableFrame()
+//		(E.g, don't call this is it would age out ConfirmedFrame!)
+//
 // Important to note that this advances the PendingFrame to the output Frame. So that any writes that occur to the buffers during this scope will go to the output frame.
 template<typename Model>
 struct TSimulationDoTickFunc
@@ -499,19 +557,22 @@ struct TSimulationDoTickFunc
 		const int32 InputFrame = State.GetPendingTickFrame();
 		const int32 OutputFrame = State.GetPendingTickFrame()+1;
 
+		// Write the output frame (important to do before caching internal data since the write could cause a resize)
+		TFrameState* OutFrameState = State.WriteNextTickFrame();
+
 		// Get inputs. These all need to be present to tick (caller's responsibility)
 		TFrameState* InFrameState = State.ReadFrame(InputFrame);
 		const TAuxState* InAuxState = State.ReadAux(InputFrame);
 		
 		check(InFrameState);
 		check(InAuxState);
-		ensure(InFrameState->HasFlag(ESimulationFrameStateFlags::InputWritten));
+		npEnsureMsgf(InFrameState->HasFlag(ESimulationFrameStateFlags::InputWritten), TEXT("DoTick called with invalid InputFrame: %d. %s"), InputFrame, *State.DebugString());
 
 		const FNetworkSimTime DeltaSimTime = InFrameState->FrameDeltaTime.Get();
 		const FNetSimTimeStep SimTimeStep { DeltaSimTime, InFrameState->TotalSimTime, OutputFrame };
-		const FNetworkSimTime OutSimTime = InFrameState->TotalSimTime + DeltaSimTime;
-
-		TFrameState* OutFrameState = State.WriteFrameWithTime(OutputFrame, OutSimTime);
+		OutFrameState->TotalSimTime = InFrameState->TotalSimTime + DeltaSimTime;
+		
+		npCheckf(OutputFrame <= State.GetMaxWritableFrame(), TEXT("DoTick called on invalid OutputFrame: %d. %s"), OutputFrame, *State.DebugString());
 
 		// PreTick
 		State.SetPendingTickFrame(OutputFrame);
diff --git a/Engine/Plugins/Runtime/NetworkPrediction/Source/NetworkPrediction/Public/Trace/NetworkPredictionTrace.h b/Engine/Plugins/Runtime/NetworkPrediction/Source/NetworkPrediction/Public/Trace/NetworkPredictionTrace.h
--- a/Engine/Plugins/Runtime/NetworkPrediction/Source/NetworkPrediction/Public/Trace/NetworkPredictionTrace.h
+++ b/Engine/Plugins/Runtime/NetworkPrediction/Source/NetworkPrediction/Public/Trace/NetworkPredictionTrace.h
@@ -152,7 +152,7 @@ public:
 	template<typename Model>
 	static void TraceSimulationEOF(const TNetworkedSimulationState<Model>& State)
 	{
-		TraceEOF_Internal(State.Capacity(), State.GetPendingTickFrame(), State.GetLatestInputFrame(), State.GetMaxTickFrame(), State.GetTotalProcessedSimulationTime(), State.GetTotalAllowedSimulationTime());
+		TraceEOF_Internal(State.GetCapacity(), State.GetPendingTickFrame(), State.GetLatestInputFrame(), State.GetTotalProcessedSimulationTime(), State.GetTotalAllowedSimulationTime());
 	}
 
 private:
@@ -161,7 +161,7 @@ private:
 	static void PopSimulationId();
 
 	static void TraceUserState_Internal(int32 Frame, ETraceUserState StateType);
-	static void TraceEOF_Internal(int32 BufferSize, int32 PendingTickFrame, int32 LatestInputFrame, int32 MaxTickFrame, FNetworkSimTime TotalSimTime, FNetworkSimTime AllowedSimTime);
+	static void TraceEOF_Internal(int32 BufferSize, int32 PendingTickFrame, int32 LatestInputFrame, FNetworkSimTime TotalSimTime, FNetworkSimTime AllowedSimTime);
 	static void TraceOOBStrInternal(int32 SimulationId, ETraceUserState StateType, const TCHAR* Fmt);
 
 	static FStringOutputDevice StrOut;
diff --git a/Engine/Plugins/Runtime/NetworkPredictionInsights/Source/NetworkPredictionInsights/Private/NetworkPredictionAnalyzer.cpp b/Engine/Plugins/Runtime/NetworkPredictionInsights/Source/NetworkPredictionInsights/Private/NetworkPredictionAnalyzer.cpp
--- a/Engine/Plugins/Runtime/NetworkPredictionInsights/Source/NetworkPredictionInsights/Private/NetworkPredictionAnalyzer.cpp
+++ b/Engine/Plugins/Runtime/NetworkPredictionInsights/Source/NetworkPredictionInsights/Private/NetworkPredictionAnalyzer.cpp
@@ -118,7 +118,6 @@ bool FNetworkPredictionAnalyzer::OnEvent(uint16 RouteId, const FOnEventContext&
 			FrameData.BufferSize = EventData.GetValue<int32>("BufferSize");
 			FrameData.PendingTickFrame = EventData.GetValue<int32>("PendingTickFrame");
 			FrameData.LatestInputFrame = EventData.GetValue<int32>("LatestInputFrame");
-			FrameData.MaxTickFrame = EventData.GetValue<int32>("MaxTickFrame");
 			FrameData.TotalSimTime = EventData.GetValue<int32>("TotalSimTime");
 			FrameData.AllowedSimTime = EventData.GetValue<int32>("AllowedSimTime");	
 			break;
diff --git a/Engine/Plugins/Runtime/NetworkPredictionInsights/Source/NetworkPredictionInsights/Private/UI/SNPSimFrameContents.cpp b/Engine/Plugins/Runtime/NetworkPredictionInsights/Source/NetworkPredictionInsights/Private/UI/SNPSimFrameContents.cpp
--- a/Engine/Plugins/Runtime/NetworkPredictionInsights/Source/NetworkPredictionInsights/Private/UI/SNPSimFrameContents.cpp
+++ b/Engine/Plugins/Runtime/NetworkPredictionInsights/Source/NetworkPredictionInsights/Private/UI/SNPSimFrameContents.cpp
@@ -215,16 +215,14 @@ void SNPSimFrameContents::NotifyContentClicked(const FSimContentsView& InContent
 			"BufferSize: {2}\n"
 			"PendingTickFrame: {3}\n"
 			"LatestInputFrame: {4}\n"
-			"MaxTickFrame: {5}\n"
-			"TotalSimTime: {6}\n"
-			"AllowedSimTime: {7}"),
+			"TotalSimTime: {5}\n"
+			"AllowedSimTime: {6}"),
 			{
 				EOFState->EngineFrame,
 				EOFState->EngineFrameDeltaTime,
 				EOFState->BufferSize,
 				EOFState->PendingTickFrame,
 				EOFState->LatestInputFrame,
-				EOFState->MaxTickFrame,
 				EOFState->TotalSimTime,
 				EOFState->AllowedSimTime
 			});
diff --git a/Engine/Plugins/Runtime/NetworkPredictionInsights/Source/NetworkPredictionInsights/Public/INetworkPredictionProvider.h b/Engine/Plugins/Runtime/NetworkPredictionInsights/Source/NetworkPredictionInsights/Public/INetworkPredictionProvider.h
--- a/Engine/Plugins/Runtime/NetworkPredictionInsights/Source/NetworkPredictionInsights/Public/INetworkPredictionProvider.h
+++ b/Engine/Plugins/Runtime/NetworkPredictionInsights/Source/NetworkPredictionInsights/Public/INetworkPredictionProvider.h
@@ -383,7 +383,6 @@ struct FSimulationData
 		int32 BufferSize;
 		int32 PendingTickFrame;
 		int32 LatestInputFrame;
-		int32 MaxTickFrame;
 		FSimTime TotalSimTime;
 		FSimTime AllowedSimTime;
 
