Prevent warnings from occuring when multiple Network Prediction Model Defs have the same priority.
This will happen more and more frequently as more Network Predicted types are created, especially across different plugins and modules.

diff --git a/Engine/Plugins/Runtime/NetworkPrediction/Source/NetworkPrediction/Public/NetworkPredictionModelDefRegistry.h b/Engine/Plugins/Runtime/NetworkPrediction/Source/NetworkPrediction/Public/NetworkPredictionModelDefRegistry.h
--- a/Engine/Plugins/Runtime/NetworkPrediction/Source/NetworkPrediction/Public/NetworkPredictionModelDefRegistry.h
+++ b/Engine/Plugins/Runtime/NetworkPrediction/Source/NetworkPrediction/Public/NetworkPredictionModelDefRegistry.h
@@ -49,7 +49,7 @@ public:
 		{
 			if (LHS.SortPriority == RHS.SortPriority)
 			{
-				UE_LOG(LogNetworkPrediction, Warning, TEXT("ModelDefs %s and %s have same sort priority. Using lexical sort as backup"), LHS.Name, RHS.Name);
+				UE_LOG(LogNetworkPrediction, Log, TEXT("ModelDefs %s and %s have same sort priority. Using lexical sort as backup"), LHS.Name, RHS.Name);
 				int32 StrCmpResult = FCString::Strcmp(LHS.Name, RHS.Name);
 				npEnsureMsgf(StrCmpResult != 0, TEXT("Duplicate ModelDefs appear to have been registered."));
 				return StrCmpResult > 0;
