Change max simulation time step for Network Prediction to a CVar instead of a fixed 100 milliseconds.
np.MaxSimulationStepMS

diff --git a/Engine/Plugins/Runtime/NetworkPrediction/Source/NetworkPrediction/Private/NetworkPredictionWorldManager.cpp b/Engine/Plugins/Runtime/NetworkPrediction/Source/NetworkPrediction/Private/NetworkPredictionWorldManager.cpp
--- a/Engine/Plugins/Runtime/NetworkPrediction/Source/NetworkPrediction/Private/NetworkPredictionWorldManager.cpp
+++ b/Engine/Plugins/Runtime/NetworkPrediction/Source/NetworkPrediction/Private/NetworkPredictionWorldManager.cpp
@@ -15,6 +15,15 @@
 // Do extra checks to make sure Physics and GameThread (PrimitiveComponent) are in sync at verious points in the rollback process
 #define NP_ENSURE_PHYSICS_GT_SYNC !(UE_BUILD_SHIPPING || UE_BUILD_TEST)
 
+namespace NetworkPredictionWorldManagerCVars
+{
+	static int32 MaxStepMS = 100;
+	static FAutoConsoleVariableRef CVarMaxSimulationStepMS(
+		TEXT("np.MaxSimulationStepMS"),
+		MaxStepMS,
+		TEXT("The maximum time step (in integer milliseconds) that we will allow for simulation frames. Frame delta times will be clamped between 1 MS and this value.."));
+}
+
 UNetworkPredictionWorldManager* UNetworkPredictionWorldManager::ActiveInstance=nullptr;
 
 // -----------------------------------------------------------------------------------------------
@@ -335,13 +344,12 @@ void UNetworkPredictionWorldManager::BeginNewSimulationFrame(UWorld* InWorld, EL
 	{
 		// Update VariableTickState	
 		constexpr int32 MinStepMS = 1;
-		constexpr int32 MaxStepMS = 100;
 
 		VariableTickState.UnspentTimeMS += fEngineFrameDeltaTimeMS;
 		float fDeltaMS = FMath::FloorToFloat(VariableTickState.UnspentTimeMS);
 		VariableTickState.UnspentTimeMS -= fDeltaMS;
 
-		const int32 DeltaSimMS = FMath::Clamp((int32)fDeltaMS, MinStepMS, MaxStepMS);
+		const int32 DeltaSimMS = FMath::Clamp((int32)fDeltaMS, MinStepMS, NetworkPredictionWorldManagerCVars::MaxStepMS);
 
 		FVariableTickState::FFrame& PendingFrameData = VariableTickState.Frames[VariableTickState.PendingFrame];
 		PendingFrameData.DeltaMS = DeltaSimMS;
