Swapped the NetworkPredicion plugin to trace its strings as string-type fields instead of as trace event attachments.

#rnx

diff --git a/Engine/Plugins/Runtime/NetworkPrediction/Source/NetworkPrediction/Private/NetworkPredictionTrace.cpp b/Engine/Plugins/Runtime/NetworkPrediction/Source/NetworkPrediction/Private/NetworkPredictionTrace.cpp
--- a/Engine/Plugins/Runtime/NetworkPrediction/Source/NetworkPrediction/Private/NetworkPredictionTrace.cpp
+++ b/Engine/Plugins/Runtime/NetworkPrediction/Source/NetworkPrediction/Private/NetworkPredictionTrace.cpp
@@ -40,6 +40,7 @@ UE_TRACE_EVENT_END()
 UE_TRACE_EVENT_BEGIN(NetworkPrediction, SimulationCreated)
 	UE_TRACE_EVENT_FIELD(uint32, SimulationID) // server assigned (shared client<->server)
 	UE_TRACE_EVENT_FIELD(int32, TraceID) // process unique id
+	UE_TRACE_EVENT_FIELD(UE::Trace::AnsiString, DebugName)
 UE_TRACE_EVENT_END()
 
 UE_TRACE_EVENT_BEGIN(NetworkPrediction, SimulationConfig)
@@ -78,6 +79,7 @@ UE_TRACE_EVENT_END()
 
 // General system fault. Log message is in attachment
 UE_TRACE_EVENT_BEGIN(NetworkPrediction, SystemFault)
+	UE_TRACE_EVENT_FIELD(UE::Trace::WideString, Message)
 UE_TRACE_EVENT_END()
 
 // Traces general tick state (called before ticking N sims)
@@ -103,6 +105,7 @@ UE_TRACE_EVENT_BEGIN(NetworkPrediction, ShouldReconcile)
 UE_TRACE_EVENT_END()
 
 UE_TRACE_EVENT_BEGIN(NetworkPrediction, Reconcile)
+	UE_TRACE_EVENT_FIELD(UE::Trace::AnsiString, UserString)
 UE_TRACE_EVENT_END()
 
 UE_TRACE_EVENT_BEGIN(NetworkPrediction, RollbackInject)
@@ -130,18 +133,23 @@ UE_TRACE_EVENT_END()
 UE_TRACE_EVENT_BEGIN(NetworkPrediction, OOBStateMod)
 	UE_TRACE_EVENT_FIELD(int32, TraceID)
 	UE_TRACE_EVENT_FIELD(int32, Frame)
+	UE_TRACE_EVENT_FIELD(UE::Trace::AnsiString, Source)
 UE_TRACE_EVENT_END()
 
 UE_TRACE_EVENT_BEGIN(NetworkPrediction, InputCmd)
+	UE_TRACE_EVENT_FIELD(UE::Trace::AnsiString, Value)
 UE_TRACE_EVENT_END()
 
 UE_TRACE_EVENT_BEGIN(NetworkPrediction, SyncState)
+	UE_TRACE_EVENT_FIELD(UE::Trace::AnsiString, Value)
 UE_TRACE_EVENT_END()
 
 UE_TRACE_EVENT_BEGIN(NetworkPrediction, AuxState)
+	UE_TRACE_EVENT_FIELD(UE::Trace::AnsiString, Value)
 UE_TRACE_EVENT_END()
 
 UE_TRACE_EVENT_BEGIN(NetworkPrediction, PhysicsState)
+	UE_TRACE_EVENT_FIELD(UE::Trace::AnsiString, Value)
 UE_TRACE_EVENT_END()
 
 // ---------------------------------------------------------------------------
@@ -150,10 +158,10 @@ void FNetworkPredictionTrace::TraceSimulationCreated_Internal(FNetworkPrediction
 {
 	const uint16 AttachmentSize = Builder.Len() * sizeof(FStringBuilderBase::ElementType);
 
-	UE_TRACE_LOG(NetworkPrediction, SimulationCreated, NetworkPredictionChannel, AttachmentSize)
+	UE_TRACE_LOG(NetworkPrediction, SimulationCreated, NetworkPredictionChannel)
 		<< SimulationCreated.SimulationID((int32)ID)
 		<< SimulationCreated.TraceID(ID.GetTraceID())
-		<< SimulationCreated.Attachment(Builder.ToString(), AttachmentSize);
+		<< SimulationCreated.DebugName(Builder.ToString(), Builder.Len());
 }
 
 void FNetworkPredictionTrace::TraceWorldFrameStart(UGameInstance* GameInstance, float DeltaSeconds)
@@ -209,33 +217,30 @@ void FNetworkPredictionTrace::TraceSimTick(int32 TraceID)
 
 void FNetworkPredictionTrace::TraceUserState_Internal(ETraceUserState StateType, FAnsiStringBuilderBase& Builder)
 {
-	// Due to how string store works on the analysis side, its better to transmit the string with the null terminator
-	const uint16 AttachmentSize = (uint16)((Builder.Len()+1) * sizeof(FAnsiStringBuilderBase::ElementType));
-
 	switch(StateType)
 	{
 		case ETraceUserState::Input:
 		{
-			UE_TRACE_LOG(NetworkPrediction, InputCmd, NetworkPredictionChannel, AttachmentSize)
-				<< InputCmd.Attachment(Builder.GetData(), AttachmentSize);
+			UE_TRACE_LOG(NetworkPrediction, InputCmd, NetworkPredictionChannel)
+				<< InputCmd.Value(Builder.GetData(), Builder.Len());
 			break;
 		}
 		case ETraceUserState::Sync:
 		{
-			UE_TRACE_LOG(NetworkPrediction, SyncState, NetworkPredictionChannel, AttachmentSize)
-				<< SyncState.Attachment(Builder.GetData(), AttachmentSize);
+			UE_TRACE_LOG(NetworkPrediction, SyncState, NetworkPredictionChannel)
+				<< SyncState.Value(Builder.GetData(), Builder.Len());
 			break;
 		}
 		case ETraceUserState::Aux:
 		{
-			UE_TRACE_LOG(NetworkPrediction, AuxState, NetworkPredictionChannel, AttachmentSize)
-				<< AuxState.Attachment(Builder.GetData(), AttachmentSize);
+			UE_TRACE_LOG(NetworkPrediction, AuxState, NetworkPredictionChannel)
+				<< AuxState.Value(Builder.GetData(), Builder.Len());
 			break;
 		}
 		case ETraceUserState::Physics:
 		{
-			UE_TRACE_LOG(NetworkPrediction, PhysicsState, NetworkPredictionChannel, AttachmentSize)
-				<< PhysicsState.Attachment(Builder.GetData(), AttachmentSize);
+			UE_TRACE_LOG(NetworkPrediction, PhysicsState, NetworkPredictionChannel)
+				<< PhysicsState.Value(Builder.GetData(), Builder.Len());
 			break;
 		}
 	}
@@ -250,10 +255,8 @@ void FNetworkPredictionTrace::TraceNetRecv(int32 Frame, int32 TimeMS)
 
 void FNetworkPredictionTrace::TraceReconcile(const FAnsiStringView& StrView)
 {
-	const uint16 AttachmentSize = (uint16)((StrView.Len()+1) * sizeof(FAnsiStringView::ElementType));
-
-	UE_TRACE_LOG(NetworkPrediction, Reconcile, NetworkPredictionChannel, AttachmentSize)
-		<< Reconcile.Attachment(StrView.GetData(), AttachmentSize);
+	UE_TRACE_LOG(NetworkPrediction, Reconcile, NetworkPredictionChannel)
+		<< Reconcile.UserString(StrView.GetData(), StrView.Len());
 }
 
 void FNetworkPredictionTrace::TraceShouldReconcile(int32 TraceID)
@@ -311,12 +314,10 @@ void FNetworkPredictionTrace::TraceProduceInput(int32 TraceID)
 
 void FNetworkPredictionTrace::TraceOOBStateMod(int32 TraceID, int32 Frame, const FAnsiStringView& StrView)
 {
-	const uint16 AttachmentSize = (uint16)((StrView.Len()+1) * sizeof(FAnsiStringView::ElementType));
-
-	UE_TRACE_LOG(NetworkPrediction, OOBStateMod, NetworkPredictionChannel, AttachmentSize)
+	UE_TRACE_LOG(NetworkPrediction, OOBStateMod, NetworkPredictionChannel)
 		<< OOBStateMod.TraceID(TraceID)
 		<< OOBStateMod.Frame(Frame)
-		<< OOBStateMod.Attachment(StrView.GetData(), AttachmentSize);
+		<< OOBStateMod.Source(StrView.GetData(), StrView.Len());
 }
 
 #include "CoreTypes.h"
@@ -362,13 +363,12 @@ void FNetworkPredictionTrace::TraceSystemFault(const TCHAR* Fmt, ...)
 	GROWABLE_LOGF( 
 
 		check(Result >= 0 );
-	const uint16 AttachmentSize = (Result+1) * sizeof(TCHAR);
 
 	UE_LOG(LogNetworkPrediction, Warning, TEXT("SystemFault: %s"), Buffer);
 	);
 
-	UE_TRACE_LOG(NetworkPrediction, SystemFault, NetworkPredictionChannel, AttachmentSize)
-		<< SystemFault.Attachment(Buffer, AttachmentSize);
+	UE_TRACE_LOG(NetworkPrediction, SystemFault, NetworkPredictionChannel)
+		<< SystemFault.Message(Buffer, Result);
 
 	FMemory::SystemFree(AllocatedBuffer);
 }
\ No newline at end of file
diff --git a/Engine/Plugins/Runtime/NetworkPredictionInsights/Source/NetworkPredictionInsights/Private/NetworkPredictionAnalyzer.cpp b/Engine/Plugins/Runtime/NetworkPredictionInsights/Source/NetworkPredictionInsights/Private/NetworkPredictionAnalyzer.cpp
--- a/Engine/Plugins/Runtime/NetworkPredictionInsights/Source/NetworkPredictionInsights/Private/NetworkPredictionAnalyzer.cpp
+++ b/Engine/Plugins/Runtime/NetworkPredictionInsights/Source/NetworkPredictionInsights/Private/NetworkPredictionAnalyzer.cpp
@@ -3,6 +3,7 @@
 #include "NetworkPredictionAnalyzer.h"
 #include "NetworkPredictionProvider.h"
 #include "Containers/StringView.h"
+#include "TraceServices/Utils.h"
 
 // As we are tracing from multiple threads and we have trace events that are stringed together we need to track some state per thread to not trace bad data.
 struct FNetworkPredictionAnalyzer::FThreadState
@@ -73,10 +74,8 @@ bool FNetworkPredictionAnalyzer::OnEvent(uint16 RouteId, EStyle Style, const FOn
 	{
 		FThreadState& ThreadState = GetThreadState(Context.ThreadInfo.GetId());
 
-		// FIXME: we want to trace ANSI strings but are converting them to WIDE here to store them for analysis.
-		// Latest version of insights should have formal ANSI string support so this should be able to go away.
-		const ANSICHAR* AnsiStr = reinterpret_cast<const ANSICHAR*>(EventData.GetAttachment());
-		NetworkPredictionProvider.WriteUserState(ThreadState.TraceID, ThreadState.PendingWriteFrame, this->EngineFrameNumber, Type, ThreadState.PendingUserStateSource, Session.StoreString(StringCast<TCHAR>(AnsiStr).Get()));
+		FString Value = TraceServices::FTraceAnalyzerUtils::LegacyAttachmentString<ANSICHAR>("Value", Context);
+		NetworkPredictionProvider.WriteUserState(ThreadState.TraceID, ThreadState.PendingWriteFrame, this->EngineFrameNumber, Type, ThreadState.PendingUserStateSource, Session.StoreString(*Value));
 	};
 
 	bool bUnhandled = false;
@@ -107,7 +106,7 @@ bool FNetworkPredictionAnalyzer::OnEvent(uint16 RouteId, EStyle Style, const FOn
 			ThreadState.TraceID = EventData.GetValue<int32>("TraceID");
 
 			FSimulationData::FConst& ConstData = NetworkPredictionProvider.WriteSimulationCreated(ThreadState.TraceID);
-			ConstData.DebugName = FString(EventData.GetAttachmentSize() / sizeof(TCHAR), reinterpret_cast<const TCHAR*>(EventData.GetAttachment()));
+			ConstData.DebugName = TraceServices::FTraceAnalyzerUtils::LegacyAttachmentString<ANSICHAR>("DebugName", Context);
 			ConstData.ID.SimID = EventData.GetValue<uint32>("SimulationID");
 			break;
 		}
@@ -162,7 +161,8 @@ bool FNetworkPredictionAnalyzer::OnEvent(uint16 RouteId, EStyle Style, const FOn
 
 			FThreadState& ThreadState = GetThreadState(Context.ThreadInfo.GetId());
 
-			const TCHAR* StoredString = Session.StoreString(reinterpret_cast<const TCHAR*>(EventData.GetAttachment()));
+			FString Message = TraceServices::FTraceAnalyzerUtils::LegacyAttachmentString<WIDECHAR>("Message", Context);
+			const TCHAR* StoredString = Session.StoreString(*Message);
 			ensureMsgf(ThreadState.TraceID > 0, TEXT("Invalid TraceID when analyzing SystemFault: %s"), StoredString);
 
 			NetworkPredictionProvider.WriteSystemFault(
@@ -264,11 +264,8 @@ bool FNetworkPredictionAnalyzer::OnEvent(uint16 RouteId, EStyle Style, const FOn
 			// Valid TraceID should have already been set
 			ensure(ThreadState.TraceID > 0);
 
-			// FIXME: ANSI to TCHAR conversion should be removable. See UserState comment above	
-			const ANSICHAR* AttachmentData = reinterpret_cast<const ANSICHAR*>(EventData.GetAttachment());
-			int32 AttachmentSize = EventData.GetAttachmentSize();
-			
-			NetworkPredictionProvider.WriteReconcileStr(ThreadState.TraceID, Session.StoreString(StringCast<TCHAR>(AttachmentData, AttachmentSize).Get()));
+			FString UserString = TraceServices::FTraceAnalyzerUtils::LegacyAttachmentString<ANSICHAR>("UserString", Context);
+			NetworkPredictionProvider.WriteReconcileStr(ThreadState.TraceID, Session.StoreString(*UserString));
 			break;
 		}
 
@@ -333,11 +330,8 @@ bool FNetworkPredictionAnalyzer::OnEvent(uint16 RouteId, EStyle Style, const FOn
 			ThreadState.PendingUserStateSource = ENP_UserStateSource::OOB;
 			NetworkPredictionProvider.WriteOOBStateMod(ThreadState.TraceID);
 
-			// FIXME: ANSI to TCHAR conversion should be removable. See UserState comment above	
-			const ANSICHAR* AttachmentData = reinterpret_cast<const ANSICHAR*>(EventData.GetAttachment());
-			const int32 AttachmentSize = EventData.GetAttachmentSize();
-			
-			NetworkPredictionProvider.WriteOOBStateModStr(ThreadState.TraceID, Session.StoreString(StringCast<TCHAR>(AttachmentData, AttachmentSize).Get()));
+			FString Source = TraceServices::FTraceAnalyzerUtils::LegacyAttachmentString<ANSICHAR>("Source", Context);
+			NetworkPredictionProvider.WriteOOBStateModStr(ThreadState.TraceID, Session.StoreString(Source));
 			break;
 		}
 	}
