NetworkPrediction: Replacing a GetPing()*4 with a call to GetPingInMilliseconds

diff --git a/Engine/Plugins/Runtime/NetworkPrediction/Source/NetworkPrediction/Private/NetworkPhysics.cpp b/Engine/Plugins/Runtime/NetworkPrediction/Source/NetworkPrediction/Private/NetworkPhysics.cpp
--- a/Engine/Plugins/Runtime/NetworkPrediction/Source/NetworkPrediction/Private/NetworkPhysics.cpp
+++ b/Engine/Plugins/Runtime/NetworkPrediction/Source/NetworkPrediction/Private/NetworkPhysics.cpp
@@ -1344,7 +1344,7 @@ void UNetworkPhysicsManager::TickDrawDebug()
 
 				if (APlayerState* PS = LocalPC->GetPlayerState<APlayerState>())
 				{
-					int32 ActualPing = PS->GetPing() * 4;
+					int32 ActualPing = (int32)PS->GetPingInMilliseconds();
 					YPos += Canvas->DrawText(GEngine->GetMediumFont(), FString::Printf(TEXT("Ping: %d"), ActualPing), XPos, YPos);
 				}
 
