Network Prediction Plugin
-Edigrate source changes over from UE4 plus some additional fixes

#rnx

diff --git a/Engine/Plugins/Runtime/NetworkPrediction/Config/BaseEngine.ini b/Engine/Plugins/Runtime/NetworkPrediction/Config/BaseEngine.ini
new file mode 100644
--- /dev/null
+++ b/Engine/Plugins/Runtime/NetworkPrediction/Config/BaseEngine.ini
@@ -0,0 +1,2 @@
+[CoreRedirects]
++ClassRedirects=(OldName="/Script/NetworkPredictionExtras.NetworkPhysicsComponent", NewName="/Script/NetworkPrediction.NetworkPhysicsComponent")
diff --git a/Engine/Plugins/Runtime/NetworkPrediction/Config/BaseNetworkPrediction.ini b/Engine/Plugins/Runtime/NetworkPrediction/Config/BaseNetworkPrediction.ini
--- a/Engine/Plugins/Runtime/NetworkPrediction/Config/BaseNetworkPrediction.ini
+++ b/Engine/Plugins/Runtime/NetworkPrediction/Config/BaseNetworkPrediction.ini
@@ -2,4 +2,4 @@
 Settings=(PreferredTickingPolicy=Fixed,ReplicatedManagerClassOverride=None,FixedTickFrameRate=60,bForceEngineFixTickForcePhysics=True,SimulatedProxyNetworkLOD=ForwardPredict,FixedTickInterpolationBufferedMS=100,IndependentTickInterpolationBufferedMS=100,IndependentTickInterpolationMaxBufferedMS=250)
 +DevHUDs=(HUDName="Debug",Items=((DisplayName="Debug ON",ExecCommand="np2.debug 1",bAutoBack=False),(DisplayName="Debug OFF",ExecCommand="np2.debug 0",bAutoBack=False),(DisplayName="Debug Draw Tolerance ON",ExecCommand="np2.Debug.Tolerance 1",bAutoBack=False),(DisplayName="Debug Draw Tolerance OFF",ExecCommand="np2.debug.Tolerance 0",bAutoBack=False),(DisplayName="Dump",ExecCommand="np2.dump",bAutoBack=False,bRequirePIE=True)),bRequirePIE=False,bRequireNotPIE=False)
 +DevHUDs=(HUDName="Input Macros",Items=((DisplayName="Start/Stop Record",ExecCommand="np2.RecordInputs Temp",bAutoBack=False),(DisplayName="Playback",ExecCommand="np2.PlaybackInputs Temp",bAutoBack=False)),bRequirePIE=False,bRequireNotPIE=False)
-+DevHUDs=(HUDName="Change Map",Items=((DisplayName="Empty",ExecCommand="serverexec servertravel Physics2_Empty"),(DisplayName="Platforms Small",ExecCommand="serverexec servertravel Physics2_Obstacle_Constraint"),(DisplayName="Obstacle Course",ExecCommand="serverexec servertravel Physics2_ObstacleCourse"),(DisplayName="Soccer",ExecCommand="serverexec servertravel Soccer"),(DisplayName="Soccer Empty",ExecCommand="serverexec servertravel Physics2_soccerball_empty"),(DisplayName="Portals",ExecCommand="serverexec servertravel ObstacleCourse"),(DisplayName="Obstacle Course",ExecCommand="serverexec servertravel ObstacleCourse")),bRequirePIE=False,bRequireNotPIE=True)
++DevHUDs=(HUDName="Change Map",Items=((DisplayName="Empty",ExecCommand="serverexec servertravel Physics2_Empty"),(DisplayName="Physics Tracks",ExecCommand="serverexec servertravel TestMap_Physics_Tracks"),(DisplayName="Physics Tracks 2",ExecCommand="serverexec servertravel TestMap_Physics_Tracks2"),(DisplayName="Soccer",ExecCommand="serverexec servertravel Soccer"),(DisplayName="Soccer Empty",ExecCommand="serverexec servertravel Physics2_soccerball_empty"),(DisplayName="Portals",ExecCommand="serverexec servertravel Portals"),(DisplayName="Cube DM",ExecCommand="serverexec servertravel TestMap_Cube_DM"),(DisplayName="Swings",ExecCommand="serverexec servertravel TestMap_Swings"),(DisplayName="HoverShips",ExecCommand="serverexec servertravel Physics2_HoverShips")),bRequirePIE=False,bRequireNotPIE=true)
diff --git a/Engine/Plugins/Runtime/NetworkPrediction/Source/NetworkPrediction/NetworkPrediction.Build.cs b/Engine/Plugins/Runtime/NetworkPrediction/Source/NetworkPrediction/NetworkPrediction.Build.cs
--- a/Engine/Plugins/Runtime/NetworkPrediction/Source/NetworkPrediction/NetworkPrediction.Build.cs
+++ b/Engine/Plugins/Runtime/NetworkPrediction/Source/NetworkPrediction/NetworkPrediction.Build.cs
@@ -51,7 +51,6 @@ namespace UnrealBuildTool.Rules
                 PrivateDependencyModuleNames.AddRange(
                 new string[]
                 {
-					"EditorFramework",
                     "UnrealEd",
                 });
             }
diff --git a/Engine/Plugins/Runtime/NetworkPrediction/Source/NetworkPrediction/Private/Async/NetworkPredictionAsyncModelDefRegistry.cpp b/Engine/Plugins/Runtime/NetworkPrediction/Source/NetworkPrediction/Private/Async/NetworkPredictionAsyncModelDefRegistry.cpp
new file mode 100644
--- /dev/null
+++ b/Engine/Plugins/Runtime/NetworkPrediction/Source/NetworkPrediction/Private/Async/NetworkPredictionAsyncModelDefRegistry.cpp
@@ -0,0 +1,8 @@
+// Copyright Epic Games, Inc. All Rights Reserved.
+#include "Async/NetworkPredictionAsyncModelDefRegistry.h"
+
+namespace UE_NP
+{
+	FNetworkPredictionAsyncModelDefRegistry FNetworkPredictionAsyncModelDefRegistry::Singleton;
+}
+
diff --git a/Engine/Plugins/Runtime/NetworkPrediction/Source/NetworkPrediction/Private/Async/NetworkPredictionAsyncProxy.cpp b/Engine/Plugins/Runtime/NetworkPrediction/Source/NetworkPrediction/Private/Async/NetworkPredictionAsyncProxy.cpp
new file mode 100644
--- /dev/null
+++ b/Engine/Plugins/Runtime/NetworkPrediction/Source/NetworkPrediction/Private/Async/NetworkPredictionAsyncProxy.cpp
@@ -0,0 +1,86 @@
+// Copyright Epic Games, Inc. All Rights Reserved.
+
+#include "Async/NetworkPredictionAsyncProxy.h"
+#include "Async/NetworkPredictionAsyncWorldManager.h"
+#include "Async/NetworkPredictionAsyncID.h"
+
+bool FNetworkPredictionAsyncProxy::RegisterProxy(UWorld* World)
+{
+	if (!ensure(Manager == nullptr))
+	{
+		return false;
+	}
+
+	Manager = UE_NP::FNetworkPredictionAsyncWorldManager::Get(World);
+	if (ensure(Manager))
+	{
+		const bool bIsClient = (World->GetNetMode() == NM_Client);
+		return Manager->RegisterInstance(ID, bIsClient);
+	}
+
+	return false;
+}
+
+void FNetworkPredictionAsyncProxy::UnregisterProxy()
+{
+	if (Manager == nullptr)
+	{
+		return;
+	}
+
+	Manager->UnregisterInstance(ID);
+}
+
+void FNetworkPredictionAsyncProxy::RegisterController(APlayerController* PC)
+{
+	if (!ensure(Manager) || !ensure(ID.IsValid()))
+	{
+		return;
+	}
+
+	Manager->RegisterController(ID, PC);
+}
+
+bool FNetworkPredictionAsyncProxy::NetSerialize(FArchive& Ar, class UPackageMap* Map, bool& bOutSuccess)
+{
+	if (!ensure(Manager))
+	{
+		// We are NetSerializing without being initialized. RegisterProxy should be called from something like InitializeComponent(s), prior to NetSerialization
+		return false;
+	}
+
+	// ID
+	uint32 RawID=(uint32)ID;
+	Ar.SerializeIntPacked(RawID);
+	
+	if (Ar.IsLoading())
+	{
+		if ((int32)ID != RawID)
+		{
+			UE_NP::FNetworkPredictionAsyncID NewID((int32)RawID);
+			if (ID.IsValid())
+			{
+				Manager->RemapClientInstance(ID, NewID);
+			}
+
+			ID = NewID;
+		}
+	}
+
+	// Sim state
+	Manager->NetSerializeInstance(ID, Ar);
+	return true;
+}
+
+bool FNetworkPredictionAsyncProxy::Identical(const FNetworkPredictionAsyncProxy* Other, uint32 PortFlags) const
+{
+	return CachedLatestFrame == Other->CachedLatestFrame;
+}
+
+void FNetworkPredictionAsyncProxy::OnPreReplication()
+{
+	if (Manager)
+	{
+		CachedLatestFrame = Manager->GetLatestFrame();
+	}
+}
\ No newline at end of file
diff --git a/Engine/Plugins/Runtime/NetworkPrediction/Source/NetworkPrediction/Private/Async/NetworkPredictionAsyncWorldManager.cpp b/Engine/Plugins/Runtime/NetworkPrediction/Source/NetworkPrediction/Private/Async/NetworkPredictionAsyncWorldManager.cpp
new file mode 100644
--- /dev/null
+++ b/Engine/Plugins/Runtime/NetworkPrediction/Source/NetworkPrediction/Private/Async/NetworkPredictionAsyncWorldManager.cpp
@@ -0,0 +1,150 @@
+// Copyright Epic Games, Inc. All Rights Reserved.
+
+#include "Async/NetworkPredictionAsyncWorldManager.h"
+#include "GameFramework/PlayerController.h"
+#include "Engine/World.h"
+#include "EngineUtils.h"
+
+namespace UE_NP {
+
+void FNetworkPredictionAsyncWorldManager::RegisterController(FNetworkPredictionAsyncID ID, APlayerController* PC)
+{
+	const bool bIsServer = IsServer();
+	const bool bEnableClientInputBuffering = false;
+	
+	bool bIsLocalController = false;
+	ENetworkPredictionAsyncInputSource InputSource = ENetworkPredictionAsyncInputSource::None;
+
+	if (PC)
+	{
+		bIsLocalController = PC->IsLocalController();
+		FControllerInfo* ControllerInfo = PlayerControllerMap.FindByKey(ID);
+		if (!ControllerInfo)
+		{
+			ControllerInfo = &PlayerControllerMap.AddDefaulted_GetRef();
+		}
+		ControllerInfo->ID = ID;
+		ControllerInfo->PC = PC;
+		ControllerInfo->LastFrame = 0;
+
+		// Server: non local controllers pull from the buffer
+		// client: pull from buffer only if local input buffering enabled
+		if ((bIsServer && !bIsLocalController) || (!bIsServer && bEnableClientInputBuffering))
+		{
+			InputSource = ENetworkPredictionAsyncInputSource::Buffered;
+		}
+		else
+		{
+			// Otherwise, pull from local PendingInputCmd that this sim registered with
+			InputSource = ENetworkPredictionAsyncInputSource::Local;
+		}
+
+		ControllerInfo->InputSource = InputSource;
+	}
+	else
+	{
+		int32 ExistingIdx = PlayerControllerMap.IndexOfByKey(ID);
+		if (ExistingIdx != INDEX_NONE)
+		{
+			PlayerControllerMap.RemoveAt(ExistingIdx, 1, false);
+		}
+		if (bIsServer)
+		{
+			// Non PC backed sim: only pull input cmd if on server
+			InputSource = ENetworkPredictionAsyncInputSource::Local;
+			bIsLocalController = true;
+		}
+	}
+
+	SetInputSource(ID, InputSource, bIsLocalController);
+}
+
+bool FNetworkPredictionAsyncWorldManager::IsServer() const
+{
+	npCheckSlow(World);
+	return World->GetNetMode() != NM_Client;
+}
+
+void FNetworkPredictionAsyncWorldManager::NetSerializePlayerControllerInputCmds(int32 PhysicsStep, TFunctionRef<void(FNetworkPredictionAsyncID, int32, FArchive&)> NetSerializeFunc)
+{
+	const bool bIsServer = IsServer();
+
+	if (bIsServer)
+	{
+		for (FControllerInfo& Info : PlayerControllerMap)
+		{
+			// Server: Non locally controlled. look for a Cmd from the PlayerController InputCmd Bus
+			APlayerController::FServerFrameInfo& PCFrameInfo = Info.PC->GetServerFrameInfo();
+			APlayerController::FInputCmdBuffer& PCInputBuffer = Info.PC->GetInputBuffer();
+
+			if (PCFrameInfo.LastProcessedInputFrame != INDEX_NONE)
+			{
+				auto NetSerializeCmd = [&](int32 PCFrame, int32 SimFrame)
+				{
+					const TArray<uint8>& Data = PCInputBuffer.Get(PCFrame);
+					if (Data.Num() > 0)
+					{
+						uint8* RawData = const_cast<uint8*>(Data.GetData());
+						FNetBitReader Ar(nullptr, RawData, ((int64)Data.Num()) << 3);
+						NetSerializeFunc(Info.ID, SimFrame, Ar);
+						Info.LastFrame = SimFrame;
+					}
+				};
+
+				if (PCFrameInfo.bFault)
+				{
+					// While in fault just write the current input to the buffer and don't bother with future inputs
+					if (Info.LastFrame != PhysicsStep)
+					{
+						NetSerializeCmd(PCFrameInfo.LastProcessedInputFrame, PhysicsStep);
+					}
+				}
+				else
+				{
+					// De-NetSerialize all new InputCmds from PC's buffer ["Client Frame" based] to our internal typed buffer ["local PhysicsStep Frame" based]
+					const int32 CmdOffset = PCFrameInfo.LastProcessedInputFrame - PhysicsStep; // Map PC frame to local Physics Step
+					const int32 WriteStart = FMath::Max(PhysicsStep, Info.LastFrame+1); // Start at the current PhysicsStep or highest frame we haven't processed
+					const int32 WriteEnd = PCInputBuffer.HeadFrame() - CmdOffset; // go until we reach the end of valid data in the PC input buffer
+
+					for (int32 Frame = WriteStart; Frame <= WriteEnd; ++Frame)
+					{
+						//UE_LOG(LogTemp, Log, TEXT("[S] De-NetSerialize Cmd %d Frame %d. NumBytes: %d"), Frame + CmdOffset, Frame, PCInputBuffer.Get(Frame + CmdOffset).Num());
+						NetSerializeCmd(Frame + CmdOffset, Frame);
+					}
+				}
+			}
+		}
+	}
+	else
+	{
+		for (FControllerInfo& Info : PlayerControllerMap)
+		{
+			if (Info.ID.IsClientGenerated())
+			{
+				continue;
+			}
+
+			FNetBitWriter Writer(nullptr, 256 << 3);			
+			const int32 InputFrame = Info.InputSource == ENetworkPredictionAsyncInputSource::Buffered ? PhysicsStep : INDEX_NONE;
+			NetSerializeFunc(Info.ID, InputFrame, Writer);
+
+			if (Writer.GetNumBits() > 0)
+			{
+				// We will probably need settings or some way to not send client->server at extremely high rates
+				// for now, its every game thread frame
+				TArray<uint8> SendData;
+				if (Writer.IsError() == false)
+				{
+					int32 NumBytes = (int32)Writer.GetNumBytes();
+					SendData = MoveTemp(*const_cast<TArray<uint8>*>(Writer.GetBuffer()));
+					SendData.SetNum(NumBytes);
+				}
+
+
+				Info.PC->PushClientInput(PhysicsStep, SendData);
+			}
+		}
+	}
+}
+
+}; // namespace UE_NP
\ No newline at end of file
diff --git a/Engine/Plugins/Runtime/NetworkPrediction/Source/NetworkPrediction/Private/NetworkPhysics.cpp b/Engine/Plugins/Runtime/NetworkPrediction/Source/NetworkPrediction/Private/NetworkPhysics.cpp
--- a/Engine/Plugins/Runtime/NetworkPrediction/Source/NetworkPrediction/Private/NetworkPhysics.cpp
+++ b/Engine/Plugins/Runtime/NetworkPrediction/Source/NetworkPrediction/Private/NetworkPhysics.cpp
@@ -51,6 +51,8 @@ namespace UE_NETWORK_PHYSICS
 	NP_DEVCVAR_INT(FixedLocalFrameOffset, -1, "np2.FixedLocalFrameOffset", "When > 0, use hardcoded frame offset on client from head");	
 	NP_DEVCVAR_INT(FixedLocalFrameOffsetTolerance, 3, "np2.FixedLocalFrameOffsetTolerance", "Tolerance when using np2.FixedLocalFrameOffset");
 
+	NP_DEVCVAR_INT(ForceInputFault, 0, "np2.ForceInputFault", "Force input fault correction");
+
 	// NOTE: These have temporarily been switched to regular CVar style variables, because the NP_DECVAR_* macros
 	// can only make ECVF_Cheat variables, which cannot be hotfixed with DefaultEngine.ini. Instead they use
 	// ConsoleVariables.ini, which isn't loaded in shipping.
@@ -173,7 +175,10 @@ struct FNetworkPhysicsRewindCallback : public Chaos::IRewindCallback
 		const float Error = FQuat::ErrorAutoNormalize(A, B);
 		const bool b = D == 0 ? A != B : Error > D;
 
-		//UE_LOG(LogTemp, Warning, TEXT("Error: %f"), Error);
+		if (b)
+		{
+			//UE_LOG(LogTemp, Warning, TEXT("%s Error: %f"), Str, Error);
+		}
 
 		UE_CLOG(UE_NETWORK_PHYSICS::LogCorrections > 0 && b && Str, LogNetworkPhysics, Log, TEXT("%s correction. Server: %s. Local: %s. Delta: %f"), Str, *A.ToString(), *B.ToString(), FQuat::ErrorAutoNormalize(A, B));
 		return b;
@@ -250,7 +255,7 @@ struct FNetworkPhysicsRewindCallback : public Chaos::IRewindCallback
 
 				if (Proxy->GetHandle_LowLevel() == nullptr)
 				{
-					UE_LOG(LogTemp, Warning, TEXT("No valid LowLEvel handle yet..."));
+					UE_LOG(LogNetworkPhysics, Warning, TEXT("No valid LowLevel handle yet."));
 					continue;
 				}
 
@@ -276,7 +281,7 @@ struct FNetworkPhysicsRewindCallback : public Chaos::IRewindCallback
 					}
 				}
 #endif
-				
+
 				//UE_LOG(LogTemp, Warning, TEXT("0x%X P.ObjectState(): %d. R: %s [%s]"), (int64)Proxy,  P.ObjectState(), *P.R().ToString(), *Obj.Physics.Rotation.ToString());
 				//UE_LOG(LogTemp, Warning, TEXT("			R: %s [%s]"), *FRotator(P.R()).ToString(), *FRotator(Obj.Physics.Rotation).ToString());
 
@@ -380,29 +385,33 @@ struct FNetworkPhysicsRewindCallback : public Chaos::IRewindCallback
 
 				// TODO: This nullcheck is to avoid null access crash in Editor Tests.
 				// Should re-evaluate whether/why this is occurring in the first place.
-				if (CorrectionState.Proxy)
+				if (ensure(CorrectionState.Proxy))
 				{
-				if (auto* PT = CorrectionState.Proxy->GetPhysicsThreadAPI())
-				{
-					UE_CLOG(UE_NETWORK_PHYSICS::LogCorrections > 0, LogNetworkPhysics, Log, TEXT("Applying Correction from frame %d (actual step: %d). Location: %s"), CorrectionState.Frame, PhysicsStep, *FVector(CorrectionState.Physics.Location).ToString());
-				
-					PT->SetX(CorrectionState.Physics.Location, false);
-					PT->SetV(CorrectionState.Physics.LinearVelocity, false);
-					PT->SetR(CorrectionState.Physics.Rotation, false);
-					PT->SetW(CorrectionState.Physics.AngularVelocity, false);
+					if (auto* PT = CorrectionState.Proxy->GetPhysicsThreadAPI())
+					{
+						UE_CLOG(UE_NETWORK_PHYSICS::LogCorrections > 0, LogNetworkPhysics, Log, TEXT("Applying Correction from frame %d (actual step: %d). Location: %s"), CorrectionState.Frame, PhysicsStep, *FVector(CorrectionState.Physics.Location).ToString());
+
+						npEnsure(CorrectionState.Physics.Location.ContainsNaN() == false);
+						npEnsure(CorrectionState.Physics.LinearVelocity.ContainsNaN() == false);
+						npEnsure(CorrectionState.Physics.AngularVelocity.ContainsNaN() == false);
 					
+						PT->SetX(CorrectionState.Physics.Location, false);
+						PT->SetV(CorrectionState.Physics.LinearVelocity, false);
+						PT->SetR(CorrectionState.Physics.Rotation, false);
+						PT->SetW(CorrectionState.Physics.AngularVelocity, false);
+
 						//Solver->GetParticles().MarkTransientDirtyParticle(PT->GetProxy()->GetHandle_LowLevel());
 
-					//if (PT->ObjectState() != CorrectionState.Physics.ObjectState)
-					{
-						ensure(CorrectionState.Physics.ObjectState != Chaos::EObjectStateType::Uninitialized);
+						//if (PT->ObjectState() != CorrectionState.Physics.ObjectState)
+						{
+							ensure(CorrectionState.Physics.ObjectState != Chaos::EObjectStateType::Uninitialized);
 							UE_CLOG(UE_NETWORK_PHYSICS::LogCorrections > 0 && PT->ObjectState() != CorrectionState.Physics.ObjectState, LogNetworkPhysics, Log, TEXT("Applying Correction State %d"), CorrectionState.Physics.ObjectState);
-						PT->SetObjectState(CorrectionState.Physics.ObjectState);
+							PT->SetObjectState(CorrectionState.Physics.ObjectState);
+						}
 					}
 				}
 			}
 		}
-		}
 
 		// Marhsall data back to GT based on what was requested for networking
 		if (!RewindData->IsResim())
@@ -657,15 +666,19 @@ void UNetworkPhysicsManager::OnWorldPostInit(UWorld* World, const UWorld::Initia
 			const IConsoleVariable* NumFramesCVar = IConsoleManager::Get().FindConsoleVariable(TEXT("p.RewindCaptureNumFrames"));
 			if (ensure(NumFramesCVar))
 			{
-				NumFrames = NumFramesCVar->GetInt();
+				// 1 frame is required to enable rewind capture. 
+				NumFrames = FMath::Max<int32>(1, NumFramesCVar->GetInt());
 			}
 
-			Solver->EnableRewindCapture(NumFrames, true, MakeUnique<FNetworkPhysicsRewindCallback>());
-			RewindCallback = static_cast<FNetworkPhysicsRewindCallback*>(Solver->GetRewindCallback());
-			RewindCallback->RewindData = Solver->GetRewindData();
-			RewindCallback->Solver = Solver;
-			RewindCallback->World = World;
-			RewindCallback->NetworkPhysicsManager = this;
+			if (ensure(NumFrames > 0))
+			{
+				Solver->EnableRewindCapture(NumFrames, true, MakeUnique<FNetworkPhysicsRewindCallback>());
+				RewindCallback = static_cast<FNetworkPhysicsRewindCallback*>(Solver->GetRewindCallback());
+				RewindCallback->RewindData = Solver->GetRewindData();
+				RewindCallback->Solver = Solver;
+				RewindCallback->World = World;
+				RewindCallback->NetworkPhysicsManager = this;
+			}
 		}
 	}
 #endif
@@ -732,7 +745,7 @@ void UNetworkPhysicsManager::PostNetRecv()
 
 				if (UE_NETWORK_PHYSICS::TimeDilationEnabled > 0)
 				{
-				PhysScene->SetNetworkDeltaTimeScale(RealTimeDilation);
+					PhysScene->SetNetworkDeltaTimeScale(RealTimeDilation);
 				}
 
 				if (LatestAckdClientFrame != INDEX_NONE)
@@ -740,7 +753,7 @@ void UNetworkPhysicsManager::PostNetRecv()
 					const int32 FrameOffset = LatestAckdClientFrame - LatestAckdServerFrame;
 					if (FrameOffset != LocalOffset)
 					{
-						UE_CLOG(UE_NETWORK_PHYSICS::LogCorrections > 0, LogNetworkPhysics, Log, TEXT("LocalOFfset Changed: %d -> %d"), LocalOffset, FrameOffset);
+						UE_CLOG(UE_NETWORK_PHYSICS::LogCorrections > 0, LogNetworkPhysics, Log, TEXT("LocalOFfset Changed: %d -> %d [C:%d. S:%d]"), LocalOffset, FrameOffset, LatestAckdClientFrame, LatestAckdServerFrame);
 						LocalOffset = FrameOffset;
 					}
 				}
@@ -832,9 +845,9 @@ void UNetworkPhysicsManager::PostNetRecv()
 
 							if (PhysicsState->OwningActor)
 							{
-							PhysicsState->OwningActor->GetReplicatedMovement_Mutable().bRepPhysics = false;
-							PhysicsState->OwningActor->SetReplicateMovement(false);
-						}
+								PhysicsState->OwningActor->GetReplicatedMovement_Mutable().bRepPhysics = false;
+								PhysicsState->OwningActor->SetReplicateMovement(false);
+							}
 						}
 						else
 						{
@@ -1009,7 +1022,7 @@ void UNetworkPhysicsManager::ProcessInputs_External(int32 PhysicsStep, const TAr
 	}
 
 	// ----------------------------------------------------------------------
-	//	Call into Subsystems. This is where they will marshall input to the PT
+	//	Call into Subsystems. This is where they will marshal input to the PT
 	// ----------------------------------------------------------------------
 
 	bool bOutSentClientInput = false;
@@ -1067,6 +1080,9 @@ void UNetworkPhysicsManager::ProcessClientInputBuffers_External(int32 PhysicsSte
 	UWorld* World = GetWorld();
 	constexpr int32 MaxBufferedCmds = 16;
 
+	const bool bForceFault = UE_NETWORK_PHYSICS::ForceInputFault > 0;
+	UE_NETWORK_PHYSICS::ForceInputFault = 0;
+
 	for (FConstPlayerControllerIterator Iterator = World->GetPlayerControllerIterator(); Iterator; ++Iterator)
 	{
 		if (APlayerController* PC = Iterator->Get())
@@ -1076,7 +1092,7 @@ void UNetworkPhysicsManager::ProcessClientInputBuffers_External(int32 PhysicsSte
 
 			auto UpdateLastProcessedInputFrame = [&]()
 			{
-				const int32 NumBufferedInputCmds = InputBuffer.HeadFrame() - FrameInfo.LastProcessedInputFrame;
+				const int32 NumBufferedInputCmds = bForceFault ? 0 : (InputBuffer.HeadFrame() - FrameInfo.LastProcessedInputFrame);
 
 				// Check Overflow
 				if (NumBufferedInputCmds > MaxBufferedCmds)
@@ -1248,7 +1264,7 @@ void UNetworkPhysicsManager::TickDrawDebug()
 				{
 					(*Func)(P);
 				}
-			}
+			}			
 		}
 	}
 
diff --git a/Engine/Plugins/Runtime/NetworkPrediction/Source/NetworkPrediction/Private/NetworkPhysicsComponent.cpp b/Engine/Plugins/Runtime/NetworkPrediction/Source/NetworkPrediction/Private/NetworkPhysicsComponent.cpp
new file mode 100644
--- /dev/null
+++ b/Engine/Plugins/Runtime/NetworkPrediction/Source/NetworkPrediction/Private/NetworkPhysicsComponent.cpp
@@ -0,0 +1,209 @@
+// Copyright Epic Games, Inc. All Rights Reserved.
+
+#include "NetworkPhysicsComponent.h"
+#include "Engine/World.h"
+#include "Physics/Experimental/PhysScene_Chaos.h"
+#include "PBDRigidsSolver.h"
+#include "ChaosSolversModule.h"
+#include "RewindData.h"
+#include "EngineUtils.h"
+#include "Net/UnrealNetwork.h"
+#include "Chaos/SimCallbackObject.h"
+#include "UObject/UObjectIterator.h"
+#include "GameFramework/PlayerController.h"
+#include "DrawDebugHelpers.h"
+#include "PhysicsProxy/SingleParticlePhysicsProxy.h"
+#include "TimerManager.h"
+#include "NetworkPredictionDebug.h"
+#include "NetworkPredictionCVars.h"
+
+UNetworkPhysicsComponent::UNetworkPhysicsComponent()
+{
+	bWantsInitializeComponent = true;
+	bAutoActivate = true;
+	SetIsReplicatedByDefault(true);
+}
+
+void UNetworkPhysicsComponent::InitializeComponent()
+{
+#if WITH_CHAOS
+	Super::InitializeComponent();
+
+	UWorld* World = GetWorld();	
+	checkSlow(World);
+	UNetworkPhysicsManager* Manager = World->GetSubsystem<UNetworkPhysicsManager>();
+	if (!Manager)
+	{
+		return;
+	}
+
+	// Test is component is valid for Network Physics. Needs a valid physics ActorHandle
+	auto ValidComponent = [](UActorComponent* Component)
+	{
+		bool bValid = false;
+		if (UPrimitiveComponent* PrimComp = Cast<UPrimitiveComponent>(Component))
+		{
+			bValid = FPhysicsInterface::IsValid(PrimComp->BodyInstance.ActorHandle);
+		}
+		return bValid;
+	};
+
+	auto SelectComponent = [&ValidComponent](const TArray<UActorComponent*>& Components)
+	{
+		UPrimitiveComponent* Pc = nullptr;
+		for (UActorComponent* Ac : Components)
+		{
+			if (ValidComponent(Ac))
+			{
+				Pc = (UPrimitiveComponent*)Ac;
+				break;
+			}
+
+		}
+		return Pc;
+	};
+
+	UPrimitiveComponent* PrimitiveComponent = nullptr;
+	if (AActor* MyActor = GetOwner())
+	{
+		// Explicitly tagged component
+		if (ManagedComponentTag != NAME_None)
+		{
+			if (UPrimitiveComponent* FoundComponent = SelectComponent(MyActor->GetComponentsByTag(UPrimitiveComponent::StaticClass(), ManagedComponentTag)))
+			{
+				PrimitiveComponent = FoundComponent;
+			}
+			else
+			{
+				UE_LOG(LogNetworkPhysics, Warning, TEXT("Actor %s: could not find a valid Primitive Component with Tag %s"), *MyActor->GetPathName(), *ManagedComponentTag.ToString());
+			}
+		}
+
+		// Root component
+		if (!PrimitiveComponent && ValidComponent(MyActor->GetRootComponent()))
+		{
+			PrimitiveComponent = CastChecked<UPrimitiveComponent>(MyActor->GetRootComponent());
+		}
+
+		// Any other valid primitive component?
+		if (!PrimitiveComponent)
+		{
+			if (UPrimitiveComponent* FoundComponent = SelectComponent(MyActor->K2_GetComponentsByClass(UPrimitiveComponent::StaticClass())))
+			{
+				PrimitiveComponent = FoundComponent;
+			}
+		}
+	}
+
+	if (ensureMsgf(PrimitiveComponent, TEXT("No PrimitiveComponent found on %s"), *GetPathName()))
+	{
+		NetworkPhysicsState.Proxy = PrimitiveComponent->BodyInstance.ActorHandle;
+		NetworkPhysicsState.OwningActor = GetOwner();
+		ensure(NetworkPhysicsState.OwningActor);
+
+		Manager->RegisterPhysicsProxy(&NetworkPhysicsState);
+
+		Manager->RegisterPhysicsProxyDebugDraw(&NetworkPhysicsState, [this](const UNetworkPhysicsManager::FDrawDebugParams& P)
+			{
+				AActor* Actor = this->GetOwner();
+				FBox LocalSpaceBox = Actor->CalculateComponentsBoundingBoxInLocalSpace();
+				const float Thickness = 2.f;
+
+				FVector ActorOrigin;
+				FVector ActorExtent;
+				LocalSpaceBox.GetCenterAndExtents(ActorOrigin, ActorExtent);
+				ActorExtent *= Actor->GetActorScale3D();
+				DrawDebugBox(P.DrawWorld, P.Loc, ActorExtent, P.Rot, P.Color, false, P.Lifetime, 0, Thickness);
+			});
+	}
+#endif
+}
+
+void UNetworkPhysicsComponent::UninitializeComponent()
+{
+	Super::UninitializeComponent();
+#if WITH_CHAOS
+	if (UWorld* World = GetWorld())
+	{
+		if (UNetworkPhysicsManager* Manager = World->GetSubsystem<UNetworkPhysicsManager>())
+		{
+			Manager->UnregisterPhysicsProxy(&NetworkPhysicsState);
+		}
+	}
+#endif
+}
+
+APlayerController* UNetworkPhysicsComponent::GetOwnerPC() const
+{
+	if (APawn* PawnOwner = Cast<APawn>(GetOwner()))
+	{
+		return Cast<APlayerController>(PawnOwner->GetController());
+	}
+
+	return nullptr;
+}
+
+void UNetworkPhysicsComponent::GetLifetimeReplicatedProps(TArray<FLifetimeProperty>& OutLifetimeProps) const
+{
+	Super::GetLifetimeReplicatedProps(OutLifetimeProps);
+	DOREPLIFETIME( UNetworkPhysicsComponent, NetworkPhysicsState);
+}
+
+// ========================================================================================
+// ========================================================================================
+// ========================================================================================
+
+#if 0
+
+void UNetworkPhysicsComponent::ProcessInputs_External(FMockManagedState& State, int32 PhysicsStep, int32 LocalFrameOffset)
+{
+	if (bRecording)
+	{
+		if (CurrentInputCmdStream)
+		{
+			CurrentInputCmdStream->Add(State.InputCmd);
+		}
+	}
+	else if (CurrentInputCmdStream && CurrentInputCmdStream->IsValidIndex(PlaybackIdx))
+	{
+		State.InputCmd = (*CurrentInputCmdStream)[PlaybackIdx++ % CurrentInputCmdStream->Num()];		
+	}
+}
+
+void UNetworkPhysicsComponent::StartRecording(TArray<FMockPhysInputCmd>* Stream)
+{
+	if (bRecording)
+	{
+		return;
+	}
+
+	CurrentInputCmdStream = Stream;
+	bRecording = true;
+}
+
+void UNetworkPhysicsComponent::StopRecording()
+{
+	if (CurrentInputCmdStream)
+	{
+		UE_LOG(LogTemp, Log, TEXT("Recorded %d Inputs."), CurrentInputCmdStream->Num());
+	}
+
+	bRecording = false;
+	CurrentInputCmdStream = nullptr;
+}
+
+void UNetworkPhysicsComponent::StartPlayback(TArray<FMockPhysInputCmd>* Stream)
+{
+	CurrentInputCmdStream = Stream;
+	if (CurrentInputCmdStream && CurrentInputCmdStream->Num() > 0)
+	{
+		PlaybackIdx = 0;
+	}
+	else
+	{
+		PlaybackIdx = INDEX_NONE;
+	}
+}
+
+
+#endif
\ No newline at end of file
diff --git a/Engine/Plugins/Runtime/NetworkPrediction/Source/NetworkPrediction/Private/NetworkPredictionAsync.cpp b/Engine/Plugins/Runtime/NetworkPrediction/Source/NetworkPrediction/Private/NetworkPredictionAsync.cpp
deleted file mode 100644
diff --git a/Engine/Plugins/Runtime/NetworkPrediction/Source/NetworkPrediction/Private/NetworkPredictionModule.cpp b/Engine/Plugins/Runtime/NetworkPrediction/Source/NetworkPrediction/Private/NetworkPredictionModule.cpp
--- a/Engine/Plugins/Runtime/NetworkPrediction/Source/NetworkPrediction/Private/NetworkPredictionModule.cpp
+++ b/Engine/Plugins/Runtime/NetworkPrediction/Source/NetworkPrediction/Private/NetworkPredictionModule.cpp
@@ -13,6 +13,7 @@
 #include "Misc/Parse.h"
 #include "String/ParseTokens.h"
 #include "NetworkPredictionModelDefRegistry.h"
+#include "Async/NetworkPredictionAsyncModelDefRegistry.h"
 
 #if WITH_EDITOR
 #include "Editor.h"
@@ -20,6 +21,7 @@
 #include "NetworkPredictionSettings.h"
 #endif
 
+
 #define LOCTEXT_NAMESPACE "FNetworkPredictionModule"
 
 class FNetworkPredictionModule : public INetworkPredictionModule
@@ -137,6 +139,7 @@ void FNetworkPredictionModule::FinalizeNetworkPredictionTypes()
 {
 	FGlobalCueTypeTable::Get().FinalizeCueTypes();
 	FNetworkPredictionModelDefRegistry::Get().FinalizeTypes();
+	UE_NP::FNetworkPredictionAsyncModelDefRegistry::Get().FinalizeTypes();
 }
 
 IMPLEMENT_MODULE( FNetworkPredictionModule, NetworkPrediction )
diff --git a/Engine/Plugins/Runtime/NetworkPrediction/Source/NetworkPrediction/Private/NetworkPredictionSettings.cpp b/Engine/Plugins/Runtime/NetworkPrediction/Source/NetworkPrediction/Private/NetworkPredictionSettings.cpp
--- a/Engine/Plugins/Runtime/NetworkPrediction/Source/NetworkPrediction/Private/NetworkPredictionSettings.cpp
+++ b/Engine/Plugins/Runtime/NetworkPrediction/Source/NetworkPrediction/Private/NetworkPredictionSettings.cpp
@@ -217,35 +217,46 @@ namespace UE_NETWORK_PREDICTION
 		}
 
 		// One time init:
+		TWeakObjectPtr<UWorld> WeakWorld(InWorld);
+		bool bRegisteredInput  = false;
+		auto BindInput = [&bRegisteredInput, WeakWorld]()
+		{
+			if (UWorld* World = WeakWorld.Get())
+			{
+				if (ULocalPlayer* LocalPlayer = World->GetFirstLocalPlayerFromController())
+				{
+					//LocalPlayer->Exec(InWorld, *FString::Printf(TEXT("setbind %d \"np2.DevMenu %d\n"), i, i), *GLog);
+
+					LocalPlayer->Exec(World, TEXT("setbind one \"np2.DevMenu 1\""), *GLog);
+					LocalPlayer->Exec(World, TEXT("setbind two \"np2.DevMenu 2\""), *GLog);
+					LocalPlayer->Exec(World, TEXT("setbind three \"np2.DevMenu 3\""), *GLog);
+					LocalPlayer->Exec(World, TEXT("setbind four \"np2.DevMenu 4\""), *GLog);
+					LocalPlayer->Exec(World, TEXT("setbind five \"np2.DevMenu 5\""), *GLog);
+					LocalPlayer->Exec(World, TEXT("setbind six \"np2.DevMenu 6\""), *GLog);
+					LocalPlayer->Exec(World, TEXT("setbind seven \"np2.DevMenu 7\""), *GLog);
+					LocalPlayer->Exec(World, TEXT("setbind eight \"np2.DevMenu 8\""), *GLog);
+					LocalPlayer->Exec(World, TEXT("setbind nine \"np2.DevMenu 9\""), *GLog);
+					LocalPlayer->Exec(World, TEXT("setbind zero \"np2.DevMenu 0\""), *GLog);
+					bRegisteredInput = true;
+				}
+			}
+		};
+
+		if (!bRegisteredInput)
+		{
+			BindInput();
+		}
+
 		static FDelegateHandle Handle;
 		if (Handle.IsValid())
 		{
 			return;
-		}	
+		}
 
 		Handle = AHUD::OnHUDPostRender.AddLambda([](AHUD* HUD, UCanvas* Canvas)
 		{
 			DrawDevHUD(WorldHandlesMap, HUD, Canvas);
 		});
-
-		//for (int32 i=0; i <= 9; ++i)
-		{
-			if (ULocalPlayer* LocalPlayer = InWorld->GetFirstLocalPlayerFromController())
-			{
-				//LocalPlayer->Exec(InWorld, *FString::Printf(TEXT("setbind %d \"np2.DevMenu %d\n"), i, i), *GLog);
-
-				LocalPlayer->Exec(InWorld, TEXT("setbind one \"np2.DevMenu 1\""), *GLog);
-				LocalPlayer->Exec(InWorld, TEXT("setbind two \"np2.DevMenu 2\""), *GLog);
-				LocalPlayer->Exec(InWorld, TEXT("setbind three \"np2.DevMenu 3\""), *GLog);
-				LocalPlayer->Exec(InWorld, TEXT("setbind four \"np2.DevMenu 4\""), *GLog);
-				LocalPlayer->Exec(InWorld, TEXT("setbind five \"np2.DevMenu 5\""), *GLog);
-				LocalPlayer->Exec(InWorld, TEXT("setbind six \"np2.DevMenu 6\""), *GLog);
-				LocalPlayer->Exec(InWorld, TEXT("setbind seven \"np2.DevMenu 7\""), *GLog);
-				LocalPlayer->Exec(InWorld, TEXT("setbind eight \"np2.DevMenu 8\""), *GLog);
-				LocalPlayer->Exec(InWorld, TEXT("setbind nine \"np2.DevMenu 9\""), *GLog);
-				LocalPlayer->Exec(InWorld, TEXT("setbind zero \"np2.DevMenu 0\""), *GLog);
-			}
-		}
 	}
 }
 
diff --git a/Engine/Plugins/Runtime/NetworkPrediction/Source/NetworkPrediction/Private/NetworkPredictionUtil.cpp b/Engine/Plugins/Runtime/NetworkPrediction/Source/NetworkPrediction/Private/NetworkPredictionUtil.cpp
new file mode 100644
--- /dev/null
+++ b/Engine/Plugins/Runtime/NetworkPrediction/Source/NetworkPrediction/Private/NetworkPredictionUtil.cpp
@@ -0,0 +1,75 @@
+// Copyright Epic Games, Inc. All Rights Reserved.
+
+#include "NetworkPredictionUtil.h"
+#include "GameFramework/Actor.h"
+#include "Components/PrimitiveComponent.h"
+#include "NetworkPredictionLog.h"
+
+namespace UE_NP
+{
+	FSingleParticlePhysicsProxy* FindBestPhysicsProxy(AActor* Owner, FName NamedComponent)
+	{
+		// Test is component is valid for Network Physics. Needs a valid physics ActorHandle
+		auto ValidComponent = [](UActorComponent* Component)
+		{
+			bool bValid = false;
+			if (UPrimitiveComponent* PrimComp = Cast<UPrimitiveComponent>(Component))
+			{
+				bValid = FPhysicsInterface::IsValid(PrimComp->BodyInstance.ActorHandle);
+			}
+			return bValid;
+		};
+
+		auto SelectComponent = [&ValidComponent](const TArray<UActorComponent*>& Components)
+		{
+			UPrimitiveComponent* Pc = nullptr;
+			for (UActorComponent* Ac : Components)
+			{
+				if (ValidComponent(Ac))
+				{
+					Pc = (UPrimitiveComponent*)Ac;
+					break;
+				}
+
+			}
+			return Pc;
+		};
+
+		FSingleParticlePhysicsProxy* Proxy = nullptr;
+		UPrimitiveComponent* PrimitiveComponent = nullptr;
+		
+		// Explicitly tagged component
+		if (NamedComponent != NAME_None)
+		{
+			if (UPrimitiveComponent* FoundComponent = SelectComponent(Owner->GetComponentsByTag(UPrimitiveComponent::StaticClass(), NamedComponent)))
+			{
+				PrimitiveComponent = FoundComponent;
+			}
+			else
+			{
+				UE_LOG(LogNetworkPrediction, Warning, TEXT("Actor %s: could not find a valid Primitive Component with Tag %s"), *Owner->GetPathName(), *NamedComponent.ToString());
+			}
+		}
+
+		// Root component
+		if (!PrimitiveComponent && ValidComponent(Owner->GetRootComponent()))
+		{
+			PrimitiveComponent = CastChecked<UPrimitiveComponent>(Owner->GetRootComponent());
+		}
+
+		// Any other valid primitive component?
+		if (!PrimitiveComponent)
+		{
+			if (UPrimitiveComponent* FoundComponent = SelectComponent(Owner->K2_GetComponentsByClass(UPrimitiveComponent::StaticClass())))
+			{
+				PrimitiveComponent = FoundComponent;
+			}
+		}
+
+		if (PrimitiveComponent)
+		{
+			Proxy = PrimitiveComponent->BodyInstance.ActorHandle;
+		}
+		return Proxy;
+	}
+}
\ No newline at end of file
diff --git a/Engine/Plugins/Runtime/NetworkPrediction/Source/NetworkPrediction/Private/NetworkPredictionWorldManager.cpp b/Engine/Plugins/Runtime/NetworkPrediction/Source/NetworkPrediction/Private/NetworkPredictionWorldManager.cpp
--- a/Engine/Plugins/Runtime/NetworkPrediction/Source/NetworkPrediction/Private/NetworkPredictionWorldManager.cpp
+++ b/Engine/Plugins/Runtime/NetworkPrediction/Source/NetworkPrediction/Private/NetworkPredictionWorldManager.cpp
@@ -15,15 +15,6 @@
 // Do extra checks to make sure Physics and GameThread (PrimitiveComponent) are in sync at verious points in the rollback process
 #define NP_ENSURE_PHYSICS_GT_SYNC !(UE_BUILD_SHIPPING || UE_BUILD_TEST)
 
-namespace NetworkPredictionWorldManagerCVars
-{
-	static int32 MaxStepMS = 100;
-	static FAutoConsoleVariableRef CVarMaxSimulationStepMS(
-		TEXT("np.MaxSimulationStepMS"),
-		MaxStepMS,
-		TEXT("The maximum time step (in integer milliseconds) that we will allow for simulation frames. Frame delta times will be clamped between 1 MS and this value.."));
-}
-
 UNetworkPredictionWorldManager* UNetworkPredictionWorldManager::ActiveInstance=nullptr;
 
 // -----------------------------------------------------------------------------------------------
@@ -345,12 +336,13 @@ void UNetworkPredictionWorldManager::BeginNewSimulationFrame(UWorld* InWorld, EL
 	{
 		// Update VariableTickState	
 		constexpr int32 MinStepMS = 1;
+		constexpr int32 MaxStepMS = 100;
 
 		VariableTickState.UnspentTimeMS += fEngineFrameDeltaTimeMS;
 		float fDeltaMS = FMath::FloorToFloat(VariableTickState.UnspentTimeMS);
 		VariableTickState.UnspentTimeMS -= fDeltaMS;
 
-		const int32 DeltaSimMS = FMath::Clamp((int32)fDeltaMS, MinStepMS, NetworkPredictionWorldManagerCVars::MaxStepMS);
+		const int32 DeltaSimMS = FMath::Clamp((int32)fDeltaMS, MinStepMS, MaxStepMS);
 
 		FVariableTickState::FFrame& PendingFrameData = VariableTickState.Frames[VariableTickState.PendingFrame];
 		PendingFrameData.DeltaMS = DeltaSimMS;
diff --git a/Engine/Plugins/Runtime/NetworkPrediction/Source/NetworkPrediction/Public/Async/NetworkPredictionAsyncDataStore.h b/Engine/Plugins/Runtime/NetworkPrediction/Source/NetworkPrediction/Public/Async/NetworkPredictionAsyncDataStore.h
new file mode 100644
--- /dev/null
+++ b/Engine/Plugins/Runtime/NetworkPrediction/Source/NetworkPrediction/Public/Async/NetworkPredictionAsyncDataStore.h
@@ -0,0 +1,71 @@
+// Copyright Epic Games, Inc. All Rights Reserved.
+
+#pragma once
+#include "Containers/StaticArray.h"
+#include "NetworkPredictionAsyncDefines.h"
+#include "NetworkPredictionCheck.h"
+
+namespace UE_NP {
+
+struct FAsyncDataStoreCollection
+{
+public:
+
+	template<typename DataStoreType, typename AsyncModelDef>
+	DataStoreType* GetDataStore()
+	{
+		return GetDataStore_Internal<DataStoreType, AsyncModelDef>();
+	};
+
+	template<typename DataStoreType, typename AsyncModelDef>
+	const DataStoreType* GetDataStore() const
+	{
+		return const_cast<FAsyncDataStoreCollection*>(this)->GetDataStore_Internal<DataStoreType, AsyncModelDef>();
+	}
+
+	void Reset()
+	{
+		for (TUniquePtr<IDataStore>& Ptr : DataStoreArray)
+		{
+			if (Ptr.IsValid())
+			{
+				Ptr->Reset();
+			}
+		}
+	}
+
+private:
+
+	struct IDataStore
+	{
+		virtual ~IDataStore() = default;
+		virtual void Reset() = 0;
+	};
+
+	// This needs to be static for thread safety. We can't resize the array on the GT once the PT is running
+	// FIXME: this may not be true anymore. PT could have its own data store
+	TStaticArray<TUniquePtr<IDataStore>, MaxAsyncModelDefs> DataStoreArray;
+
+
+	template<typename DataStoreType, typename AsyncModelDef>
+	DataStoreType* GetDataStore_Internal()
+	{
+		npEnsureMsgf(AsyncModelDef::ID >= 0, TEXT("AsyncModelDef %s has invalid ID assigned. Could be missing NP_ASYNC_MODEL_REGISTER."), AsyncModelDef::GetName());
+
+		struct FThisDataStore : IDataStore
+		{
+			DataStoreType Self;
+			void Reset() final override { Self.Reset(); }
+		};
+
+		TUniquePtr<IDataStore>& Item = DataStoreArray[AsyncModelDef::ID];
+		if (Item.IsValid() == false)
+		{
+			Item = MakeUnique<FThisDataStore>();
+		}
+
+		return &((FThisDataStore*)Item.Get())->Self;
+	};
+};
+
+} // namespace UE_NP
\ No newline at end of file
diff --git a/Engine/Plugins/Runtime/NetworkPrediction/Source/NetworkPrediction/Public/Async/NetworkPredictionAsyncDefines.h b/Engine/Plugins/Runtime/NetworkPrediction/Source/NetworkPrediction/Public/Async/NetworkPredictionAsyncDefines.h
new file mode 100644
--- /dev/null
+++ b/Engine/Plugins/Runtime/NetworkPrediction/Source/NetworkPrediction/Public/Async/NetworkPredictionAsyncDefines.h
@@ -0,0 +1,31 @@
+// Copyright Epic Games, Inc. All Rights Reserved.
+
+#pragma once
+
+namespace UE_NP
+{
+#ifndef NP_MAX_ASYNC_MODEL_DEFS
+	#define NP_MAX_ASYNC_MODEL_DEFS 16
+#endif
+
+const int32 MaxAsyncModelDefs = NP_MAX_ASYNC_MODEL_DEFS;
+
+#ifndef NP_NUM_FRAME_STORAGE
+	#define NP_NUM_FRAME_STORAGE 64
+#endif
+
+const int32 NumFramesStorage = NP_NUM_FRAME_STORAGE;
+
+#ifndef NP_FRAME_STORAGE_GROWTH
+	#define NP_FRAME_STORAGE_GROWTH 8
+#endif
+
+const int32 FrameStorageGrowth = NP_FRAME_STORAGE_GROWTH;
+
+#ifndef NP_FRAME_INPUTCMD_BUFFER_SIZE
+	#define NP_FRAME_INPUTCMD_BUFFER_SIZE 16
+#endif
+
+const int32 InputCmdBufferSize = NP_FRAME_INPUTCMD_BUFFER_SIZE;
+
+};
\ No newline at end of file
diff --git a/Engine/Plugins/Runtime/NetworkPrediction/Source/NetworkPrediction/Public/Async/NetworkPredictionAsyncGlobalTickContext.h b/Engine/Plugins/Runtime/NetworkPrediction/Source/NetworkPrediction/Public/Async/NetworkPredictionAsyncGlobalTickContext.h
new file mode 100644
--- /dev/null
+++ b/Engine/Plugins/Runtime/NetworkPrediction/Source/NetworkPrediction/Public/Async/NetworkPredictionAsyncGlobalTickContext.h
@@ -0,0 +1,30 @@
+// Copyright Epic Games, Inc. All Rights Reserved.
+
+#pragma once
+#include "NetworkPredictionAsyncDataStore.h"
+#include "NetworkPredictionCheck.h"
+#include "Async/NetworkPredictionAsyncInstanceData.h"
+
+// Data that is available in sim ticks. This is meant to hold the complete "World view of what is going o
+struct FNetworkPredictionAsyncTickContext
+{
+	FNetworkPredictionAsyncTickContext(UE_NP::FAsyncDataStoreCollection* InDataStoreCollection)
+		: DataStoreCollection(InDataStoreCollection) { }
+
+	UWorld* World = nullptr;
+	float DeltaTimeSeconds;
+	int32 SimulationFrame;
+	int32 LocalStorageFrame;
+	bool bIsResim = false;
+
+	template<typename AsyncModelDef>
+	UE_NP::TAsyncModelDataStore_Input<AsyncModelDef>* GetDataStore()
+	{
+		npCheckSlow(DataStoreCollection);
+		return DataStoreCollection->GetDataStore<UE_NP::TAsyncModelDataStore_Output<AsyncModelDef>, AsyncModelDef>();
+	}
+
+private:
+
+	UE_NP::FAsyncDataStoreCollection* DataStoreCollection;
+};
\ No newline at end of file
diff --git a/Engine/Plugins/Runtime/NetworkPrediction/Source/NetworkPrediction/Public/Async/NetworkPredictionAsyncID.h b/Engine/Plugins/Runtime/NetworkPrediction/Source/NetworkPrediction/Public/Async/NetworkPredictionAsyncID.h
new file mode 100644
--- /dev/null
+++ b/Engine/Plugins/Runtime/NetworkPrediction/Source/NetworkPrediction/Public/Async/NetworkPredictionAsyncID.h
@@ -0,0 +1,29 @@
+// Copyright Epic Games, Inc. All Rights Reserved.
+
+#pragma once
+
+namespace UE_NP {
+
+// Unique server assigned ID for a network sim instance
+struct FNetworkPredictionAsyncID
+{
+	FNetworkPredictionAsyncID() = default;
+	explicit FNetworkPredictionAsyncID(int32 InSpawnID) : SpawnID(InSpawnID) { }
+
+	bool IsValid() const { return SpawnID != 0; }
+	void Reset() { SpawnID = 0; }
+	bool IsClientGenerated() const { return SpawnID < 0; }
+
+	operator int32() const { return SpawnID; }
+	bool operator<  (const FNetworkPredictionAsyncID &rhs) const { return(this->SpawnID < rhs.SpawnID); }
+	bool operator<= (const FNetworkPredictionAsyncID &rhs) const { return(this->SpawnID <= rhs.SpawnID); }
+	bool operator>  (const FNetworkPredictionAsyncID &rhs) const { return(this->SpawnID > rhs.SpawnID); }
+	bool operator>= (const FNetworkPredictionAsyncID &rhs) const { return(this->SpawnID >= rhs.SpawnID); }
+	bool operator== (const FNetworkPredictionAsyncID &rhs) const { return(this->SpawnID == rhs.SpawnID); }
+	bool operator!= (const FNetworkPredictionAsyncID &rhs) const { return(this->SpawnID != rhs.SpawnID); }
+
+private:
+	int32 SpawnID=0;
+};
+
+} // namespace UE_NP
\ No newline at end of file
diff --git a/Engine/Plugins/Runtime/NetworkPrediction/Source/NetworkPrediction/Public/Async/NetworkPredictionAsyncInstanceData.h b/Engine/Plugins/Runtime/NetworkPrediction/Source/NetworkPrediction/Public/Async/NetworkPredictionAsyncInstanceData.h
new file mode 100644
--- /dev/null
+++ b/Engine/Plugins/Runtime/NetworkPrediction/Source/NetworkPrediction/Public/Async/NetworkPredictionAsyncInstanceData.h
@@ -0,0 +1,276 @@
+// Copyright Epic Games, Inc. All Rights Reserved.
+
+#pragma once
+#include "NetworkPredictionAsyncDataStore.h"
+#include "NetworkPredictionAsyncModelDef.h"
+#include "NetworkPredictionUtil.h"
+#include "Containers/CircularBuffer.h"
+#include "NetworkPredictionAsyncID.h"
+#include "NetworkPredictionAsyncPersistentStorage.h"
+#include "NetworkPredictionAsyncDefines.h"
+#include "Misc/EnumClassFlags.h"
+
+namespace UE_NP {
+
+enum class ENetworkPredictionAsyncInputSource : uint8 // fixme: move
+{
+	None,		// There no gt-created input for this instance
+	Local,		// Input comes from the PendingInputCmd* that was given to us
+	Buffered	// Input comes from the PendingInputCmdBuffer in our internal data store
+};
+
+enum class EInstanceFlags
+{
+	None				= 0,
+	LocallyControlled	= 1 << 0,	// Instance is locally controlled
+};
+ENUM_CLASS_FLAGS(EInstanceFlags);
+
+// Static data about an instance (data that we don't track changes frame-to-frame)
+template<typename AsyncModelDef>
+struct TAsncInstanceStaticData
+{
+	HOIST_ASYNCMODELDEF_TYPES()
+
+	uint16 Index;
+	LocalStateType LocalState;
+};
+
+template<typename AsyncModelDef>
+struct TAsyncFrameInstanceData
+{
+	HOIST_ASYNCMODELDEF_TYPES()
+
+	InputCmdType InputCmd;
+	NetStateType NetState;
+};
+
+// A snapshot of the world for a given frame (only contains per-frame instance data)
+template<typename AsyncModelDef>
+struct TAsncFrameSnapshot
+{
+	HOIST_ASYNCMODELDEF_TYPES();
+
+	TArray<InputCmdType> InputCmds;
+	TArray<NetStateType> NetStates;
+
+	void Reset()
+	{
+		InputCmds.Reset();
+		NetStates.Reset();
+	}
+};
+
+template<typename AsyncModelDef>
+struct TAsyncPendingInputCmdPtr
+{
+	HOIST_ASYNCMODELDEF_TYPES()
+
+	int32 Index;
+	InputCmdType* PendingInputCmd = nullptr;
+	
+	bool operator==(const int32& idx) const
+	{
+		return Index == idx;
+	}
+};
+
+template<typename AsyncModelDef>
+struct TAsyncPendingInputCmdBuffer
+{
+	HOIST_ASYNCMODELDEF_TYPES()
+
+	int32 Index;
+	TStaticArray<InputCmdType, InputCmdBufferSize> Buffer;
+
+	bool operator==(const int32& idx) const
+	{
+		return Index == idx;
+	}
+};
+
+template<typename AsyncModelDef>
+struct TAsyncOutputWriteTargets
+{
+	HOIST_ASYNCMODELDEF_TYPES()
+
+	NetStateType* OutNetState = nullptr;
+};
+
+template<typename AsyncModelDef>
+struct TAsyncNetRecvData
+{
+	HOIST_ASYNCMODELDEF_TYPES()
+
+	struct FInstance
+	{
+		mutable int32 Frame = 0;
+		uint8 Flags = 0;
+		InputCmdType InputCmd;
+		NetStateType NetState;
+
+		InputCmdType LatestInputCmd; // latest, unprocessed InputCmd from server
+	};
+
+
+	TArray<FInstance> NetRecvInstances;
+	TBitArray<> NetRecvDirtyMask;
+	TBitArray<> PendingCorrectionMask;
+
+	void Reset()
+	{
+		NetRecvInstances.Reset();
+		NetRecvDirtyMask.Reset();
+		PendingCorrectionMask.Reset();
+	}
+
+	void MarkIndexDirty(const int32 idx)
+	{
+		NpResizeForIndex(NetRecvInstances, idx);
+		NpResizeAndSetBit(NetRecvDirtyMask, idx, true);
+	}
+};
+
+template<typename AsyncModelDef>
+struct TAsyncLocalStateMod
+{
+	HOIST_ASYNCMODELDEF_TYPES()
+
+	TAsyncLocalStateMod(FNetworkPredictionAsyncID id, TUniqueFunction<void(LocalStateType&)> f)
+		: ID(id), Func(MoveTemp(f)) {}
+
+	FNetworkPredictionAsyncID ID;
+	TUniqueFunction<void(LocalStateType&)> Func;
+
+	bool operator==(const FNetworkPredictionAsyncID& id) const { return ID == id; }
+};
+
+template<typename AsyncModelDef>
+struct TAsyncNetStateMod
+{
+	HOIST_ASYNCMODELDEF_TYPES()
+
+	TAsyncNetStateMod(int32 i, TUniqueFunction<void(NetStateType&)> f)
+		: Index(i), Func(MoveTemp(f)) {}
+
+	int32 Index;
+	TUniqueFunction<void(NetStateType&)> Func;
+};
+
+// ---------------------------------------------------------------------
+// DataStore for Marshaled data GT->PT (lives on the SimCallback Input)
+// ---------------------------------------------------------------------
+
+template<typename AsyncModelDef>
+struct TAsyncModelDataStore_Input
+{
+	HOIST_ASYNCMODELDEF_TYPES()
+
+	struct FNewInstance
+	{
+		FNetworkPredictionAsyncID ID;
+		TAsncInstanceStaticData<AsyncModelDef> StaticData;
+		NetStateType NetState;
+	};
+
+	struct FLocalInputCmd
+	{
+		FLocalInputCmd(int32 InIndex, InputCmdType& InInputCmd)
+			: Index(InIndex), InputCmd(InInputCmd) {}
+		int32 Index;
+		InputCmdType InputCmd;
+	};
+
+	TArray<FNewInstance> NewInstances;
+	TArray<TAsyncLocalStateMod<AsyncModelDef>> LocalStateMods;
+	TArray<TAsyncNetStateMod<AsyncModelDef>> NetStateMods;
+	TArray<FLocalInputCmd> LocalInputCmds;
+	TArray<FNetworkPredictionAsyncID> DeletedInstances;
+
+	void Reset()
+	{
+		NewInstances.Reset();
+		LocalStateMods.Reset();
+		NetStateMods.Reset();
+		LocalInputCmds.Reset();
+		DeletedInstances.Reset();
+	}
+};
+
+// ---------------------------------------------------------------------
+// DataStore for Marshaled data PT->GT (lives on the SimCallback Output)
+// ---------------------------------------------------------------------
+
+template<typename AsyncModelDef>
+struct TAsyncModelDataStore_Output
+{
+	void Reset()
+	{
+		Snapshot.Reset();
+	}
+
+	TAsncFrameSnapshot<AsyncModelDef> Snapshot;
+};
+
+// ---------------------------------------------------------------------
+// DataStore for internal data (Accessible on the PT)
+// ---------------------------------------------------------------------
+
+// The internal data store that is accessible on the PT
+template<typename AsyncModelDef>
+struct TAsyncModelDataStore_Internal
+{
+	HOIST_ASYNCMODELDEF_TYPES()	
+	TSortedMap<FNetworkPredictionAsyncID, TAsncInstanceStaticData<AsyncModelDef>> Instances;
+	TFrameStorage<TAsncFrameSnapshot<AsyncModelDef>> Frames;
+
+	TQueue<TAsyncNetRecvData<AsyncModelDef>> NetRecvQueue;
+	TAsyncNetRecvData<AsyncModelDef> NetRecvData;
+
+	void Reset()
+	{
+		Instances.Reset();
+		NetRecvData.Reset();
+	}
+};
+
+// External: for the game thread
+template<typename AsyncModelDef>
+struct TAsyncModelDataStore_External
+{
+	HOIST_ASYNCMODELDEF_TYPES()
+
+	TSortedMap<FNetworkPredictionAsyncID, TAsncInstanceStaticData<AsyncModelDef>> Instances;
+	TBitArray<> FreeIndices;
+
+	TAsncFrameSnapshot<AsyncModelDef> LatestSnapshot;
+
+	TArray<TAsyncPendingInputCmdPtr<AsyncModelDef>>	InActivePendingInputCmds;	
+	TArray<TAsyncPendingInputCmdPtr<AsyncModelDef>>	ActivePendingInputCmds; // PendingINputCmds that get marshaled directly to PT
+	TSortedMap<FNetworkPredictionAsyncID, TUniqueObj<TAsyncPendingInputCmdBuffer<AsyncModelDef>>> PendingInputCmdBuffers; // GT buffered InputCmds
+
+	TArray<TAsyncOutputWriteTargets<AsyncModelDef>>	OutputWriteTargets;
+
+	TAsyncNetRecvData<AsyncModelDef> PendingNetRecv;	// Network Data is written here (clients, receiving)
+
+	TArray<TAsyncLocalStateMod<AsyncModelDef>> DeferredLocalStateMods;
+	
+	void Reset()
+	{
+		Instances.Reset();
+		FreeIndices.Reset();
+		LatestSnapshot.Reset();
+		InActivePendingInputCmds.Reset();
+		ActivePendingInputCmds.Reset();
+		PendingInputCmdBuffers.Reset();
+		OutputWriteTargets.Reset();
+		PendingNetRecv.Reset();
+		DeferredLocalStateMods.Reset();
+	}
+};
+
+
+
+
+} // namespace UE_NP
+
diff --git a/Engine/Plugins/Runtime/NetworkPrediction/Source/NetworkPrediction/Public/Async/NetworkPredictionAsyncModelDef.h b/Engine/Plugins/Runtime/NetworkPrediction/Source/NetworkPrediction/Public/Async/NetworkPredictionAsyncModelDef.h
new file mode 100644
--- /dev/null
+++ b/Engine/Plugins/Runtime/NetworkPrediction/Source/NetworkPrediction/Public/Async/NetworkPredictionAsyncModelDef.h
@@ -0,0 +1,38 @@
+// Copyright Epic Games, Inc. All Rights Reserved.
+
+#pragma once
+
+class AActor;
+class APlayerController;
+
+namespace UE_NP {
+
+using FAsyncModelDefID = int32;
+#define NP_ASYNC_MODEL_BODY() static UE_NP::FAsyncModelDefID ID;
+
+struct FNetworkPredictionAsyncModelDef
+{
+	// Actual defs should have:
+	// NP_MODEL_BODY(); 
+	//	static const TCHAR* GetName() { return TEXT("YourName"); }
+	//	static constexpr int32 GetSortPriority() { return (int32)ENetworkPredictionSortPriority::PreKinematicMovers; }
+
+	using InputCmdType		= void;	// The state the controller is authoritative over (e.g, can be the client)
+	using NetStateType	= void;	// The object state that we do network and track frame-to-frame
+	using LocalStateType = void; // Local state that is not networked and we don't track frame-to-frame. (e.g, the local PhysicsProxy*)
+	using OutStateType		= void;	// Output state marshaled back to the GT
+	
+	using ControlKeyType = APlayerController*;	// Identifier for the controller
+
+	using SimulationTickType = void; // what to call ::SimulationTick on 
+};
+
+#define HOIST_ASYNCMODELDEF_TYPES() \
+	using InputCmdType = typename AsyncModelDef::InputCmdType; \
+	using LocalStateType = typename AsyncModelDef::LocalStateType; \
+	using NetStateType = typename AsyncModelDef::NetStateType; \
+	using OutStateType = typename AsyncModelDef::OutStateType; \
+	using ControlKeyType = typename AsyncModelDef::ControlKeyType; \
+	using SimulationTickType = typename AsyncModelDef::SimulationTickType;
+
+} // namespace UE_NP
\ No newline at end of file
diff --git a/Engine/Plugins/Runtime/NetworkPrediction/Source/NetworkPrediction/Public/Async/NetworkPredictionAsyncModelDefRegistry.h b/Engine/Plugins/Runtime/NetworkPrediction/Source/NetworkPrediction/Public/Async/NetworkPredictionAsyncModelDefRegistry.h
new file mode 100644
--- /dev/null
+++ b/Engine/Plugins/Runtime/NetworkPrediction/Source/NetworkPrediction/Public/Async/NetworkPredictionAsyncModelDefRegistry.h
@@ -0,0 +1,106 @@
+// Copyright Epic Games, Inc. All Rights Reserved.
+
+#pragma once
+
+#include "NetworkPredictionLog.h"
+#include "NetworkPredictionAsyncModelDef.h"
+#include "NetworkPredictionCheck.h"
+
+namespace UE_NP {
+
+// Basic ModelDef type registry. ModelDefs are registered and assigned an ID (their ModelDef::ID) based on ModelDef::SortPriority.
+// This ID is used as indices into the various NP service arrays
+class FNetworkPredictionAsyncModelDefRegistry
+{
+public:
+
+	NETWORKPREDICTION_API static FNetworkPredictionAsyncModelDefRegistry& Get()
+	{
+		return Singleton;
+	}
+
+	template<typename AsyncModelDef>
+	void RegisterType()
+	{
+		bFinalized = false; // New type must re-finalize
+
+		if (!npEnsure(ModelDefList.Contains(&AsyncModelDef::ID) == false))
+		{
+			return;
+		}
+
+		FTypeInfo TypeInfo =
+		{
+			&AsyncModelDef::ID,					// Must include NP_MODEL_BODY()
+			AsyncModelDef::GetSortPriority(),	// Must implement ::GetSortPriorty()
+			AsyncModelDef::GetName()			// Must implement ::GetName()
+		};
+
+		ModelDefList.Emplace(TypeInfo);
+	}
+
+	void FinalizeTypes()
+	{
+		if (bFinalized)
+		{
+			return;
+		}
+
+		bFinalized = true;
+		ModelDefList.Sort([](const FTypeInfo& LHS, const FTypeInfo& RHS) -> bool
+			{
+				if (LHS.SortPriority == RHS.SortPriority)
+				{
+					UE_LOG(LogNetworkPrediction, Log, TEXT("AsyncModelDefs %s and %s have same sort priority. Using lexical sort as backup"), LHS.Name, RHS.Name);
+					int32 StrCmpResult = FCString::Strcmp(LHS.Name, RHS.Name);
+					npEnsureMsgf(StrCmpResult != 0, TEXT("Duplicate AsyncModelDefs appear to have been registered."));
+					return StrCmpResult > 0;
+				}
+
+				return LHS.SortPriority < RHS.SortPriority; 
+			});
+
+		int32 Count = 0;
+		for (FTypeInfo& TypeInfo : ModelDefList)
+		{
+			*TypeInfo.IDPtr = Count++;
+		}
+	}
+
+private:
+
+	NETWORKPREDICTION_API static FNetworkPredictionAsyncModelDefRegistry Singleton;
+
+	struct FTypeInfo
+	{
+		bool operator==(const FAsyncModelDefID* OtherIDPtr) const
+		{
+			return IDPtr == OtherIDPtr;
+		}
+
+		FAsyncModelDefID* IDPtr = nullptr;
+		int32 SortPriority;
+		const TCHAR* Name;
+	};
+
+	TArray<FTypeInfo> ModelDefList;
+
+	bool bFinalized = false;
+};
+
+template<typename AsyncModelDef>
+struct TNetworkPredictionAsyncModelDefRegisterHelper
+{
+	TNetworkPredictionAsyncModelDefRegisterHelper()
+	{
+		FNetworkPredictionAsyncModelDefRegistry::Get().RegisterType<AsyncModelDef>();
+	}
+};
+
+} // namespace UE_NP
+
+// Helper to register ModelDef type.
+// Sets static ID to -1 (invalid) and calls global registration function
+#define NP_ASYNC_MODEL_REGISTER(X) \
+	UE_NP::FAsyncModelDefID X::ID=-1; \
+	static UE_NP::TNetworkPredictionAsyncModelDefRegisterHelper<X> NetModelAr_##X = UE_NP::TNetworkPredictionAsyncModelDefRegisterHelper<X>();
\ No newline at end of file
diff --git a/Engine/Plugins/Runtime/NetworkPrediction/Source/NetworkPrediction/Public/Async/NetworkPredictionAsyncPersistentStorage.h b/Engine/Plugins/Runtime/NetworkPrediction/Source/NetworkPrediction/Public/Async/NetworkPredictionAsyncPersistentStorage.h
new file mode 100644
--- /dev/null
+++ b/Engine/Plugins/Runtime/NetworkPrediction/Source/NetworkPrediction/Public/Async/NetworkPredictionAsyncPersistentStorage.h
@@ -0,0 +1,70 @@
+// Copyright Epic Games, Inc. All Rights Reserved.
+
+#pragma once
+
+namespace UE_NP {
+
+
+// Something like this is needed and should be managed by the async physics system for persistent PT storage
+// Note that the way this is used now, is TPersistentStorage<TSortedMap<ObjKey, ObjState>> - e.g, the whole snapshot
+// We could maybe provide a way to do sparse allocated per object storage.
+template<typename T>
+struct TPersistentStorage
+{
+	int32 GetTailFrame() const { return FMath::Max(0, HeadFrame - Buffer.Num() + 1); }
+	int32 GetHeadFrame() const { return HeadFrame; }
+
+	// Const function that returns the frame that is safe to write to.
+	// (Something like) this is the only thing ISimCallbackObject should really be using
+	T* GetWritable() const { return const_cast<T*>(&Buffer[HeadFrame % Buffer.Num()]); }
+
+	T* GetFrame(int32 Frame)
+	{
+		if (Frame <= HeadFrame && Frame >= GetTailFrame())
+		{
+			return &Buffer[Frame % Buffer.Num()];
+		}
+
+		return nullptr;
+	}
+
+	void SetHeadFrame(int32 NewHeadFrame)
+	{
+		HeadFrame = NewHeadFrame;
+	}
+
+	void IncrementFrame()
+	{
+		T* Pre = GetWritable();
+		HeadFrame++;
+		T* Next = GetWritable();
+		*Next = *Pre;
+	}
+
+	void RollbackToFrame(int32 Frame)
+	{
+		npEnsure(Frame >= GetTailFrame());
+		HeadFrame = Frame;
+	}
+
+private:
+
+	int32 HeadFrame = 0;
+	TStaticArray<T, 64> Buffer;
+};
+
+template<typename T>
+struct TFrameStorage
+{
+	T&       operator[](int32 Index) { return Buffer[Index % Buffer.Num()]; }
+	const T& operator[](int32 Index) const { return Buffer[Index % Buffer.Num()]; }
+
+	int32 Num() const { return Buffer.Num(); }
+
+private:
+	
+	TStaticArray<T, UE_NP::NumFramesStorage> Buffer;
+};
+
+
+} // namespace UE_NP
\ No newline at end of file
diff --git a/Engine/Plugins/Runtime/NetworkPrediction/Source/NetworkPrediction/Public/Async/NetworkPredictionAsyncProxy.h b/Engine/Plugins/Runtime/NetworkPrediction/Source/NetworkPrediction/Public/Async/NetworkPredictionAsyncProxy.h
new file mode 100644
--- /dev/null
+++ b/Engine/Plugins/Runtime/NetworkPrediction/Source/NetworkPrediction/Public/Async/NetworkPredictionAsyncProxy.h
@@ -0,0 +1,64 @@
+// Copyright Epic Games, Inc. All Rights Reserved.
+
+#pragma once
+
+#include "UObject/Class.h"
+#include "NetworkPredictionAsyncID.h"
+#include "NetworkPredictionAsyncProxy.generated.h"
+
+class APlayerController;
+
+namespace UE_NP
+{
+	class FNetworkPredictionAsyncWorldManager;
+};
+
+USTRUCT()
+struct NETWORKPREDICTION_API FNetworkPredictionAsyncProxy
+{
+	GENERATED_BODY()
+
+	// acquires an ID but does not register with any sims (a proxy may participate in multiple sims at once)
+	bool RegisterProxy(UWorld* World);
+
+	// Unregisters proxy with NP system and all sims it was participating in
+	void UnregisterProxy();
+
+	void RegisterController(APlayerController* PC);
+
+	bool NetSerialize(FArchive& Ar, class UPackageMap* Map, bool& bOutSuccess);
+	bool Identical(const FNetworkPredictionAsyncProxy* Other, uint32 PortFlags) const;
+	void OnPreReplication();
+
+	// Include NetworkPredictionAsyncProxyImpl.h in your .cpp to use these.
+	template<typename AsyncModelDef>
+	void RegisterSim(typename AsyncModelDef::LocalStateType&& LocalState, typename AsyncModelDef::NetStateType&& NetState, typename AsyncModelDef::InputCmdType* PendingInputCmd, typename AsyncModelDef::NetStateType* OutNetState);
+
+	// Unregisters with specific sim. Note that UnregisterProxy will do this automatically
+	template<typename AsyncModelDef>
+	void UnregisterSim();
+
+	template<typename AsyncModelDef>
+	void ModifyLocalState(TUniqueFunction<void(typename AsyncModelDef::LocalStateType&)> Func);
+
+	template<typename AsyncModelDef>
+	void ModifyNetState(TUniqueFunction<void(typename AsyncModelDef::NetStateType&)> Func);
+
+private:
+
+	UE_NP::FNetworkPredictionAsyncWorldManager* Manager = nullptr;
+	UE_NP::FNetworkPredictionAsyncID ID;
+
+	int32 CachedLatestFrame = INDEX_NONE;
+};
+
+template<>
+struct TStructOpsTypeTraits<FNetworkPredictionAsyncProxy> : public TStructOpsTypeTraitsBase2<FNetworkPredictionAsyncProxy>
+{
+	enum
+	{
+		WithNetSerializer = true,
+		WithNetSharedSerialization = true,
+		WithIdentical = true,
+	};
+ };
\ No newline at end of file
diff --git a/Engine/Plugins/Runtime/NetworkPrediction/Source/NetworkPrediction/Public/Async/NetworkPredictionAsyncProxyImpl.h b/Engine/Plugins/Runtime/NetworkPrediction/Source/NetworkPrediction/Public/Async/NetworkPredictionAsyncProxyImpl.h
new file mode 100644
--- /dev/null
+++ b/Engine/Plugins/Runtime/NetworkPrediction/Source/NetworkPrediction/Public/Async/NetworkPredictionAsyncProxyImpl.h
@@ -0,0 +1,44 @@
+// Copyright Epic Games, Inc. All Rights Reserved.
+
+#pragma once
+
+#include "Async/NetworkPredictionAsyncProxy.h"
+#include "Async/NetworkPredictionAsyncWorldManager.h"
+
+template<typename AsyncModelDef>
+void FNetworkPredictionAsyncProxy::RegisterSim(typename AsyncModelDef::LocalStateType&& LocalState, typename AsyncModelDef::NetStateType&& NetState, typename AsyncModelDef::InputCmdType* PendingInputCmd, typename AsyncModelDef::NetStateType* OutNetState)
+{
+	if (!(ensureMsgf(Manager, TEXT("Proxy not registered prior to call to RegisterSim. Please call RegisterProxy first"))))
+	{
+		return;
+	}
+
+	Manager->RegisterSim<AsyncModelDef>(ID, MoveTemp(LocalState), MoveTemp(NetState), PendingInputCmd, OutNetState);
+}
+
+template<typename AsyncModelDef>
+void FNetworkPredictionAsyncProxy::UnregisterSim()
+{
+	if (npEnsureSlow(Manager) && npEnsureSlow(ID.IsValid()))
+	{
+		Manager->UnregisterSim<AsyncModelDef>(ID);
+	}
+}
+
+template<typename AsyncModelDef>
+void FNetworkPredictionAsyncProxy::ModifyLocalState(TUniqueFunction<void(typename AsyncModelDef::LocalStateType&)> Func)
+{
+	if (npEnsureSlow(Manager) && npEnsureSlow(ID.IsValid()))
+	{
+		Manager->ModifyLocalState<AsyncModelDef>(ID, MoveTemp(Func));
+	}
+}
+
+template<typename AsyncModelDef>
+void FNetworkPredictionAsyncProxy::ModifyNetState(TUniqueFunction<void(typename AsyncModelDef::NetStateType&)> Func)
+{
+	if (npEnsureSlow(Manager) && npEnsureSlow(ID.IsValid()))
+	{
+		Manager->ModifyNetState<AsyncModelDef>(ID, MoveTemp(Func));
+	}
+}
\ No newline at end of file
diff --git a/Engine/Plugins/Runtime/NetworkPrediction/Source/NetworkPrediction/Public/Async/NetworkPredictionAsyncSimCallback.h b/Engine/Plugins/Runtime/NetworkPrediction/Source/NetworkPrediction/Public/Async/NetworkPredictionAsyncSimCallback.h
new file mode 100644
--- /dev/null
+++ b/Engine/Plugins/Runtime/NetworkPrediction/Source/NetworkPrediction/Public/Async/NetworkPredictionAsyncSimCallback.h
@@ -0,0 +1,58 @@
+// Copyright Epic Games, Inc. All Rights Reserved.
+
+#pragma once
+
+#include "Physics/Experimental/PhysScene_Chaos.h"
+#include "NetworkPredictionAsyncInstanceData.h"
+
+namespace UE_NP {
+
+struct FNetworkPredictionSimCallbackInput : public Chaos::FSimCallbackInput
+{
+	bool NetSendInputCmd(FNetBitWriter& Ar) override { return true; } // FIXME: remove these functions	
+	bool NetRecvInputCmd(APlayerController* PC, FNetBitReader& Ar) override { return true; } // FIXME: remove these functions
+
+	void Reset()
+	{
+		DataStoreCollection.Reset();
+	}
+
+	template<typename AsyncModelDef>
+	TAsyncModelDataStore_Input<AsyncModelDef>* GetDataStore()
+	{
+		return DataStoreCollection.GetDataStore<TAsyncModelDataStore_Input<AsyncModelDef>, AsyncModelDef>();
+	}
+
+	template<typename AsyncModelDef>
+	const TAsyncModelDataStore_Input<AsyncModelDef>* GetDataStore() const
+	{
+		return DataStoreCollection.GetDataStore<TAsyncModelDataStore_Input<AsyncModelDef>, AsyncModelDef>();
+	}
+
+private:
+
+	FAsyncDataStoreCollection DataStoreCollection;
+};
+
+
+struct FNetworkPredictionSimCallbackOutput : public Chaos::FSimCallbackOutput
+{
+	void Reset()
+	{
+		SimulationFrame = 0;
+		LocalFrameOffset = 0;
+	}
+
+	int32 SimulationFrame = 0;
+	int32 LocalFrameOffset = 0;
+
+	template<typename AsyncModelDef>
+	TAsyncModelDataStore_Output<AsyncModelDef>* GetDataStore()
+	{
+		return DataStoreCollection.GetDataStore<TAsyncModelDataStore_Output<AsyncModelDef>, AsyncModelDef>();
+	}
+
+	FAsyncDataStoreCollection DataStoreCollection;
+};
+
+} // namespace UE_NP
\ No newline at end of file
diff --git a/Engine/Plugins/Runtime/NetworkPrediction/Source/NetworkPrediction/Public/Async/NetworkPredictionAsyncWorldManager.h b/Engine/Plugins/Runtime/NetworkPrediction/Source/NetworkPrediction/Public/Async/NetworkPredictionAsyncWorldManager.h
new file mode 100644
--- /dev/null
+++ b/Engine/Plugins/Runtime/NetworkPrediction/Source/NetworkPrediction/Public/Async/NetworkPredictionAsyncWorldManager.h
@@ -0,0 +1,581 @@
+// Copyright Epic Games, Inc. All Rights Reserved.
+
+#pragma once
+
+#include "Physics/Experimental/PhysScene_Chaos.h"
+#include "PBDRigidsSolver.h"
+#include "ChaosSolversModule.h"
+#include "RewindData.h"
+
+#include "NetworkPhysics.h"
+#include "NetworkPredictionLog.h"
+
+#include "Async/NetworkPredictionAsyncInstanceData.h"
+#include "Async/NetworkPredictionAsyncSimCallback.h"
+#include "Async/NetworkPredictionAsyncID.h"
+#include "Async/NetworkPredictionAsyncGlobalTickContext.h"
+
+#include "Services/NetworkPredictionAsyncServiceRegistry.h"
+#include "Services/NetworkPredictionAsyncService_Registration.inl"
+#include "Services/NetworkPredictionAsyncService_Marshal.inl"
+#include "Services/NetworkPredictionAsyncService_Ticking.inl"
+#include "Services/NetworkPredictionAsyncService_Reconcile.inl"
+#include "Services/NetworkPredictionAsyncService_NetSerialize.inl"
+
+
+namespace UE_NP {
+
+// -----------------------------------------------------------------------------------------
+//
+// -----------------------------------------------------------------------------------------
+
+class FNetworkPredictionAsyncSimCallback: public Chaos::TSimCallbackObject<FNetworkPredictionSimCallbackInput,FNetworkPredictionSimCallbackOutput>
+{
+public:
+
+	virtual void OnPreSimulate_Internal() override
+	{
+		const FNetworkPredictionSimCallbackInput* Input = this->GetConsumerInput_Internal();
+		if (Input == nullptr)
+		{
+			return;
+		}
+
+		Chaos::FPhysicsSolver* PhysicsSolver = static_cast<Chaos::FPhysicsSolver*>(this->GetSolver());
+
+		const int32 LocalFrame = PhysicsSolver->GetCurrentFrame();
+		const int32 SimulationFrame = LocalFrame - LocalFrameOffset;
+		
+		bool bIsResim = true;
+		if (LocalFrame > LatestFrame)
+		{
+			bIsResim = false;
+			LatestFrame = LocalFrame;
+		}
+
+		FNetworkPredictionSimCallbackOutput& Output = this->GetProducerOutputData_Internal();
+		Output.SimulationFrame = SimulationFrame;
+		Output.LocalFrameOffset = LocalFrameOffset;
+
+		// Marshal all data from GT
+		for (TUniquePtr<IAsyncMarshalService_Internal>& ServicePtr : MarshalServices.Array)
+		{
+			if (IAsyncMarshalService_Internal* Service = ServicePtr.Get())
+			{
+				Service->MarshalPreSimulate(LocalFrame, bIsResim, Input, Output);
+			}
+		}
+		
+		FNetworkPredictionAsyncTickContext Context(&DataStoreCollection);
+		Context.World = World;
+		Context.DeltaTimeSeconds = this->GetDeltaTime_Internal();
+		Context.LocalStorageFrame = LocalFrame+1; // MarshalPreSimulate "finalized" LocalFrame and we are running the tick on the next frame's storage
+		Context.SimulationFrame = SimulationFrame;
+		Context.bIsResim = bIsResim;
+
+		// Tick
+		for (TUniquePtr<IAsyncTickingService>& ServicePtr : TickingServices.Array)
+		{
+			if (IAsyncTickingService* Service = ServicePtr.Get())
+			{
+				Service->Tick_Internal(Context, Output);
+			}
+		}
+	}
+
+	void ApplyCorrections_Internal(int32 PhysicsStep, Chaos::FSimCallbackInput* Input) override { } // REMOVE
+	void FirstPreResimStep_Internal(int32 PhysicsStep) override { } // REMOVE
+	int32 TriggerRewindIfNeeded_Internal(int32 LastCompletedStep) { return INDEX_NONE; } // REMOVE (as virtual on TSimCallbackObject)
+
+	template<typename AsyncModelDef>
+	void InstantiateServicesForModelDef()
+	{
+		TAsyncModelDataStore_Internal<AsyncModelDef>* DataStore = GetDataStore<AsyncModelDef>();
+		MarshalServices.Instantiate<AsyncModelDef, TAsyncMarshalService_Internal<AsyncModelDef>>(DataStore);
+		TickingServices.Instantiate<AsyncModelDef, TAsyncTickingService<AsyncModelDef>>(DataStore);
+		ReconcileServices.Instantiate<AsyncModelDef, TAsyncReconcileService<AsyncModelDef>>(DataStore);
+	}
+
+	template<typename AsyncModelDef>
+	TAsyncModelDataStore_Internal<AsyncModelDef>* GetDataStore()
+	{
+		return DataStoreCollection.GetDataStore<TAsyncModelDataStore_Internal<AsyncModelDef>, AsyncModelDef>();
+	}
+
+	void OnContactModification_Internal(const TArrayView<Chaos::FPBDCollisionConstraintHandleModification>& Modifications) final override { }
+
+	int32 LatestFrame = -1; // Highest frame we've simmed. Used to detect resim
+	UWorld* World = nullptr;
+
+	// ------------------------------------------------------------------------------------
+
+	int32 Reconcile_Internal(int32 LastCompletedStep)
+	{
+		// pop latest localframe offset
+		while(LocalFrameOffsetQueue.Dequeue(LocalFrameOffset));
+
+		int32 RewindFrame = INDEX_NONE;
+		for (TUniquePtr<IAsyncReconcileService>& ServicePtr : ReconcileServices.Array)
+		{
+			if (IAsyncReconcileService* Service = ServicePtr.Get())
+			{
+				const int32 Frame = Service->Reconcile(LastCompletedStep, LocalFrameOffset);
+				if (Frame != INDEX_NONE)
+				{
+					RewindFrame = RewindFrame == INDEX_NONE ? Frame : FMath::Min<int32>(RewindFrame, Frame);
+				}
+			}
+		}
+
+		return RewindFrame;
+	}
+
+	void PushLatestLocalFrameOffset(int32 InLocalFrameOffset)
+	{
+		LocalFrameOffsetQueue.Enqueue(InLocalFrameOffset);
+	}
+
+private:
+
+	TQueue<int32> LocalFrameOffsetQueue;
+	int32 LocalFrameOffset = 0;
+
+	TAsyncServiceStorage<IAsyncMarshalService_Internal> MarshalServices;
+	TAsyncServiceStorage<IAsyncTickingService> TickingServices;
+	TAsyncServiceStorage<IAsyncReconcileService> ReconcileServices;
+
+	FAsyncDataStoreCollection DataStoreCollection;
+};
+
+// -----------------------------------------------------------------------------------------
+//
+// -----------------------------------------------------------------------------------------
+
+class FNetworkPredictionAsyncWorldManager: public INetworkPhysicsSubsystem
+{
+public:
+
+	FNetworkPredictionAsyncWorldManager(UWorld* InWorld)
+	{
+		World = InWorld;
+		if (ensure(World))
+		{
+			FPhysScene* PhysScene = World->GetPhysicsScene();
+			if (ensureAlways(PhysScene))
+			{
+				Chaos::FPhysicsSolver* Solver = PhysScene->GetSolver();
+				if (ensureAlways(Solver))
+				{
+					if (ensureAlways(Solver->GetRewindCallback()))
+					{
+						AsyncCallback = Solver->CreateAndRegisterSimCallbackObject_External<FNetworkPredictionAsyncSimCallback>(true, true);
+						AsyncCallback->World = World;
+					}
+					else
+					{
+						UE_LOG(LogNetworkPhysics, Warning, TEXT("Rewind not enabled on Physics solver. FNetworkPredictionAsyncWorldManager will be disabled"));
+					}
+				}
+			}
+		}
+	}
+
+	bool RegisterInstance(FNetworkPredictionAsyncID& OutID, bool bForClient)
+	{
+		// Clients get temporary negative SpawnIDs so that they can initialize local state etc before finding out the replicated ID from the server
+		OutID = FNetworkPredictionAsyncID(bForClient ? ClientSpawnIDCounter-- : SpawnIDCounter++);
+		return true;
+	}
+
+	void UnregisterInstance(FNetworkPredictionAsyncID& ID)
+	{
+		npCheckSlow(AsyncCallback);
+		FNetworkPredictionSimCallbackInput* Input = AsyncCallback->GetProducerInputData_External();
+		npCheckSlow(Input);
+
+		ForEachSim<IAsyncRegistrationService>(GlobalInstanceData_External.FindChecked(ID), RegistrationServices, [&](IAsyncRegistrationService* Service)
+		{
+			Service->UnregisterInstance(ID, Input);
+		});
+
+		GlobalInstanceData_External.FindAndRemoveChecked(ID);
+		int32 idx = PlayerControllerMap.IndexOfByKey(ID);
+		if (idx != INDEX_NONE)
+		{
+			PlayerControllerMap.RemoveAtSwap(idx, 1, false);
+		}
+		ID.Reset();
+	}
+
+	// Remaps the temporary, negative ClientID to the received serverID
+	void RemapClientInstance(FNetworkPredictionAsyncID ClientID, FNetworkPredictionAsyncID ServerID)
+	{
+		npCheckSlow(AsyncCallback);
+		FNetworkPredictionSimCallbackInput* Input = AsyncCallback->GetProducerInputData_External();
+		npCheckSlow(Input);
+
+		FGlobalInstanceData Data = GlobalInstanceData_External.FindAndRemoveChecked(ClientID);
+		GlobalInstanceData_External.Emplace(ServerID, Data);
+
+		if (FControllerInfo* ControllerInfo = PlayerControllerMap.FindByKey(ClientID))
+		{
+			ControllerInfo->ID = ServerID;
+		}
+
+		ForEachSim<IAsyncRegistrationService>(Data, RegistrationServices, [&](IAsyncRegistrationService* Service)
+		{
+			Service->RemapClientInstance(ClientID, ServerID, Input);
+		});
+	}
+
+	template<typename AsyncModelDef>
+	void RegisterSim(FNetworkPredictionAsyncID ID, typename AsyncModelDef::LocalStateType&& InitialValueLocal, typename AsyncModelDef::NetStateType&& InitialValueNet, 
+		typename AsyncModelDef::InputCmdType* PendingInputCmd, typename AsyncModelDef::NetStateType* OutNetState)
+	{
+		// Ensure services are instantiated for this AsyncModelDef
+		InstantiateServicesForModelDef<AsyncModelDef>();
+
+		// Track that this ID is now participating in this sim
+		npEnsureSlow(AsyncModelDef::ID < 32);
+		GlobalInstanceData_External.FindOrAdd(ID).Sims |= (1 << AsyncModelDef::ID);
+
+		// Call into RegistrationService
+		npCheckSlow(AsyncCallback);
+		FNetworkPredictionSimCallbackInput* Input = AsyncCallback->GetProducerInputData_External();
+		npCheckSlow(Input);
+
+		TAsyncRegistrationService<AsyncModelDef>* Service = (TAsyncRegistrationService<AsyncModelDef>*)RegistrationServices.Array[AsyncModelDef::ID].Get();
+		Service->RegisterInstance(ID, Input, MoveTemp(InitialValueLocal), MoveTemp(InitialValueNet), PendingInputCmd, OutNetState);
+	}
+
+	template<typename AsyncModelDef>
+	void UnregisterSim(FNetworkPredictionAsyncID ID)
+	{
+		// Call into RegistrationService
+		npCheckSlow(AsyncCallback);
+		FNetworkPredictionSimCallbackInput* Input = AsyncCallback->GetProducerInputData_External();
+		npCheckSlow(Input);
+
+		RegistrationServices.Array[AsyncModelDef::ID]->UnregisterInstance(ID, Input);
+
+		GlobalInstanceData_External.FindOrAdd(ID).Sims &= ~(1 << AsyncModelDef::ID);
+	}
+	
+	// Binds an NP instance to a PlayerController. PlayerController is required for networking, by setting a controlling PC the server will
+	// accept InputCmds from this client.
+	void RegisterController(FNetworkPredictionAsyncID ID, APlayerController* PC);
+
+	void SetInputSource(FNetworkPredictionAsyncID ID, ENetworkPredictionAsyncInputSource Type, bool bIsLocal)
+	{
+		FGlobalInstanceData& InstanceData = GlobalInstanceData_External.FindChecked(ID);
+		InstanceData.Flags = bIsLocal ? (InstanceData.Flags | (uint8)EInstanceFlags::LocallyControlled) : (InstanceData.Flags & ~(uint8)EInstanceFlags::LocallyControlled);
+
+		ForEachSim<IAsyncRegistrationService>(GlobalInstanceData_External.FindChecked(ID), RegistrationServices, [&](IAsyncRegistrationService* Service)
+		{
+			Service->SetInputSource(ID, Type, bIsLocal);
+		});
+	}
+
+	// Serializes a registered instance to the archive
+	//	This includes the latest state/input + a "future input" which is the latest, but currently unprocessed, InputCmd we have on the GT
+	void NetSerializeInstance(FNetworkPredictionAsyncID ID, FArchive& Ar)
+	{
+		FGlobalInstanceData& InstanceData = GlobalInstanceData_External.FindChecked(ID);
+
+		// 1. Frame (this is the latest frame we have results for, on the game thread)
+		int32 Frame = GlobalOutput_External.SimulationFrame;
+		Ar << Frame;
+
+		uint8 FutureDelta = 0;
+
+		if (Ar.IsSaving())
+		{
+			// Find what the latest InputCmd frame we have for this Instance
+			// (for Remote players will be InputBuffer's LastFrame)
+			// (for local players it will be the PendingInputCmd/NextPhysicsStep)
+			int32 FutureFrame = 0;
+			FControllerInfo* ControllerInfo = PlayerControllerMap.FindByKey(ID);
+
+			int32 PendingInputFrame = INDEX_NONE;
+			
+			if (ControllerInfo && ControllerInfo->InputSource == ENetworkPredictionAsyncInputSource::Buffered)
+			{
+				const int32 SignedFutureDelta = ControllerInfo->LastFrame - GlobalOutput_External.SimulationFrame;
+				if (SignedFutureDelta > 0 && npEnsure(SignedFutureDelta < 16))
+				{
+					FutureDelta = (uint8)SignedFutureDelta;
+					FutureFrame = ControllerInfo->LastFrame;
+					PendingInputFrame = FutureFrame; // This is telling the NetSend call below to send this frame from the PendingInputCmd
+				}
+			}
+			else
+			{
+				const int32 SignedFutureDelta = this->NextPhysicsStep - GlobalOutput_External.SimulationFrame;
+				if (npEnsure(SignedFutureDelta > 0) && npEnsure(SignedFutureDelta < 16))
+				{
+					FutureDelta = (uint8)SignedFutureDelta;
+				}
+			}
+
+			// 2. Future Frame Delta (this is the latest frame we have input for, but haven't processed yet)
+			Ar << FutureDelta;
+
+			ForEachSim<INetSerializeService>(InstanceData, NetSerializeServices, [&](INetSerializeService* Service)
+			{
+				// 3. AsyncModelDef-specific payloads
+				Service->NetSend(ID, InstanceData.Flags, Frame, PendingInputFrame, (FNetBitWriter&)Ar);
+			});	
+		}
+		else
+		{
+			// 2. Future Frame Delta (this is a future frame that we have input for, but no NetState etc since the server has not processed this frame yet)
+			Ar << FutureDelta;
+
+			const int32 FutureFrame = Frame + (int32)FutureDelta;
+			
+			ForEachSim<INetSerializeService>(InstanceData, NetSerializeServices, [&](INetSerializeService* Service)
+			{
+				// 3. AsyncModelDef-specific payloads
+				Service->NetRecv(ID, InstanceData.Flags, Frame, FutureFrame, (FNetBitReader&)Ar);
+			});	
+		}
+	}
+
+	template<typename AsyncModelDef>
+	void ModifyLocalState(FNetworkPredictionAsyncID ID, TUniqueFunction<void(typename AsyncModelDef::LocalStateType&)> Func)
+	{
+		FNetworkPredictionSimCallbackInput* AsyncInput = AsyncCallback->GetProducerInputData_External();
+		if (ensure(AsyncInput))
+		{
+			TAsyncMarshalService_External<AsyncModelDef>* Service = (TAsyncMarshalService_External<AsyncModelDef>*)MarshalServices.Array[AsyncModelDef::ID].Get();
+			Service->ModifyLocalState(ID, MoveTemp(Func), AsyncInput);
+		}
+	}
+
+	template<typename AsyncModelDef>
+	void ModifyNetState(FNetworkPredictionAsyncID ID, TUniqueFunction<void(typename AsyncModelDef::NetStateType&)> Func)
+	{
+		FNetworkPredictionSimCallbackInput* AsyncInput = AsyncCallback->GetProducerInputData_External();
+		if (ensure(AsyncInput))
+		{
+			TAsyncMarshalService_External<AsyncModelDef>* Service = (TAsyncMarshalService_External<AsyncModelDef>*)MarshalServices.Array[AsyncModelDef::ID].Get();
+			Service->ModifyNetState(ID, MoveTemp(Func), AsyncInput);
+		}
+	}
+
+	// -----------------------------------------------------------------------------------------------------
+	// INetworkPhysicsSubsystem
+	// -----------------------------------------------------------------------------------------------------
+
+	static FName GetName() { return FName("NPAsync"); }
+	static FNetworkPredictionAsyncWorldManager* Get(UWorld* World)
+	{
+		if (UNetworkPhysicsManager* NetworkPhysicsManager = World->GetSubsystem<UNetworkPhysicsManager>())
+		{
+			if (FNetworkPredictionAsyncWorldManager* Existing = NetworkPhysicsManager->GetSubsystem<FNetworkPredictionAsyncWorldManager>())
+			{
+				return Existing;
+			}
+
+			FNetworkPredictionAsyncWorldManager* NewInstance = NetworkPhysicsManager->RegisterSubsystem<FNetworkPredictionAsyncWorldManager>(MakeUnique<FNetworkPredictionAsyncWorldManager>(World));
+			return NewInstance;
+		}
+		ensure(false);
+		return nullptr;
+	}
+
+	void PostNetRecv(UWorld* InWorld, int32 LocalOffset, int32 LastProcessedFrame) override
+	{
+	}
+
+	void PreNetSend(UWorld* InWorld, float DeltaSeconds) override
+	{
+	}
+
+	void ProcessInputs_External(int32 PhysicsStep, int32 LocalFrameOffset, bool& bOutSendClientInputCmd) override
+	{
+		npEnsureMsgf(NextPhysicsStep == 0 || (NextPhysicsStep == PhysicsStep), TEXT("Unexpected PhysicsStep %d (NextPhysicsStep: %d)"), PhysicsStep, NextPhysicsStep);
+
+		NextPhysicsStep = PhysicsStep+1;
+
+		// Prep the next AsyncInput
+		npCheckSlow(AsyncCallback);
+		FNetworkPredictionSimCallbackInput* AsyncInput = AsyncCallback->GetProducerInputData_External();
+		npCheckSlow(AsyncInput);
+		
+		AsyncCallback->PushLatestLocalFrameOffset(LocalFrameOffset);
+		
+		// Sync Networked InputCmds
+		NetSerializePlayerControllerInputCmds(PhysicsStep, [&](FNetworkPredictionAsyncID ID, int32 PendingInputBufferFrame, FArchive& Ar)
+		{
+			FGlobalInstanceData& Data = GlobalInstanceData_External.FindChecked(ID);
+			ForEachSim<INetSerializeService>(Data, NetSerializeServices, [&](INetSerializeService* Service)
+			{
+				Service->NetSerializeInputCmd(ID, PendingInputBufferFrame, Ar);
+				if (Ar.IsSaving())
+				{
+					bOutSendClientInputCmd = true;
+				}
+			});
+		});
+
+		// Marshal data from GT to PT input
+		for (TUniquePtr<IAsyncMarshalService_External>& ServicePtr : MarshalServices.Array)
+		{
+			if (IAsyncMarshalService_External* Service = ServicePtr.Get())
+			{
+				Service->MarshalInput(PhysicsStep, AsyncInput);
+			}
+		}
+
+		// Consume all output data. Note we are consuming all SimCallback Output objects here even though we really only want the latest -
+		// but there maybe objects that were destroyed and not present in the final output, so this is trying to not miss them.
+		while (Chaos::TSimCallbackOutputHandle<FNetworkPredictionSimCallbackOutput> AsyncOutput = AsyncCallback->PopFutureOutputData_External())
+		{
+			GlobalOutput_External.SimulationFrame = AsyncOutput->SimulationFrame;
+			GlobalOutput_External.LocalFrameOffset = AsyncOutput->LocalFrameOffset;
+
+			for (TUniquePtr<IAsyncMarshalService_External>& ServicePtr : MarshalServices.Array)
+			{
+				if (IAsyncMarshalService_External* Service = ServicePtr.Get())
+				{
+					Service->MarshalOutput(AsyncOutput.Get());
+				}
+			}
+		}
+	}
+
+	int32 TriggerRewindIfNeeded_Internal(int32 LastCompletedStep) override
+	{
+		return AsyncCallback->Reconcile_Internal(LastCompletedStep);
+	}
+
+	void ProcessInputs_Internal(int32 PhysicsStep) override
+	{
+
+	}
+
+	void PreResimStep_Internal(int32 PhysicsStep, bool bFirst) override
+	{
+	}
+
+	// Returns latest simulation frame number that has been marshaled back from PT
+	int32 GetLatestFrame() const
+	{
+		return GlobalOutput_External.SimulationFrame;
+	}
+
+private:
+
+	struct FGlobalInstanceData
+	{
+		uint32 Sims;	// Bitmask of sims we are subscribed to (1 << ModelDef::GetID())
+		uint8 Flags;	// EInstanceFlags
+		// ENetworkPredictionAsyncInputSource InputSource = ENetworkPredictionAsyncInputSource::None; ???
+	};
+
+	template<typename InterfaceType>
+	void ForEachSim(FGlobalInstanceData& InstanceData, TAsyncServiceStorage<InterfaceType>& Storage, TFunctionRef<void(InterfaceType*)> Func)
+	{
+		uint32 SimMask = InstanceData.Sims;
+		int32 SimID=0;
+		while(SimMask > 0)
+		{
+			if (SimMask & 1)
+			{
+				if (ensure(Storage.Array.IsValidIndex(SimID)))
+				{
+					InterfaceType* Service = Storage.Array[SimID].Get();
+					if (npEnsure(Service))
+					{
+						Func(Service);
+					}
+				}
+			}
+			SimMask >>= 1;
+			SimID++;
+		}
+	}
+
+	template<typename AsyncModelDef>
+	void InstantiateServicesForModelDef()
+	{
+		// InitializedModelDefsMask
+		if ((InitializedModelDefsMask & (1 << AsyncModelDef::ID)) == 0)
+		{
+			InitializedModelDefsMask |= (1 << AsyncModelDef::ID);
+			
+			npCheckSlow(AsyncCallback);
+			AsyncCallback->InstantiateServicesForModelDef<AsyncModelDef>();
+
+			TAsyncModelDataStore_External<AsyncModelDef>* DataStore = GetDataStore<AsyncModelDef>();
+			TAsyncModelDataStore_Internal<AsyncModelDef>* DataStore_Internal = AsyncCallback->GetDataStore<AsyncModelDef>();
+
+			RegistrationServices.Instantiate<AsyncModelDef, TAsyncRegistrationService<AsyncModelDef>>(DataStore);
+			MarshalServices.Instantiate<AsyncModelDef, TAsyncMarshalService_External<AsyncModelDef>>(DataStore, DataStore_Internal);
+			NetSerializeServices.Instantiate<AsyncModelDef, TNetSerializeService<AsyncModelDef>>(DataStore);
+		}
+	}
+
+	template<typename AsyncModelDef>
+	TAsyncModelDataStore_External<AsyncModelDef>* GetDataStore()
+	{
+		return DataStoreCollection.GetDataStore<TAsyncModelDataStore_External<AsyncModelDef>, AsyncModelDef>();
+	}
+
+	template<typename AsyncModelDef>
+	TAsyncModelDataStore_Internal<AsyncModelDef>* GetDataStore_Internal()
+	{
+		npCheckSlow(AsyncCallback);
+		return AsyncCallback->GetDataStore<TAsyncModelDataStore_Internal<AsyncModelDef>, AsyncModelDef>();
+	}
+
+	bool IsServer() const;
+	void NetSerializePlayerControllerInputCmds(int32 PhysicsStep, TFunctionRef<void(FNetworkPredictionAsyncID, int32, FArchive&)> NetSerializeFunc);
+
+	UPROPERTY()
+	UWorld* World = nullptr;
+	FNetworkPredictionAsyncSimCallback* AsyncCallback = nullptr;
+
+	// -------------------------------------------------------------------------------------
+	//	Global data, shared by all the AsyncModelDefs
+	// -------------------------------------------------------------------------------------
+
+	int32 SpawnIDCounter = 1;
+	int32 ClientSpawnIDCounter = -1;
+	int32 NextPhysicsStep = 0;
+
+	uint32 InitializedModelDefsMask = 0;
+
+	TSortedMap<FNetworkPredictionAsyncID, FGlobalInstanceData> GlobalInstanceData_External;
+	struct FGlobalOutput
+	{
+		int32 SimulationFrame = 0;
+		int32 LocalFrameOffset = 0;
+	};
+
+	FGlobalOutput GlobalOutput_External;
+
+	struct FControllerInfo
+	{
+		FNetworkPredictionAsyncID ID;
+		APlayerController* PC = nullptr;
+		ENetworkPredictionAsyncInputSource InputSource = ENetworkPredictionAsyncInputSource::None;
+		int32 LastFrame = 0; // Last frame server processed (de-netserialize into datastore PendingInputCmdBuffer)
+
+		bool operator==(const FNetworkPredictionAsyncID& id) const { return ID == id; }
+	};
+
+	TArray<FControllerInfo> PlayerControllerMap;
+
+	// -------------------------------------------------------------------------------------
+	//	AsyncModelDef specific data and services
+	// -------------------------------------------------------------------------------------
+	FAsyncDataStoreCollection DataStoreCollection;
+
+	TAsyncServiceStorage<IAsyncRegistrationService> RegistrationServices;
+	TAsyncServiceStorage<IAsyncMarshalService_External> MarshalServices;
+	TAsyncServiceStorage<INetSerializeService> NetSerializeServices;
+};
+
+}; // namespace UE_NP
\ No newline at end of file
diff --git a/Engine/Plugins/Runtime/NetworkPrediction/Source/NetworkPrediction/Public/Async/Services/NetworkPredictionAsyncServiceRegistry.h b/Engine/Plugins/Runtime/NetworkPrediction/Source/NetworkPrediction/Public/Async/Services/NetworkPredictionAsyncServiceRegistry.h
new file mode 100644
--- /dev/null
+++ b/Engine/Plugins/Runtime/NetworkPrediction/Source/NetworkPrediction/Public/Async/Services/NetworkPredictionAsyncServiceRegistry.h
@@ -0,0 +1,36 @@
+// Copyright Epic Games, Inc. All Rights Reserved.
+
+#pragma once
+
+#include "Async/NetworkPredictionAsyncInstanceData.h"
+#include "Templates/UnrealTemplate.h"
+
+namespace UE_NP {
+
+template<typename InAsyncServiceInterface, int32 NumInlineServices=4>
+struct TAsyncServiceStorage
+{
+	using AsyncServiceInterface = InAsyncServiceInterface;
+	TArray<TUniquePtr<AsyncServiceInterface>, TInlineAllocator<NumInlineServices>> Array;
+
+	template<typename AsyncModelDef, typename ServiceType, typename... ArgsType>
+	void Instantiate(ArgsType&&... Args)
+	{
+		npCheck(AsyncModelDef::ID >= 0);
+
+		// Resize array for this ModelDef if necessary
+		if (Array.IsValidIndex(AsyncModelDef::ID) == false)
+		{
+			Array.SetNum(AsyncModelDef::ID+1);
+		}
+
+		// Allocate instance on the UniquePtr if necessary
+		auto& Ptr = Array[AsyncModelDef::ID];
+		if (Ptr.IsValid() == false)
+		{
+			Ptr = MakeUnique<ServiceType>(Forward<ArgsType>(Args)...);
+		}
+	}
+};
+
+} // namespace UE_NP
\ No newline at end of file
diff --git a/Engine/Plugins/Runtime/NetworkPrediction/Source/NetworkPrediction/Public/Async/Services/NetworkPredictionAsyncService_Marshal.inl b/Engine/Plugins/Runtime/NetworkPrediction/Source/NetworkPrediction/Public/Async/Services/NetworkPredictionAsyncService_Marshal.inl
new file mode 100644
--- /dev/null
+++ b/Engine/Plugins/Runtime/NetworkPrediction/Source/NetworkPrediction/Public/Async/Services/NetworkPredictionAsyncService_Marshal.inl
@@ -0,0 +1,310 @@
+// Copyright Epic Games, Inc. All Rights Reserved.
+
+#pragma once
+
+/*=============================================================================
+
+Note on ordering of operations and frame numbers:
+
+[GT]
+	Write data to next FNetworkPredictionSimCallbackInput as necessary
+	Game Thread can know which Simulation Frame number the input will be processed on (FNetworkPredictionAsyncWorldManager::NextPhysicsStep)
+
+[PT]
+	CurrentPhysicsStep = N
+
+	[Marshal Service]
+		1. Apply Data from GT:
+			a. New Instances
+			b. State Mods
+			c. InputCmds (locally controlled clients + remotely controlled on server... NOT remote controlled on clients this is done in Reconcile)
+
+		<<< Data for frame SimulationFrame is now "finalized" >>>
+			Think of it like "this is the data given to the [AsyncModelDef's] Tick function and Physics Simulation.
+
+		2. Network Synchronization
+			a. Clients apply the "finalized" data from the server here
+		3. Game Thread Synchronization
+			b. Clients and Server marshal data back to GT here
+
+		4. Copy "finalized" data forward to next frame (LocalFrame + 1)
+
+	[Tick Service]
+		
+		5. Run the user's Tick function (LocalFrame+1)			
+
+	[Physics Engine]
+	
+		6. Step physics Simulation (rigid bodies, etc. This has nothing to do with NP gameplay data)
+
+
+Its important to preserve the ordering above. Particular the "We apply all GT data and THEN "finalize" back to GT".
+	-This allows us to "cover" all data modifications from GT->PT->GT
+	-E.g, when we register a new instance on the GT, we marshal its creation to the PT, and marshal it back when the PT consumes its (rather than waiting for the next PT tick)
+		-But we can also during creation, immediately write it to the GT output. (RegisterInstance_Impl in NetworkPredictionAsyncService_Registration.inl)
+	-In other words, the GT can accurately "predict" its mods/new instances on the latest.
+		-(note this prediction isn't perfect because we are writing to the previously marshalled frame on the GT, but it prevents "gaps" where the data/new instance isn't present)
+
+That is the current thinking at least. The only alternative approach I see is to marshal back the previous frame's data without any of the latest GT data applied. This was the initial 
+attempt but the "gap" mentioned above proved to be pretty annoying and something we had to guard against. 
+	
+
+=============================================================================*/
+
+namespace UE_NP {
+
+class IAsyncMarshalService_Internal
+{
+public:
+
+	virtual ~IAsyncMarshalService_Internal() = default;
+	virtual void MarshalPreSimulate(const int32 LocalFrame, const bool bIsResim, const FNetworkPredictionSimCallbackInput* SimInput, FNetworkPredictionSimCallbackOutput& Output) = 0;
+};
+
+template<typename AsyncModelDef>
+class TAsyncMarshalService_Internal : public IAsyncMarshalService_Internal
+{
+	HOIST_ASYNCMODELDEF_TYPES()
+
+public:
+
+	TAsyncMarshalService_Internal(TAsyncModelDataStore_Internal<AsyncModelDef>* InDataStore)
+	{
+		npCheck(InDataStore != nullptr);
+		DataStore = InDataStore;
+	}
+
+	// Marshal data from SimInput and to SimOutput, prior to ticking
+	void MarshalPreSimulate(const int32 LocalFrame, const bool bIsResim, const FNetworkPredictionSimCallbackInput* SimInput, FNetworkPredictionSimCallbackOutput& SimOutput) final override
+	{
+		const TAsyncModelDataStore_Input<AsyncModelDef>* InputData = SimInput->GetDataStore<AsyncModelDef>();
+		TAsncFrameSnapshot<AsyncModelDef>& Snapshot = DataStore->Frames[LocalFrame];
+
+		/*
+		// Copy previous frame's contents
+		if (LocalFrame > 0)
+		{
+			Snapshot.NetStates = DataStore->Frames[LocalFrame-1].NetStates;
+			if (!bIsResim)
+			{
+				// Don't copy inputs forward during a resim
+				// This is because non locally controlled sim inputs will be set in the Reconcile phase, 
+				// and copying previous InputCmds would over-write that.
+				Snapshot.InputCmds = DataStore->Frames[LocalFrame-1].InputCmds;
+			}
+		}
+		*/
+
+		// New instances
+		for (const typename TAsyncModelDataStore_Input<AsyncModelDef>::FNewInstance& NewInstance : InputData->NewInstances)
+		{
+			DataStore->Instances.FindOrAdd(NewInstance.ID) = NewInstance.StaticData;
+
+			NpResizeForIndex(Snapshot.InputCmds, NewInstance.StaticData.Index);
+			NpResizeForIndex(Snapshot.NetStates, NewInstance.StaticData.Index);
+
+			Snapshot.NetStates[NewInstance.StaticData.Index] = NewInstance.NetState;
+		}
+
+		// State Mods
+		for (const TAsyncLocalStateMod<AsyncModelDef>& LocalMod : InputData->LocalStateMods)
+		{
+			LocalMod.Func(DataStore->Instances.FindChecked(LocalMod.ID).LocalState);
+		}
+
+		for (const TAsyncNetStateMod<AsyncModelDef>& NetMod : InputData->NetStateMods)
+		{
+			NetMod.Func(Snapshot.NetStates[NetMod.Index]);
+		}
+
+		// New Local InputCmds
+		for (const typename TAsyncModelDataStore_Input<AsyncModelDef>::FLocalInputCmd& LocalInput : InputData->LocalInputCmds)
+		{
+			Snapshot.InputCmds[LocalInput.Index] = LocalInput.InputCmd;
+		}
+
+		// Deletes
+		for (const FNetworkPredictionAsyncID& ID : InputData->DeletedInstances)
+		{
+			// I'm not sure if this is enough. Maybe we should track the 'kill frame' on the PT rather than removing anything from data structures
+			DataStore->Instances.Remove(ID);
+		}
+
+		// Apply Corrections
+		ApplyCorrections(LocalFrame);
+
+		// Copy Snapshot back for GT
+		TAsyncModelDataStore_Output<AsyncModelDef>* OutputData = SimOutput.GetDataStore<AsyncModelDef>();
+		OutputData->Snapshot = Snapshot;
+
+		// Copy snapshot state forward to next frame
+		TAsncFrameSnapshot<AsyncModelDef>& NextSnapshot = DataStore->Frames[LocalFrame+1];
+		NextSnapshot.NetStates = Snapshot.NetStates;
+		if (!bIsResim)
+		{
+			// Don't copy inputs forward during a resim
+			// This is because non locally controlled sim inputs will be set in the Reconcile phase, 
+			// and copying previous InputCmds would over-write that.
+			NextSnapshot.InputCmds = Snapshot.InputCmds;
+		}
+		else
+		{
+			// But even in resim we have to copy the AP input cmds forward due to how we are synchronizing with the GT and correction above.
+			// Using the Input's LocalInputCmds is a bit abusive but is the simplest approach since the data is right there. See notes at 
+			// the top of this header
+			for (const typename TAsyncModelDataStore_Input<AsyncModelDef>::FLocalInputCmd& LocalInput : InputData->LocalInputCmds)
+			{
+				const int32 idx = LocalInput.Index;
+				if (Snapshot.InputCmds.IsValidIndex(idx) && NextSnapshot.InputCmds.IsValidIndex(idx))
+				{
+					NextSnapshot.InputCmds[idx] = Snapshot.InputCmds[idx];
+				}
+			}
+		}
+	}
+
+	void ApplyCorrections(int32 LocalFrame)
+	{
+		for (TConstSetBitIterator<> BitIt(DataStore->NetRecvData.PendingCorrectionMask); BitIt; ++BitIt)
+		{
+			const int32 idx = BitIt.GetIndex();
+			const typename TAsyncNetRecvData<AsyncModelDef>::FInstance& RecvData = DataStore->NetRecvData.NetRecvInstances[idx];
+			if (RecvData.Frame == LocalFrame)
+			{
+				TAsncFrameSnapshot<AsyncModelDef>& Snapshot = DataStore->Frames[LocalFrame];
+
+				NpResizeForIndex(Snapshot.InputCmds, idx);
+				NpResizeForIndex(Snapshot.NetStates, idx);
+				
+				Snapshot.InputCmds[idx] = RecvData.InputCmd;
+				Snapshot.NetStates[idx] = RecvData.NetState;
+
+				DataStore->NetRecvData.PendingCorrectionMask[idx] = false;
+			}
+		}
+	}
+	
+
+protected:
+
+	TAsyncModelDataStore_Internal<AsyncModelDef>* DataStore = nullptr;
+};
+
+// ------------------------------------------------------------------------------------------
+
+class IAsyncMarshalService_External
+{
+public:
+
+	virtual ~IAsyncMarshalService_External() = default;
+	virtual void MarshalInput(const int32 Frame, FNetworkPredictionSimCallbackInput* SimInput) = 0;
+	virtual void MarshalOutput(FNetworkPredictionSimCallbackOutput* SimOutput) = 0;
+};
+
+
+template<typename AsyncModelDef>
+class TAsyncMarshalService_External : public IAsyncMarshalService_External
+{
+	HOIST_ASYNCMODELDEF_TYPES()
+
+public:
+
+	TAsyncMarshalService_External(TAsyncModelDataStore_External<AsyncModelDef>* InDataStore_External, TAsyncModelDataStore_Internal<AsyncModelDef>* InDataStore_Internal)
+	{
+		npCheck(InDataStore_External != nullptr);
+		npCheck(InDataStore_Internal != nullptr);
+		DataStore_External = InDataStore_External;
+		DataStore_Internal = InDataStore_Internal;
+	}
+
+	// Marshal external data to SimInput
+	void MarshalInput(const int32 Frame, FNetworkPredictionSimCallbackInput* SimInput) final override
+	{
+		TAsyncModelDataStore_Input<AsyncModelDef>* InputData = SimInput->GetDataStore<AsyncModelDef>();
+
+		// ---------------------------------------------------------------------
+		//	InputCmds
+		// ---------------------------------------------------------------------
+
+		for (TAsyncPendingInputCmdPtr<AsyncModelDef>& Info : DataStore_External->ActivePendingInputCmds)
+		{
+			npCheckSlow(Info.PendingInputCmd != nullptr);
+			InputData->LocalInputCmds.Emplace(Info.Index, *Info.PendingInputCmd);
+		}
+
+		for (auto& MapIt : DataStore_External->PendingInputCmdBuffers)
+		{
+			TAsyncPendingInputCmdBuffer<AsyncModelDef>& PendingInputs = MapIt.Value.Get();
+			InputData->LocalInputCmds.Emplace(PendingInputs.Index, PendingInputs.Buffer[Frame % PendingInputs.Buffer.Num()]);
+		}
+
+		// ---------------------------------------------------------------------
+		//	Network
+		// ---------------------------------------------------------------------
+
+		if (DataStore_External->PendingNetRecv.NetRecvDirtyMask.Contains(true))
+		{
+			DataStore_Internal->NetRecvQueue.Enqueue(MoveTemp(DataStore_External->PendingNetRecv));
+			DataStore_External->PendingNetRecv.Reset();
+		}
+	}
+
+	// Marshal output data to external data store
+	void MarshalOutput(FNetworkPredictionSimCallbackOutput* SimOutput) final override
+	{
+		TAsyncModelDataStore_Output<AsyncModelDef>* OutputData = SimOutput->GetDataStore<AsyncModelDef>();
+		DataStore_External->LatestSnapshot = MoveTemp(OutputData->Snapshot);
+
+		const int32 Num = FMath::Min<int32>(DataStore_External->LatestSnapshot.NetStates.Num(), DataStore_External->OutputWriteTargets.Num());
+		for (int32 idx=0; idx < Num; ++idx)
+		{
+			if (NetStateType* OutNetState = DataStore_External->OutputWriteTargets[idx].OutNetState)
+			{
+				*OutNetState = DataStore_External->LatestSnapshot.NetStates[idx];
+			}
+		}
+	}
+	
+	void ModifyLocalState(FNetworkPredictionAsyncID ID, TUniqueFunction<void(LocalStateType&)> Func, FNetworkPredictionSimCallbackInput* AsyncInput)
+	{
+		npCheckSlow(AsyncInput);
+		if (!npEnsure(ID.IsValid()))
+		{
+			// We can't accept this state mod without a valid ID
+			return;
+		}
+
+		if (ID.IsClientGenerated())
+		{
+			// We can't just drop this mod like we can for NetState since the local state is non replicable. We need to save it to a deferred data data structure
+			DataStore_External->DeferredLocalStateMods.Emplace(ID, MoveTemp(Func));
+		}
+		else
+		{
+			AsyncInput->GetDataStore<AsyncModelDef>()->LocalStateMods.Emplace(ID, MoveTemp(Func));
+		}
+	}
+	
+	void ModifyNetState(FNetworkPredictionAsyncID ID, TUniqueFunction<void(NetStateType&)> Func, FNetworkPredictionSimCallbackInput* AsyncInput)
+	{
+		if (ID.IsClientGenerated() || !ID.IsValid())
+		{
+			// We don't have a valid server created ID yet, so we will just drop this on the floor.
+			// It doesn't matter because GT mods are already "just guesses" wrt synchronization/networking 
+			// and we will have authoritative server state repped to us when we get the ID anyways
+			return;
+		}
+
+		npCheckSlow(AsyncInput);
+		const int32 idx = DataStore_External->Instances.FindChecked(ID).Index;
+		AsyncInput->GetDataStore<AsyncModelDef>()->NetStateMods.Emplace(idx, MoveTemp(Func));
+	}
+
+private:
+
+	TAsyncModelDataStore_External<AsyncModelDef>* DataStore_External = nullptr;
+	TAsyncModelDataStore_Internal<AsyncModelDef>* DataStore_Internal = nullptr;
+};
+
+
+} // namespace UE_NP
\ No newline at end of file
diff --git a/Engine/Plugins/Runtime/NetworkPrediction/Source/NetworkPrediction/Public/Async/Services/NetworkPredictionAsyncService_NetSerialize.inl b/Engine/Plugins/Runtime/NetworkPrediction/Source/NetworkPrediction/Public/Async/Services/NetworkPredictionAsyncService_NetSerialize.inl
new file mode 100644
--- /dev/null
+++ b/Engine/Plugins/Runtime/NetworkPrediction/Source/NetworkPrediction/Public/Async/Services/NetworkPredictionAsyncService_NetSerialize.inl
@@ -0,0 +1,107 @@
+// Copyright Epic Games, Inc. All Rights Reserved.
+
+#pragma once
+#include "UObject/CoreNet.h"
+
+namespace UE_NP {
+
+class INetSerializeService
+{
+public:
+	virtual ~INetSerializeService() = default;
+
+	// Receiving Networked data. Note that frame here = "The Frame this data was sent on"
+	virtual void NetRecv(FNetworkPredictionAsyncID ID, uint8 Flags, int32 ThisFrame, int32 FutureFrame, FNetBitReader& Ar) = 0;
+	virtual void NetSend(FNetworkPredictionAsyncID ID, uint8 Flags, int32 ThisFrame, int32 PendingInputCmdFrame, FNetBitWriter& Ar) = 0;
+	virtual void NetSerializeInputCmd(FNetworkPredictionAsyncID ID, int32 Frame, FArchive& Ar) = 0;
+};
+
+template<typename AsyncModelDef>
+class TNetSerializeService : public INetSerializeService
+{
+	HOIST_ASYNCMODELDEF_TYPES()
+
+public:
+
+	TNetSerializeService(TAsyncModelDataStore_External<AsyncModelDef>* InDataStore)
+	{
+		npCheck(InDataStore != nullptr);
+		DataStore = InDataStore;
+	}
+
+	void NetRecv(FNetworkPredictionAsyncID ID, uint8 Flags, int32 ThisFrame, int32 FutureFrame, FNetBitReader& Ar) final override
+	{
+		TAsncInstanceStaticData<AsyncModelDef>& InstanceData = DataStore->Instances.FindChecked(ID);
+
+		DataStore->PendingNetRecv.MarkIndexDirty(InstanceData.Index);
+
+		typename TAsyncNetRecvData<AsyncModelDef>::FInstance& NetRecvData = DataStore->PendingNetRecv.NetRecvInstances[InstanceData.Index];
+		NetRecvData.Frame = ThisFrame;
+		NetRecvData.Flags = Flags;
+		NetRecvData.InputCmd.NetSerialize(Ar);
+		NetRecvData.NetState.NetSerialize(Ar);
+		NetRecvData.LatestInputCmd.NetSerialize(Ar);
+	}
+
+	void NetSend(FNetworkPredictionAsyncID ID, uint8 Flags, int32 ThisFrame, int32 PendingInputCmdFrame, FNetBitWriter& Ar) final override
+	{
+		TAsncInstanceStaticData<AsyncModelDef>& InstanceData = DataStore->Instances.FindChecked(ID);
+
+		const int32 idx = InstanceData.Index;
+		npCheckSlow(DataStore->LatestSnapshot.InputCmds.IsValidIndex(idx));
+		npCheckSlow(DataStore->LatestSnapshot.NetStates.IsValidIndex(idx));
+
+		DataStore->LatestSnapshot.InputCmds[idx].NetSerialize(Ar);
+		DataStore->LatestSnapshot.NetStates[idx].NetSerialize(Ar);
+
+		// Serialize "Future Input" / LatestInputCmd
+		NetSerializeInputCmd_Impl(ID, PendingInputCmdFrame, Ar, idx);
+	}
+
+	void NetSerializeInputCmd(FNetworkPredictionAsyncID ID, int32 PendingInputCmdFrame, FArchive& Ar) final override
+	{
+		NetSerializeInputCmd_Impl(ID, PendingInputCmdFrame, Ar, INDEX_NONE);
+	}
+
+private:
+
+	void NetSerializeInputCmd_Impl(FNetworkPredictionAsyncID ID, int32 PendingInputCmdFrame, FArchive& Ar, int32 idx)
+	{
+		//	This input is either in a PendingInputCmdBuffers or we just need to pull it from the PendingInputCmd*
+		if (PendingInputCmdFrame != INDEX_NONE)
+		{
+			if (auto* FoundPtr = DataStore->PendingInputCmdBuffers.Find(ID))
+			{
+				// Need latest 
+				TAsyncPendingInputCmdBuffer<AsyncModelDef>& PendingInputs = FoundPtr->Get();
+				PendingInputs.Buffer[PendingInputCmdFrame % PendingInputs.Buffer.Num()].NetSerialize(Ar);
+			}
+			else
+			{
+				UE_LOG(LogNetworkPrediction, Warning, TEXT("%s Could not find PendingInputCmdBuffer for Instance %d"), AsyncModelDef::GetName(), (int32)ID);
+				InputCmdType Default;
+				Default.NetSerialize(Ar);
+			}
+		}
+		else
+		{
+			npEnsure(Ar.IsSaving()); // we should only be here in the writing phase. Saved InputCmds should always end up in the PendingInputCmdBuffer
+			if (TAsyncPendingInputCmdPtr<AsyncModelDef>* Ptr = DataStore->ActivePendingInputCmds.FindByKey(idx != INDEX_NONE ? idx : DataStore->Instances.FindChecked(ID).Index))
+			{
+				npCheckSlow(Ptr->PendingInputCmd);
+				Ptr->PendingInputCmd->NetSerialize(Ar);
+			}
+			else
+			{
+				UE_LOG(LogNetworkPrediction, Warning, TEXT("%s Could not find ActivePendingInputCmds for Instance %d"), AsyncModelDef::GetName(), (int32)ID);
+				InputCmdType Default;
+				Default.NetSerialize(Ar);
+			}
+		}
+	}
+
+	TAsyncModelDataStore_External<AsyncModelDef>* DataStore = nullptr;
+};
+
+
+} // namespace UE_NP
\ No newline at end of file
diff --git a/Engine/Plugins/Runtime/NetworkPrediction/Source/NetworkPrediction/Public/Async/Services/NetworkPredictionAsyncService_Reconcile.inl b/Engine/Plugins/Runtime/NetworkPrediction/Source/NetworkPrediction/Public/Async/Services/NetworkPredictionAsyncService_Reconcile.inl
new file mode 100644
--- /dev/null
+++ b/Engine/Plugins/Runtime/NetworkPrediction/Source/NetworkPrediction/Public/Async/Services/NetworkPredictionAsyncService_Reconcile.inl
@@ -0,0 +1,114 @@
+// Copyright Epic Games, Inc. All Rights Reserved.
+
+#pragma once
+#include "NetworkPredictionCVars.h"
+
+namespace UE_NP {
+
+	
+NETSIM_DEVCVAR_SHIPCONST_INT(EnableFutureInputs, 0, "np2.EnableFutureInputs", "");
+
+class IAsyncReconcileService
+{
+public:
+	virtual ~IAsyncReconcileService() = default;
+	virtual int32 Reconcile(const int32 LastCompletedStep, int32 LocalFrameOffset) = 0;
+};
+
+template<typename AsyncModelDef>
+class TAsyncReconcileService : public IAsyncReconcileService
+{
+public:
+
+	HOIST_ASYNCMODELDEF_TYPES()
+
+	TAsyncReconcileService(TAsyncModelDataStore_Internal<AsyncModelDef>* InDataStore)
+	{
+		npCheck(InDataStore != nullptr);
+		DataStore = InDataStore;
+	}
+
+	int32 Reconcile(const int32 LastCompletedStep, int32 LocalFrameOffset) final override
+	{
+		if (!DataStore->NetRecvQueue.Dequeue(DataStore->NetRecvData))
+		{
+			return INDEX_NONE;
+		}
+		npEnsure(DataStore->NetRecvQueue.IsEmpty()); // We should enqueue exactly one NetRecv snapshot per marshalled input, which should be exactly one call to Reconcile
+		DataStore->NetRecvData.PendingCorrectionMask.Reset();
+
+		int32 RewindFrame = INDEX_NONE;
+		for (TConstSetBitIterator<> BitIt(DataStore->NetRecvData.NetRecvDirtyMask); BitIt; ++BitIt)
+		{
+			const int32 idx = BitIt.GetIndex();
+			const typename TAsyncNetRecvData<AsyncModelDef>::FInstance& RecvData = DataStore->NetRecvData.NetRecvInstances[idx];
+			RecvData.Frame += LocalFrameOffset; // Convert server-frame to local frame
+
+			if (RecvData.Frame >= 0)
+			{
+				const TAsncFrameSnapshot<AsyncModelDef>& Snapshot = DataStore->Frames[RecvData.Frame];
+
+				// This is probably not enough, we will need to check some alive flag to make sure this state was active at 
+				// this time or something?
+				bool bShouldReconcile = false;
+				if (Snapshot.InputCmds.IsValidIndex(idx) && Snapshot.NetStates.IsValidIndex(idx))
+				{
+					const InputCmdType& LocalInputCmd = Snapshot.InputCmds[idx];
+					const NetStateType& LocalNetState = Snapshot.NetStates[idx];
+
+					if (LocalInputCmd.ShouldReconcile(RecvData.InputCmd))
+					{
+						UE_LOG(LogNetworkPrediction, Log, TEXT("Reconcile on InputCmd. %s. Instance: %d. Frame: [%d/%d]"), AsyncModelDef::GetName(), idx, RecvData.Frame - LocalFrameOffset, RecvData.Frame);
+						bShouldReconcile = true;
+					}
+					if (LocalNetState.ShouldReconcile(RecvData.NetState))
+					{
+						UE_LOG(LogNetworkPrediction, Log, TEXT("Reconcile on NetState. %s. Instance: %d. Frame: [%d/%d]"), AsyncModelDef::GetName(), idx, RecvData.Frame - LocalFrameOffset, RecvData.Frame);
+						bShouldReconcile = true;
+					}
+				}
+				else
+				{
+					UE_LOG(LogNetworkPrediction, Log, TEXT("Reconcile due to invalid state. %s. Instance: %d. Frame: [%d/%d]"), AsyncModelDef::GetName(), idx, RecvData.Frame - LocalFrameOffset, RecvData.Frame);
+					bShouldReconcile = true;
+				}
+
+				if (bShouldReconcile)
+				{
+					NpResizeAndSetBit(DataStore->NetRecvData.PendingCorrectionMask, idx);
+					RewindFrame = RewindFrame == INDEX_NONE ? RecvData.Frame : FMath::Min<int32>(RewindFrame, RecvData.Frame);
+				}
+			}
+
+			// -------------------------------------------------------
+			// Apply all Input corrections for future frames if non-locally controlled.
+			//	doing this in a second pass maybe more cache coherent, but this is simpler for now.
+			// -------------------------------------------------------
+			
+			if ((RecvData.Flags & (uint8)EInstanceFlags::LocallyControlled) == 0)
+			{
+				const InputCmdType& SourceInput = EnableFutureInputs() ? RecvData.LatestInputCmd : RecvData.InputCmd;
+
+				const int32 NumFrames = (LastCompletedStep+1) - FMath::Max<int32>(0, RecvData.Frame+1);
+				npEnsureMsgf(NumFrames < DataStore->Frames.Num(), TEXT("Too many input corrections. [%d %d] %d %d"), (LastCompletedStep+1), FMath::Max<int32>(0, RecvData.Frame+1), NumFrames, DataStore->Frames.Num() );
+
+				for (int32 Frame= FMath::Max<int32>(0, RecvData.Frame+1); Frame <= LastCompletedStep+1; ++Frame)
+				{
+					TAsncFrameSnapshot<AsyncModelDef>& FutureSnapshot = DataStore->Frames[Frame];
+
+					NpResizeForIndex(FutureSnapshot.InputCmds, idx);
+					FutureSnapshot.InputCmds[idx] = SourceInput;
+				}
+			}			
+		}
+
+		DataStore->NetRecvData.NetRecvDirtyMask.Reset();
+		return RewindFrame;
+	}
+
+private:
+	TAsyncModelDataStore_Internal<AsyncModelDef>* DataStore = nullptr;
+};
+
+
+} // namespace UE_NP
\ No newline at end of file
diff --git a/Engine/Plugins/Runtime/NetworkPrediction/Source/NetworkPrediction/Public/Async/Services/NetworkPredictionAsyncService_Registration.inl b/Engine/Plugins/Runtime/NetworkPrediction/Source/NetworkPrediction/Public/Async/Services/NetworkPredictionAsyncService_Registration.inl
new file mode 100644
--- /dev/null
+++ b/Engine/Plugins/Runtime/NetworkPrediction/Source/NetworkPrediction/Public/Async/Services/NetworkPredictionAsyncService_Registration.inl
@@ -0,0 +1,227 @@
+// Copyright Epic Games, Inc. All Rights Reserved.
+
+#pragma once
+#include "NetworkPredictionUtil.h"
+#include "Algo/Find.h"
+#include "Algo/BinarySearch.h"
+
+namespace UE_NP {
+
+class IAsyncRegistrationService
+{
+public:
+
+	virtual ~IAsyncRegistrationService() = default;
+	virtual void UnregisterInstance(FNetworkPredictionAsyncID ID, FNetworkPredictionSimCallbackInput* Input) = 0;
+	virtual void RemapClientInstance(FNetworkPredictionAsyncID ClientID, FNetworkPredictionAsyncID ServerID, FNetworkPredictionSimCallbackInput* Input) = 0;
+	
+	virtual void SetInputSource(FNetworkPredictionAsyncID ID, ENetworkPredictionAsyncInputSource Type, bool bIsLocal) = 0;
+};
+
+template<typename AsyncModelDef>
+class TAsyncRegistrationService : public IAsyncRegistrationService
+{
+	HOIST_ASYNCMODELDEF_TYPES()
+
+public:
+
+	TAsyncRegistrationService(TAsyncModelDataStore_External<AsyncModelDef>* InDataStore)
+	{
+		npCheck(InDataStore != nullptr);
+		DataStore = InDataStore;
+	}
+
+	void RegisterInstance(FNetworkPredictionAsyncID ID, FNetworkPredictionSimCallbackInput* Input, LocalStateType&& InitialValueLocal, NetStateType&& InitialValueNet, InputCmdType* PendingInputCmd, NetStateType* OutNetState)
+	{
+		if (ID.IsClientGenerated())
+		{
+			DeferredInitMap.Emplace(ID, FDeferredRegistration{Input, MoveTemp(InitialValueLocal), MoveTemp(InitialValueNet), PendingInputCmd, OutNetState});
+		}
+		else
+		{
+			RegisterInstance_Impl(ID, Input, MoveTemp(InitialValueLocal), MoveTemp(InitialValueNet), PendingInputCmd, OutNetState);
+		}
+	}
+
+	void UnregisterInstance( FNetworkPredictionAsyncID ID, FNetworkPredictionSimCallbackInput* Input) final override
+	{
+		if (ID.IsClientGenerated())
+		{
+			DeferredInitMap.Remove(ID);
+			DeferredInputSource.Remove(ID);
+		}
+		else if(ID.IsValid())
+		{
+			TAsncInstanceStaticData<AsyncModelDef> InstanceData = DataStore->Instances.FindAndRemoveChecked(ID);
+			DataStore->FreeIndices[InstanceData.Index] = true;
+
+			// Remove from active input data structure
+			ClearInputSource(ID, InstanceData.Index);
+
+			// also remove from the inactive struct
+			int32 idx = DataStore->InActivePendingInputCmds.IndexOfByKey(InstanceData.Index);
+			if (idx != INDEX_NONE)
+			{
+				DataStore->InActivePendingInputCmds.RemoveAtSwap(idx, 1, false);
+			}
+
+			TAsyncModelDataStore_Input<AsyncModelDef>* InputDataStore = Input->GetDataStore<AsyncModelDef>();
+			InputDataStore->DeletedInstances.Add(ID);
+		}
+	}
+
+	void SetInputSource(FNetworkPredictionAsyncID ID, ENetworkPredictionAsyncInputSource Type, bool bIsLocal) final override
+	{
+		if (ID.IsClientGenerated())
+		{
+			DeferredInputSource.Emplace(ID, FDeferredInputSource{Type, bIsLocal});
+		}
+		else
+		{
+			SetInputSource_Impl(ID, Type, bIsLocal);
+		}
+	}
+
+	void RemapClientInstance(FNetworkPredictionAsyncID ClientID, FNetworkPredictionAsyncID ServerID, FNetworkPredictionSimCallbackInput* AsyncInput) final override
+	{
+		// Complete deferred registration on ServerID
+		FDeferredRegistration DeferredInfo = DeferredInitMap.FindAndRemoveChecked(ClientID);
+		RegisterInstance_Impl(ServerID, DeferredInfo.Input, MoveTemp(DeferredInfo.InitialValueLocal), MoveTemp(DeferredInfo.InitialValueNet), DeferredInfo.PendingInputCmd, DeferredInfo.OutNetState);
+
+		// Remap deferred InputSource data
+		if (FDeferredInputSource* InputSource = DeferredInputSource.Find(ClientID))
+		{
+			SetInputSource_Impl(ServerID, InputSource->Type, InputSource->bIsLocal);
+		}
+
+		// Remap any deferred local state mods
+		TAsyncModelDataStore_Input<AsyncModelDef>* InputData = AsyncInput->GetDataStore<AsyncModelDef>();
+		npCheckSlow(InputData);
+
+		// note: going in reverse would be better since we could use RemoveAtSwap but preserving order that mods are applied is important and that would reverse them
+		for (int32 idx=0; idx < DataStore->DeferredLocalStateMods.Num(); ++idx)
+		{
+			TAsyncLocalStateMod<AsyncModelDef>& StateMod = DataStore->DeferredLocalStateMods[idx];
+			if (StateMod.ID == ClientID)
+			{
+				InputData->LocalStateMods.Emplace(ServerID, MoveTemp(StateMod.Func));
+				DataStore->DeferredLocalStateMods.RemoveAt(idx, 1, false);
+			}
+		}
+	}
+
+private:
+	
+	void RegisterInstance_Impl(FNetworkPredictionAsyncID ID, FNetworkPredictionSimCallbackInput* Input, LocalStateType&& InitialValueLocal, NetStateType&& InitialValueNet, InputCmdType* PendingInputCmd, NetStateType* OutNetState)
+	{
+		npCheckSlow(ID.IsClientGenerated() == false);
+
+		// Pick a local Index for this instance
+		int32 FreeIndice = DataStore->FreeIndices.Find(false);
+		if (FreeIndice == INDEX_NONE)
+		{
+			FreeIndice = DataStore->FreeIndices.Num();
+			NpResizeBitArray(DataStore->FreeIndices, DataStore->FreeIndices.Num() + 32);
+			npCheckSlow(DataStore->FreeIndices[FreeIndice] == false);
+		}
+
+		DataStore->FreeIndices[FreeIndice] = true;
+
+		// Update external data store
+		TAsncInstanceStaticData<AsyncModelDef>& InstanceData = DataStore->Instances.Add(ID);		
+		InstanceData.Index = FreeIndice;
+		InstanceData.LocalState = MoveTemp(InitialValueLocal);
+
+		if (PendingInputCmd)
+		{
+			DataStore->InActivePendingInputCmds.Emplace(TAsyncPendingInputCmdPtr<AsyncModelDef>{FreeIndice, PendingInputCmd});
+		}
+
+		if (OutNetState)
+		{
+			NpResizeForIndex(DataStore->OutputWriteTargets, FreeIndice);
+			DataStore->OutputWriteTargets[FreeIndice] = {OutNetState};
+			*OutNetState = InitialValueNet;
+		}
+
+		// Cheat a little and put the initial state values into the latest snapshot so we don't have to special case this in other spots
+		NpResizeForIndex(DataStore->LatestSnapshot.NetStates, FreeIndice);
+		DataStore->LatestSnapshot.NetStates[FreeIndice] = InitialValueNet;
+		NpResizeForIndex(DataStore->LatestSnapshot.InputCmds, FreeIndice);
+
+		// Marshal via Input
+		TAsyncModelDataStore_Input<AsyncModelDef>* InputDataStore = Input->GetDataStore<AsyncModelDef>();
+		npCheckSlow(InputDataStore);
+
+		InputDataStore->NewInstances.Add({ID, InstanceData, MoveTemp(InitialValueNet)});
+	}
+
+	void SetInputSource_Impl(FNetworkPredictionAsyncID ID, ENetworkPredictionAsyncInputSource Type, bool bIsLocal)
+	{
+		npCheckSlow(ID.IsClientGenerated() == false);
+
+		const int32 InstanceIndex = DataStore->Instances.FindChecked(ID).Index;
+
+		// Remove previous entries
+		ClearInputSource(ID, InstanceIndex);
+
+		// Add this ID to the one of the lists
+		switch(Type)
+		{
+			case ENetworkPredictionAsyncInputSource::None:
+			{
+				break;
+			}
+
+			case ENetworkPredictionAsyncInputSource::Buffered:
+			{
+				TAsyncPendingInputCmdBuffer<AsyncModelDef>& PendingInputs = DataStore->PendingInputCmdBuffers.Add(ID).Get();
+				PendingInputs.Index = InstanceIndex;
+				break;
+			}
+
+			case ENetworkPredictionAsyncInputSource::Local:
+			{
+				int32 InactiveIdx = DataStore->InActivePendingInputCmds.IndexOfByKey(InstanceIndex);
+				if (ensureMsgf(InactiveIdx != INDEX_NONE, TEXT("Could not find PendingInputCmd info for %d"), (int32)ID))
+				{
+					DataStore->ActivePendingInputCmds.Add(DataStore->InActivePendingInputCmds[InactiveIdx]);
+				}
+				break;
+			}
+		}
+	}
+
+	void ClearInputSource(FNetworkPredictionAsyncID ID, const int32 InstanceIndex)
+	{
+		int32 idx = DataStore->ActivePendingInputCmds.IndexOfByKey(InstanceIndex);
+		if (idx != INDEX_NONE)
+		{
+			DataStore->ActivePendingInputCmds.RemoveAtSwap(idx, 1, false);
+		}
+		DataStore->PendingInputCmdBuffers.Remove(ID);
+	}
+
+	struct FDeferredRegistration
+	{
+		FNetworkPredictionSimCallbackInput* Input;
+		LocalStateType InitialValueLocal;
+		NetStateType InitialValueNet;
+		InputCmdType* PendingInputCmd;
+		NetStateType* OutNetState;
+	};
+
+	TSortedMap<FNetworkPredictionAsyncID, FDeferredRegistration> DeferredInitMap;
+
+	struct FDeferredInputSource
+	{
+		ENetworkPredictionAsyncInputSource Type = ENetworkPredictionAsyncInputSource::None;
+		bool bIsLocal = false;
+	};
+
+	TSortedMap<FNetworkPredictionAsyncID, FDeferredInputSource> DeferredInputSource;
+
+	TAsyncModelDataStore_External<AsyncModelDef>* DataStore = nullptr;
+};
+
+} // namespace UE_NP
\ No newline at end of file
diff --git a/Engine/Plugins/Runtime/NetworkPrediction/Source/NetworkPrediction/Public/Async/Services/NetworkPredictionAsyncService_Ticking.inl b/Engine/Plugins/Runtime/NetworkPrediction/Source/NetworkPrediction/Public/Async/Services/NetworkPredictionAsyncService_Ticking.inl
new file mode 100644
--- /dev/null
+++ b/Engine/Plugins/Runtime/NetworkPrediction/Source/NetworkPrediction/Public/Async/Services/NetworkPredictionAsyncService_Ticking.inl
@@ -0,0 +1,48 @@
+// Copyright Epic Games, Inc. All Rights Reserved.
+
+#pragma once
+
+namespace UE_NP {
+
+class IAsyncTickingService
+{
+public:
+	virtual ~IAsyncTickingService() = default;
+	virtual void Tick_Internal(FNetworkPredictionAsyncTickContext& Context, FNetworkPredictionSimCallbackOutput& Output) = 0;
+};
+
+template<typename AsyncModelDef>
+class TAsyncTickingService : public IAsyncTickingService
+{
+	HOIST_ASYNCMODELDEF_TYPES()
+
+public:
+
+	TAsyncTickingService(TAsyncModelDataStore_Internal<AsyncModelDef>* InDataStore)
+	{
+		npCheck(InDataStore != nullptr);
+		DataStore = InDataStore;
+	}
+
+	void Tick_Internal(FNetworkPredictionAsyncTickContext& Context, FNetworkPredictionSimCallbackOutput& Output) final override
+	{
+		TAsncFrameSnapshot<AsyncModelDef>& Snapshot = DataStore->Frames[Context.LocalStorageFrame];
+
+		for (auto& MapIt : DataStore->Instances)
+		{
+			FNetworkPredictionAsyncID& ID = MapIt.Key;
+			TAsncInstanceStaticData<AsyncModelDef>& InstanceData = MapIt.Value;
+			InputCmdType& InputCmd = Snapshot.InputCmds[InstanceData.Index];
+			NetStateType& NetState = Snapshot.NetStates[InstanceData.Index];
+
+			SimulationTickType::Tick_Internal(Context, ID, InputCmd, InstanceData.LocalState, NetState);
+		}
+	}
+
+private:
+
+	TAsyncModelDataStore_Internal<AsyncModelDef>* DataStore = nullptr;
+};
+
+
+} // namespace UE_NP
\ No newline at end of file
diff --git a/Engine/Plugins/Runtime/NetworkPrediction/Source/NetworkPrediction/Public/NetworkPhysicsComponent.h b/Engine/Plugins/Runtime/NetworkPrediction/Source/NetworkPrediction/Public/NetworkPhysicsComponent.h
new file mode 100644
--- /dev/null
+++ b/Engine/Plugins/Runtime/NetworkPrediction/Source/NetworkPrediction/Public/NetworkPhysicsComponent.h
@@ -0,0 +1,40 @@
+// Copyright Epic Games, Inc. All Rights Reserved.
+
+#pragma once
+
+#include "NetworkPhysics.h"
+#include "NetworkPhysicsComponent.generated.h"
+
+// Generic component that will register an actor's primitive component with the Network Physics system.
+// This is nothing special about this component, it just provides the replicated NetworkPhysicsState and
+// does the boiler plate registration with the Network Physics Manager. Native C++ classes or components
+// may want to just do this themselves or inherit from the component. This is mainly a convenience for
+// dropping in rollback physics support for blueprint actors.
+UCLASS(BlueprintType, meta=(BlueprintSpawnableComponent))
+class NETWORKPREDICTION_API UNetworkPhysicsComponent : public UActorComponent
+{
+	GENERATED_BODY()
+
+public:
+
+	UNetworkPhysicsComponent();
+
+	virtual void InitializeComponent() override;
+	virtual void UninitializeComponent() override;
+
+	UPROPERTY(Replicated, transient)
+	FNetworkPhysicsState NetworkPhysicsState;
+
+	UFUNCTION(BlueprintPure, Category = "Network Physics")
+	int32 GetNetworkPredictionLOD() const { return NetworkPhysicsState.LocalLOD; }
+
+	APlayerController* GetOwnerPC() const;
+
+protected:
+
+	FSingleParticlePhysicsProxy* GetManagedProxy() const { return NetworkPhysicsState.Proxy; }
+
+	// Which component's physics body to manage if there are multiple PrimitiveComponent and not the root component.
+	UPROPERTY(EditDefaultsOnly, Category = "Network Physics")
+	FName ManagedComponentTag=NAME_None;
+};
\ No newline at end of file
diff --git a/Engine/Plugins/Runtime/NetworkPrediction/Source/NetworkPrediction/Public/NetworkPredictionAsync.h b/Engine/Plugins/Runtime/NetworkPrediction/Source/NetworkPrediction/Public/NetworkPredictionAsync.h
deleted file mode 100644
diff --git a/Engine/Plugins/Runtime/NetworkPrediction/Source/NetworkPrediction/Public/NetworkPredictionBuffer.h b/Engine/Plugins/Runtime/NetworkPrediction/Source/NetworkPrediction/Public/NetworkPredictionBuffer.h
--- a/Engine/Plugins/Runtime/NetworkPrediction/Source/NetworkPrediction/Public/NetworkPredictionBuffer.h
+++ b/Engine/Plugins/Runtime/NetworkPrediction/Source/NetworkPrediction/Public/NetworkPredictionBuffer.h
@@ -168,10 +168,11 @@ struct TNetworkPredictionBuffer
 	{
 	}
 
-	TNetworkPredictionBuffer<ElementType, SingleElement>& operator=(const TNetworkPredictionBuffer<ElementType, SingleElement>&& Other)
+	TNetworkPredictionBuffer<ElementType, SingleElement>& operator=(TNetworkPredictionBuffer<ElementType, SingleElement>&& Other)
 	{
 		IndexMask = Other.IndexMask;
 		Elements = MoveTemp(Other.Elements);
+		return *this;
 	}
 	
 	ElementType& operator[](int32 Index)
diff --git a/Engine/Plugins/Runtime/NetworkPrediction/Source/NetworkPrediction/Public/NetworkPredictionCVars.h b/Engine/Plugins/Runtime/NetworkPrediction/Source/NetworkPrediction/Public/NetworkPredictionCVars.h
--- a/Engine/Plugins/Runtime/NetworkPrediction/Source/NetworkPrediction/Public/NetworkPredictionCVars.h
+++ b/Engine/Plugins/Runtime/NetworkPrediction/Source/NetworkPrediction/Public/NetworkPredictionCVars.h
@@ -115,8 +115,9 @@ namespace UE_NETWORK_PHYSICS
 	NETSIM_DEVCVAR_SHIPCONST_FLOAT(MaxAngularVelocity, 30.f, "np2.Mock.MaxAngularVelocity", "Limits how fast character can possibly rotate.")
 	NETSIM_DEVCVAR_SHIPCONST_FLOAT(MovementK, 1.25, "np2.Mock.MovementK", "Movement Coefficient (higher=faster movement)");
 	NETSIM_DEVCVAR_SHIPCONST_FLOAT(RotationK, 1.25, "np2.Mock.RotationK", "Rotation Coefficient (higher=faster movement)");
-	NETSIM_DEVCVAR_SHIPCONST_FLOAT(TurnK, 100000.f, "np2.Mock.TurnK", "Coefficient for automatic turning (higher=quicker turning)");
-	NETSIM_DEVCVAR_SHIPCONST_FLOAT(TurnDampK, 20.f, "np2.Mock.TurnDampK", "Coefficient for damping portion of turn. Higher=more damping but too higher will lead to instability.");
+	NETSIM_DEVCVAR_SHIPCONST_FLOAT(TurnK, 10000000.f, "np2.Mock.TurnK", "Coefficient for automatic turning (higher=quicker turning)");
+	NETSIM_DEVCVAR_SHIPCONST_FLOAT(TurnDampK, 100.f, "np2.Mock.TurnDampK", "Coefficient for damping portion of turn. Higher=more damping but too higher will lead to instability.");
+	NETSIM_DEVCVAR_SHIPCONST_FLOAT(DampYawVelocityK, 2.f, "np2.Mock.DampYawVelocityK", "Coefficient for damping angular velocity in yaw direction only. This is only enabled when auto target yaw is disabled.");
 
 	NETSIM_DEVCVAR_SHIPCONST_FLOAT(JumpForce, 1000000.0f, "np2.Mock.JumpForce", "Per-Frame force to apply while jumping.");	
 	NETSIM_DEVCVAR_SHIPCONST_INT(JumpFrameDuration, 4, "np2.Mock.JumpFrameDuration", "How many frames to apply jump force for");
diff --git a/Engine/Plugins/Runtime/NetworkPrediction/Source/NetworkPrediction/Public/NetworkPredictionPhysics.h b/Engine/Plugins/Runtime/NetworkPrediction/Source/NetworkPrediction/Public/NetworkPredictionPhysics.h
--- a/Engine/Plugins/Runtime/NetworkPrediction/Source/NetworkPrediction/Public/NetworkPredictionPhysics.h
+++ b/Engine/Plugins/Runtime/NetworkPrediction/Source/NetworkPrediction/Public/NetworkPredictionPhysics.h
@@ -161,12 +161,6 @@ private:
 	}
 };
 
-#else
-
-struct NETWORKPREDICTION_API FNetworkPredictionPhysicsState
-{
-};
-
 #endif // WITH_CHAOS
 
 // ------------------------------------------------------------------------------------------------------
diff --git a/Engine/Plugins/Runtime/NetworkPrediction/Source/NetworkPrediction/Public/NetworkPredictionSerialization.h b/Engine/Plugins/Runtime/NetworkPrediction/Source/NetworkPrediction/Public/NetworkPredictionSerialization.h
--- a/Engine/Plugins/Runtime/NetworkPrediction/Source/NetworkPrediction/Public/NetworkPredictionSerialization.h
+++ b/Engine/Plugins/Runtime/NetworkPrediction/Source/NetworkPrediction/Public/NetworkPredictionSerialization.h
@@ -200,6 +200,14 @@ public:
 
 		const int32 EndFrame = FNetworkPredictionSerialization::ReadCompressedFrame(Ar, ServerRecvData.LastRecvFrame); // 1. StartFrame
 		const int32 StartFrame = EndFrame - NumSendPerUpdate;
+
+		// Reset consumed frame if we detect a gap.
+		// Note this could discard unprocessed commands we previously received (but didn't process) but handling this case doesn't seem necessary or practical
+		if (ServerRecvData.LastConsumedFrame+1 < StartFrame)
+		{
+			ServerRecvData.LastConsumedFrame = StartFrame-1;
+			ServerRecvData.LastRecvFrame = StartFrame-1;
+		}
 		
 		int32 ExpectedTimeMS = ServerRecvData.TotalSimTimeMS; // SimTime we expect to process next command at
 		for (int32 Frame=ServerRecvData.LastConsumedFrame+1; Frame >= 0 && Frame <= ServerRecvData.LastRecvFrame; ++Frame)
diff --git a/Engine/Plugins/Runtime/NetworkPrediction/Source/NetworkPrediction/Public/NetworkPredictionUtil.h b/Engine/Plugins/Runtime/NetworkPrediction/Source/NetworkPrediction/Public/NetworkPredictionUtil.h
--- a/Engine/Plugins/Runtime/NetworkPrediction/Source/NetworkPrediction/Public/NetworkPredictionUtil.h
+++ b/Engine/Plugins/Runtime/NetworkPrediction/Source/NetworkPrediction/Public/NetworkPredictionUtil.h
@@ -1,7 +1,11 @@
 // Copyright Epic Games, Inc. All Rights Reserved.
 
 #pragma once
+
+#include "CoreTypes.h"
 #include "NetworkPredictionCheck.h"
+#include "UObject/NameTypes.h"
+#include "Async/NetworkPredictionAsyncDefines.h"
 
 // Sets index to value, resizing bit array if necessary and setting new bits to false
 template<typename BitArrayType>
@@ -36,4 +40,21 @@ template<typename BitArrayType>
 void NpClearBitArray(BitArrayType& BitArray)
 {
 	BitArray.SetRange(0, BitArray.Num(), false);
+}
+
+template<typename ArrayType>
+void NpResizeForIndex(ArrayType& Array, int32 Index)
+{
+	if (Array.IsValidIndex(Index) == false)
+	{
+		Array.SetNum(Index + UE_NP::FrameStorageGrowth);
+	}
+}
+
+class AActor;
+class FSingleParticlePhysicsProxy;
+
+namespace UE_NP
+{
+	NETWORKPREDICTION_API FSingleParticlePhysicsProxy* FindBestPhysicsProxy(AActor* Owning, FName NamedComponent = NAME_None);
 }
\ No newline at end of file
diff --git a/Engine/Plugins/Runtime/NetworkPrediction/Source/NetworkPrediction/Public/NetworkPredictionWorldManager.h b/Engine/Plugins/Runtime/NetworkPrediction/Source/NetworkPrediction/Public/NetworkPredictionWorldManager.h
--- a/Engine/Plugins/Runtime/NetworkPrediction/Source/NetworkPrediction/Public/NetworkPredictionWorldManager.h
+++ b/Engine/Plugins/Runtime/NetworkPrediction/Source/NetworkPrediction/Public/NetworkPredictionWorldManager.h
@@ -474,6 +474,11 @@ void UNetworkPredictionWorldManager::BindServerNetRecv_Independent(FNetworkPredi
 	ServerRecvData.TraceID = ID.GetTraceID();
 	ServerRecvData.InstanceIdx = DataStore->Instances.GetIndex(ID);
 	ServerRecvData.FramesIdx = DataStore->Frames.GetIndex(ID);
+	ServerRecvData.PendingFrame = 0;
+	ServerRecvData.TotalSimTimeMS = 0;
+	ServerRecvData.UnspentTimeMS = 0.f;
+	ServerRecvData.LastConsumedFrame = INDEX_NONE;
+	ServerRecvData.LastRecvFrame = INDEX_NONE;
 
 	RepProxy->NetSerializeFunc = [DataStore, ServerRecvIdx](const FNetSerializeParams& P)
 	{
