Unity compile fix
#rnx

diff --git a/Engine/Plugins/Runtime/NetworkPrediction/Source/NetworkPrediction/Private/NetworkPhysicsComponent.cpp b/Engine/Plugins/Runtime/NetworkPrediction/Source/NetworkPrediction/Private/NetworkPhysicsComponent.cpp
--- a/Engine/Plugins/Runtime/NetworkPrediction/Source/NetworkPrediction/Private/NetworkPhysicsComponent.cpp
+++ b/Engine/Plugins/Runtime/NetworkPrediction/Source/NetworkPrediction/Private/NetworkPhysicsComponent.cpp
@@ -16,6 +16,7 @@
 #include "TimerManager.h"
 #include "NetworkPredictionDebug.h"
 #include "NetworkPredictionCVars.h"
+#include "Components/PrimitiveComponent.h"
 
 UNetworkPhysicsComponent::UNetworkPhysicsComponent()
 {
