Network Physics
-Better handling of trying to rewind to invalid frames (happens at startup as client InputCmds are flushed through the system)
-np2.NeverResim cvar for testing
-fix crash when unregistering bad physics state

diff --git a/Engine/Plugins/Runtime/NetworkPrediction/Source/NetworkPrediction/Private/NetworkPhysics.cpp b/Engine/Plugins/Runtime/NetworkPrediction/Source/NetworkPrediction/Private/NetworkPhysics.cpp
--- a/Engine/Plugins/Runtime/NetworkPrediction/Source/NetworkPrediction/Private/NetworkPhysics.cpp
+++ b/Engine/Plugins/Runtime/NetworkPrediction/Source/NetworkPrediction/Private/NetworkPhysics.cpp
@@ -44,9 +44,10 @@ namespace UE_NETWORK_PHYSICS
 	NP_DEVCVAR_INT(DebugTolerance, 1, "np2.Debug.Tolerance", "If enabled, only draw large corrections in world.");
 	NP_DEVCVAR_FLOAT(DebugDrawTolerance_X, 5.0, "np2.Debug.Tolerance.X",  "Location Debug Drawing Toleration");
 	NP_DEVCVAR_FLOAT(DebugDrawTolerance_R, 2.0, "np2.Debug.Tolerance.R",  "Simple distance based LOD");
-		
-	NP_DEVCVAR_INT(ForceResim, 0, "np2.ForceResim", "Forces near constant resimming");		
-	NP_DEVCVAR_INT(ForceResimWithoutCorrection, 0, "np2.ForceResimWithoutCorrection", "Forces near constant resimming WITHOUT applying server data");
+	
+	NETSIM_DEVCVAR_SHIPCONST_INT(NeverResim, 0, "np2.NeverResim", "Never allow resims to be triggered on client. E.g, ignore the server.");		
+	NETSIM_DEVCVAR_SHIPCONST_INT(ForceResim, 0, "np2.ForceResim", "Forces near constant resimming");		
+	NETSIM_DEVCVAR_SHIPCONST_INT(ForceResimWithoutCorrection, 0, "np2.ForceResimWithoutCorrection", "Forces near constant resimming WITHOUT applying server data");
 	
 	NP_DEVCVAR_INT(FixedLocalFrameOffset, -1, "np2.FixedLocalFrameOffset", "When > 0, use hardcoded frame offset on client from head");	
 	NP_DEVCVAR_INT(FixedLocalFrameOffsetTolerance, 3, "np2.FixedLocalFrameOffsetTolerance", "Tolerance when using np2.FixedLocalFrameOffset");
@@ -204,8 +205,9 @@ struct FNetworkPhysicsRewindCallback : public Chaos::IRewindCallback
 		}
 
 		PendingCorrections.Reset();
+		const int32 EarliestFrame = RewindData->GetEarliestFrame_Internal();
 
-		if (UE_NETWORK_PHYSICS::ForceResimWithoutCorrection > 0)
+		if (UE_NETWORK_PHYSICS::ForceResimWithoutCorrection() > 0)
 		{
 			//const int32 ForceRewindFrame = RewindData->GetEarliestFrame_Internal()+1; // This was causing us to back > 64 frames?
 			const int32 ForceRewindFrame = LastCompletedStep - 9;
@@ -247,9 +249,9 @@ struct FNetworkPhysicsRewindCallback : public Chaos::IRewindCallback
 				FSingleParticlePhysicsProxy* Proxy = Obj.Proxy;
 				check(Proxy);
 
-				if (Obj.Frame < MinFrame)
+				if (Obj.Frame < EarliestFrame || Obj.Frame > LastCompletedStep)
 				{
-					UE_LOG(LogNetworkPhysics, Log, TEXT("Obj too stale to reconcile. Frame: %d. LastCompletedStep: %d"), Obj.Frame, LastCompletedStep);
+					UE_LOG(LogNetworkPhysics, Verbose, TEXT("Obj too stale to reconcile. Frame: %d. EarliestFrame: %d LastCompletedStep: %d"), Obj.Frame, EarliestFrame, LastCompletedStep);
 					continue;
 				}
 
@@ -290,7 +292,7 @@ struct FNetworkPhysicsRewindCallback : public Chaos::IRewindCallback
 					CompareVec(Obj.Physics.Location, P.X(), UE_NETWORK_PHYSICS::X, TEXT("Location")) ||
 					CompareVec(Obj.Physics.LinearVelocity, P.V(), UE_NETWORK_PHYSICS::V, TEXT("Velocity")) ||
 					CompareVec(Obj.Physics.AngularVelocity, P.W(), UE_NETWORK_PHYSICS::W, TEXT("Angular Velocity")) ||
-					CompareQuat(Obj.Physics.Rotation, P.R(), UE_NETWORK_PHYSICS::R, TEXT("Rotation")) || UE_NETWORK_PHYSICS::ForceResim > 0)
+					CompareQuat(Obj.Physics.Rotation, P.R(), UE_NETWORK_PHYSICS::R, TEXT("Rotation")) || UE_NETWORK_PHYSICS::ForceResim() > 0)
 				{
 
 					UE_CLOG(UE_NETWORK_PHYSICS::LogCorrections > 0, LogNetworkPhysics, Log, TEXT("Rewind Needed. SimFrame: %d. Obj.Frame: %d. LastCompletedStep: %d."), SimulationFrame, Obj.Frame, LastCompletedStep);
@@ -344,28 +346,27 @@ struct FNetworkPhysicsRewindCallback : public Chaos::IRewindCallback
 			int32 CallbackFrame = Callback->TriggerRewindIfNeeded_Internal(LastCompletedStep);
 			if (CallbackFrame != INDEX_NONE)
 			{
-				Stats.NumSubsystemCorrections++;
-				RewindToFrame = RewindToFrame == INDEX_NONE ? CallbackFrame : FMath::Min(RewindToFrame, CallbackFrame);
+				if (ensureMsgf(CallbackFrame >= EarliestFrame && CallbackFrame <= LastCompletedStep, TEXT("SimObjectCallback wants to rewind to invalid frame %d. (Valid Range: %d - %d)"), EarliestFrame, LastCompletedStep))
+				{
+					Stats.NumSubsystemCorrections++;
+					RewindToFrame = RewindToFrame == INDEX_NONE ? CallbackFrame : FMath::Min(RewindToFrame, CallbackFrame);
+				}
 			}
 		}
-
 		
 		for (auto& SubSys : NetworkPhysicsManager->SubSystems)
 		{
 			const int32 CallbackFrame = SubSys.Value->TriggerRewindIfNeeded_Internal(LastCompletedStep);
 			if (CallbackFrame != INDEX_NONE)
 			{
-				Stats.NumSubsystemCorrections++;
-				RewindToFrame = RewindToFrame == INDEX_NONE ? CallbackFrame : FMath::Min(RewindToFrame, CallbackFrame);
+				// Maybe EarliestFrame should be passed into TriggerRewindIfNeeded_Internal?
+				if (ensureMsgf(CallbackFrame >= EarliestFrame && CallbackFrame <= LastCompletedStep, TEXT("NP Subsystem %s wants to rewind to invalid frame %d. (Valid Range: %d - %d)"), *SubSys.Key.ToString(), CallbackFrame, EarliestFrame, LastCompletedStep))
+				{
+					Stats.NumSubsystemCorrections++;
+					RewindToFrame = RewindToFrame == INDEX_NONE ? CallbackFrame : FMath::Min(RewindToFrame, CallbackFrame);
+				}
 			}
 		}
-
-		// Calling code will ensure if we return a bad frame
-		if (!(RewindToFrame > INDEX_NONE && RewindToFrame < LastCompletedStep-1))
-		{
-			UE_CLOG(UE_NETWORK_PHYSICS::LogCorrections > 0 && RewindToFrame != INDEX_NONE, LogNetworkPhysics, Warning, TEXT("Trying to rewind to invalid frame %d (LastCompletedStep: %d)"), RewindToFrame, LastCompletedStep);
-			return INDEX_NONE;
-		}
 		
 		return RewindToFrame;
 	}
diff --git a/Engine/Plugins/Runtime/NetworkPrediction/Source/NetworkPrediction/Public/Async/NetworkPredictionAsyncWorldManager.h b/Engine/Plugins/Runtime/NetworkPrediction/Source/NetworkPrediction/Public/Async/NetworkPredictionAsyncWorldManager.h
--- a/Engine/Plugins/Runtime/NetworkPrediction/Source/NetworkPrediction/Public/Async/NetworkPredictionAsyncWorldManager.h
+++ b/Engine/Plugins/Runtime/NetworkPrediction/Source/NetworkPrediction/Public/Async/NetworkPredictionAsyncWorldManager.h
@@ -111,6 +111,14 @@ public:
 
 	int32 Reconcile_Internal(int32 LastCompletedStep)
 	{
+		Chaos::FPhysicsSolver* PhysicsSolver = static_cast<Chaos::FPhysicsSolver*>(this->GetSolver());
+		npCheckSlow(PhysicsSolver);
+
+		Chaos::FRewindData* RewindData = PhysicsSolver->GetRewindData();
+		npCheckSlow(RewindData);
+
+		const int32 EarliestFrame = RewindData->GetEarliestFrame_Internal();
+
 		// pop latest localframe offset
 		while(LocalFrameOffsetQueue.Dequeue(LocalFrameOffset));
 
@@ -119,7 +127,7 @@ public:
 		{
 			if (IAsyncReconcileService* Service = ServicePtr.Get())
 			{
-				const int32 Frame = Service->Reconcile(LastCompletedStep, LocalFrameOffset);
+				const int32 Frame = Service->Reconcile(LastCompletedStep, EarliestFrame, LocalFrameOffset);
 				if (Frame != INDEX_NONE)
 				{
 					RewindFrame = RewindFrame == INDEX_NONE ? Frame : FMath::Min<int32>(RewindFrame, Frame);
diff --git a/Engine/Plugins/Runtime/NetworkPrediction/Source/NetworkPrediction/Public/Async/Services/NetworkPredictionAsyncService_Reconcile.inl b/Engine/Plugins/Runtime/NetworkPrediction/Source/NetworkPrediction/Public/Async/Services/NetworkPredictionAsyncService_Reconcile.inl
--- a/Engine/Plugins/Runtime/NetworkPrediction/Source/NetworkPrediction/Public/Async/Services/NetworkPredictionAsyncService_Reconcile.inl
+++ b/Engine/Plugins/Runtime/NetworkPrediction/Source/NetworkPrediction/Public/Async/Services/NetworkPredictionAsyncService_Reconcile.inl
@@ -12,7 +12,7 @@ class IAsyncReconcileService
 {
 public:
 	virtual ~IAsyncReconcileService() = default;
-	virtual int32 Reconcile(const int32 LastCompletedStep, int32 LocalFrameOffset) = 0;
+	virtual int32 Reconcile(const int32 LastCompletedStep, const int32 EarliestFrame, int32 LocalFrameOffset) = 0;
 };
 
 template<typename AsyncModelDef>
@@ -28,7 +28,7 @@ public:
 		DataStore = InDataStore;
 	}
 
-	int32 Reconcile(const int32 LastCompletedStep, int32 LocalFrameOffset) final override
+	int32 Reconcile(const int32 LastCompletedStep, const int32 EarliestFrame, int32 LocalFrameOffset) final override
 	{
 		if (!DataStore->NetRecvQueue.Dequeue(DataStore->NetRecvData))
 		{
@@ -44,7 +44,7 @@ public:
 			const typename TAsyncNetRecvData<AsyncModelDef>::FInstance& RecvData = DataStore->NetRecvData.NetRecvInstances[idx];
 			RecvData.Frame += LocalFrameOffset; // Convert server-frame to local frame
 
-			if (RecvData.Frame >= 0)
+			if (RecvData.Frame >= EarliestFrame && RecvData.Frame <= LastCompletedStep)
 			{
 				const TAsncFrameSnapshot<AsyncModelDef>& Snapshot = DataStore->Frames[RecvData.Frame];
 
diff --git a/Engine/Plugins/Runtime/NetworkPrediction/Source/NetworkPrediction/Public/NetworkPhysics.h b/Engine/Plugins/Runtime/NetworkPrediction/Source/NetworkPrediction/Public/NetworkPhysics.h
--- a/Engine/Plugins/Runtime/NetworkPrediction/Source/NetworkPrediction/Public/NetworkPhysics.h
+++ b/Engine/Plugins/Runtime/NetworkPrediction/Source/NetworkPrediction/Public/NetworkPhysics.h
@@ -58,7 +58,7 @@ struct FNetworkPhysicsState
 	FSingleParticlePhysicsProxy* Proxy = nullptr;
 
 	// Local handle used by UNetworkPhysicsManager for registering/unregistering
-	int32 LocalManagedHandle;
+	int32 LocalManagedHandle = INDEX_NONE;
 
 	// The actual physics properties
 	FBasePhysicsState Physics;
