Network Physics
-add cpu trace events in possible hotspots
-disable correction logging by default (too spammy)

diff --git a/Engine/Plugins/Runtime/NetworkPrediction/Source/NetworkPrediction/Private/NetworkPhysics.cpp b/Engine/Plugins/Runtime/NetworkPrediction/Source/NetworkPrediction/Private/NetworkPhysics.cpp
--- a/Engine/Plugins/Runtime/NetworkPrediction/Source/NetworkPrediction/Private/NetworkPhysics.cpp
+++ b/Engine/Plugins/Runtime/NetworkPrediction/Source/NetworkPrediction/Private/NetworkPhysics.cpp
@@ -23,6 +23,7 @@
 #include "NetworkPredictionDebug.h"
 #include "GameFramework/PlayerState.h"
 #include "Misc/FileHelper.h"
+#include "ProfilingDebugging/CpuProfilerTrace.h"
 
 DEFINE_LOG_CATEGORY(LogNetworkPhysics);
 
@@ -32,7 +33,7 @@ namespace UE_NETWORK_PHYSICS
 	FAutoConsoleVariableRef CVarEnable(TEXT("np2.bEnable"), Enable, TEXT("Enabled rollback physics. Must be set before starting game"));
 
 	//NP_DEVCVAR_INT(Enable, 0, "np2.bEnable", "Enabled rollback physics. Must be set before starting game");
-	NP_DEVCVAR_INT(LogCorrections, 1, "np2.LogCorrections", "Logs corrections when they happen");		
+	NP_DEVCVAR_INT(LogCorrections, 0, "np2.LogCorrections", "Logs corrections when they happen");		
 	NP_DEVCVAR_INT(LogImpulses, -1, "np2.LogImpulses", "Logs all recorded F/T/LI/AI");
 
 	NP_DEVCVAR_FLOAT(X, 1.0f, "np2.Tolerance.X", "Location Tolerance");
@@ -198,6 +199,8 @@ struct FNetworkPhysicsRewindCallback : public Chaos::IRewindCallback
 			return INDEX_NONE;
 		}
 
+		TRACE_CPUPROFILER_EVENT_SCOPE(NPA_Physics_TriggerRewindIfNeeded_Internal);
+
 		if (ResimEndFrame != INDEX_NONE)
 		{
 			// We are already in a resim and shouldn't trigger a new one
@@ -687,6 +690,8 @@ void UNetworkPhysicsManager::OnWorldPostInit(UWorld* World, const UWorld::Initia
 
 void UNetworkPhysicsManager::PostNetRecv()
 {
+	TRACE_CPUPROFILER_EVENT_SCOPE(NPA_Physics_PostNetRecv);
+
 #if WITH_CHAOS
 	if (RewindCallback == nullptr)
 	{
@@ -896,6 +901,8 @@ void UNetworkPhysicsManager::PostNetRecv()
 
 void UNetworkPhysicsManager::PreNetSend(float DeltaSeconds)
 {
+	TRACE_CPUPROFILER_EVENT_SCOPE(NPA_Physics_PreNetSend);
+
 #if WITH_CHAOS
 	if (RewindCallback == nullptr)
 	{
@@ -999,6 +1006,8 @@ void UNetworkPhysicsManager::PreNetSend(float DeltaSeconds)
 
 void UNetworkPhysicsManager::ProcessInputs_External(int32 PhysicsStep, const TArray<Chaos::FSimCallbackInputAndObject>& SimCallbackInputs)
 {
+	TRACE_CPUPROFILER_EVENT_SCOPE(NPA_Physics_ProcessInputs_External);
+
 	UWorld* World = GetWorld();
 	const ENetMode NetMode = World->GetNetMode();
 
diff --git a/Engine/Plugins/Runtime/NetworkPrediction/Source/NetworkPrediction/Public/Async/NetworkPredictionAsyncWorldManager.h b/Engine/Plugins/Runtime/NetworkPrediction/Source/NetworkPrediction/Public/Async/NetworkPredictionAsyncWorldManager.h
--- a/Engine/Plugins/Runtime/NetworkPrediction/Source/NetworkPrediction/Public/Async/NetworkPredictionAsyncWorldManager.h
+++ b/Engine/Plugins/Runtime/NetworkPrediction/Source/NetworkPrediction/Public/Async/NetworkPredictionAsyncWorldManager.h
@@ -2,6 +2,7 @@
 
 #pragma once
 
+#include "ProfilingDebugging/CpuProfilerTrace.h"
 #include "Physics/Experimental/PhysScene_Chaos.h"
 #include "PBDRigidsSolver.h"
 #include "ChaosSolversModule.h"
@@ -23,6 +24,7 @@
 #include "Services/NetworkPredictionAsyncService_NetSerialize.inl"
 
 
+
 namespace UE_NP {
 
 // -----------------------------------------------------------------------------------------
@@ -288,6 +290,8 @@ public:
 	//	This includes the latest state/input + a "future input" which is the latest, but currently unprocessed, InputCmd we have on the GT
 	void NetSerializeInstance(FNetworkPredictionAsyncID ID, FArchive& Ar)
 	{
+		TRACE_CPUPROFILER_EVENT_SCOPE(NPA_NetSerializeInstance);
+
 		FGlobalInstanceData& InstanceData = GlobalInstanceData_External.FindChecked(ID);
 
 		// 1. Frame (this is the latest frame we have results for, on the game thread)
@@ -402,6 +406,8 @@ public:
 
 	void ProcessInputs_External(int32 PhysicsStep, int32 LocalFrameOffset, bool& bOutSendClientInputCmd) override
 	{
+		TRACE_CPUPROFILER_EVENT_SCOPE(NPA_ProcessInputs_External);
+
 		npEnsureMsgf(NextPhysicsStep == 0 || (NextPhysicsStep == PhysicsStep), TEXT("Unexpected PhysicsStep %d (NextPhysicsStep: %d)"), PhysicsStep, NextPhysicsStep);
 
 		NextPhysicsStep = PhysicsStep+1;
diff --git a/Engine/Plugins/Runtime/NetworkPrediction/Source/NetworkPrediction/Public/Async/Services/NetworkPredictionAsyncService_Marshal.inl b/Engine/Plugins/Runtime/NetworkPrediction/Source/NetworkPrediction/Public/Async/Services/NetworkPredictionAsyncService_Marshal.inl
--- a/Engine/Plugins/Runtime/NetworkPrediction/Source/NetworkPrediction/Public/Async/Services/NetworkPredictionAsyncService_Marshal.inl
+++ b/Engine/Plugins/Runtime/NetworkPrediction/Source/NetworkPrediction/Public/Async/Services/NetworkPredictionAsyncService_Marshal.inl
@@ -1,6 +1,7 @@
 // Copyright Epic Games, Inc. All Rights Reserved.
 
 #pragma once
+#include "ProfilingDebugging/CpuProfilerTrace.h"
 
 /*=============================================================================
 
@@ -77,6 +78,8 @@ public:
 	// Marshal data from SimInput and to SimOutput, prior to ticking
 	void MarshalPreSimulate(const int32 LocalFrame, const bool bIsResim, const FNetworkPredictionSimCallbackInput* SimInput, FNetworkPredictionSimCallbackOutput& SimOutput) final override
 	{
+		TRACE_CPUPROFILER_EVENT_SCOPE(NPA_MarshalPreSimulate);
+
 		const TAsyncModelDataStore_Input<AsyncModelDef>* InputData = SimInput->GetDataStore<AsyncModelDef>();
 		TAsncFrameSnapshot<AsyncModelDef>& Snapshot = DataStore->Frames[LocalFrame];
 
diff --git a/Engine/Plugins/Runtime/NetworkPrediction/Source/NetworkPrediction/Public/Async/Services/NetworkPredictionAsyncService_Reconcile.inl b/Engine/Plugins/Runtime/NetworkPrediction/Source/NetworkPrediction/Public/Async/Services/NetworkPredictionAsyncService_Reconcile.inl
--- a/Engine/Plugins/Runtime/NetworkPrediction/Source/NetworkPrediction/Public/Async/Services/NetworkPredictionAsyncService_Reconcile.inl
+++ b/Engine/Plugins/Runtime/NetworkPrediction/Source/NetworkPrediction/Public/Async/Services/NetworkPredictionAsyncService_Reconcile.inl
@@ -2,6 +2,7 @@
 
 #pragma once
 #include "NetworkPredictionCVars.h"
+#include "ProfilingDebugging/CpuProfilerTrace.h"
 
 namespace UE_NP {
 
@@ -30,6 +31,8 @@ public:
 
 	int32 Reconcile(const int32 LastCompletedStep, const int32 EarliestFrame, int32 LocalFrameOffset) final override
 	{
+		TRACE_CPUPROFILER_EVENT_SCOPE(NPA_Reconcile);
+
 		if (!DataStore->NetRecvQueue.Dequeue(DataStore->NetRecvData))
 		{
 			return INDEX_NONE;
@@ -58,18 +61,18 @@ public:
 
 					if (LocalInputCmd.ShouldReconcile(RecvData.InputCmd))
 					{
-						UE_LOG(LogNetworkPrediction, Log, TEXT("Reconcile on InputCmd. %s. Instance: %d. Frame: [%d/%d]"), AsyncModelDef::GetName(), idx, RecvData.Frame - LocalFrameOffset, RecvData.Frame);
+						//UE_LOG(LogNetworkPrediction, Log, TEXT("Reconcile on InputCmd. %s. Instance: %d. Frame: [%d/%d]"), AsyncModelDef::GetName(), idx, RecvData.Frame - LocalFrameOffset, RecvData.Frame);
 						bShouldReconcile = true;
 					}
 					if (LocalNetState.ShouldReconcile(RecvData.NetState))
 					{
-						UE_LOG(LogNetworkPrediction, Log, TEXT("Reconcile on NetState. %s. Instance: %d. Frame: [%d/%d]"), AsyncModelDef::GetName(), idx, RecvData.Frame - LocalFrameOffset, RecvData.Frame);
+						//UE_LOG(LogNetworkPrediction, Log, TEXT("Reconcile on NetState. %s. Instance: %d. Frame: [%d/%d]"), AsyncModelDef::GetName(), idx, RecvData.Frame - LocalFrameOffset, RecvData.Frame);
 						bShouldReconcile = true;
 					}
 				}
 				else
 				{
-					UE_LOG(LogNetworkPrediction, Log, TEXT("Reconcile due to invalid state. %s. Instance: %d. Frame: [%d/%d]"), AsyncModelDef::GetName(), idx, RecvData.Frame - LocalFrameOffset, RecvData.Frame);
+					//UE_LOG(LogNetworkPrediction, Log, TEXT("Reconcile due to invalid state. %s. Instance: %d. Frame: [%d/%d]"), AsyncModelDef::GetName(), idx, RecvData.Frame - LocalFrameOffset, RecvData.Frame);
 					bShouldReconcile = true;
 				}
 
diff --git a/Engine/Plugins/Runtime/NetworkPrediction/Source/NetworkPrediction/Public/Async/Services/NetworkPredictionAsyncService_Ticking.inl b/Engine/Plugins/Runtime/NetworkPrediction/Source/NetworkPrediction/Public/Async/Services/NetworkPredictionAsyncService_Ticking.inl
--- a/Engine/Plugins/Runtime/NetworkPrediction/Source/NetworkPrediction/Public/Async/Services/NetworkPredictionAsyncService_Ticking.inl
+++ b/Engine/Plugins/Runtime/NetworkPrediction/Source/NetworkPrediction/Public/Async/Services/NetworkPredictionAsyncService_Ticking.inl
@@ -1,6 +1,7 @@
 // Copyright Epic Games, Inc. All Rights Reserved.
 
 #pragma once
+#include "ProfilingDebugging/CpuProfilerTrace.h"
 
 namespace UE_NP {
 
@@ -26,6 +27,8 @@ public:
 
 	void Tick_Internal(FNetworkPredictionAsyncTickContext& Context, FNetworkPredictionSimCallbackOutput& Output) final override
 	{
+		TRACE_CPUPROFILER_EVENT_SCOPE(NPA_TickInternal);
+
 		TAsncFrameSnapshot<AsyncModelDef>& Snapshot = DataStore->Frames[Context.LocalStorageFrame];
 
 		for (auto& MapIt : DataStore->Instances)
