Network Prediction
Remove project specific maps from NPE dev hud

diff --git a/Engine/Plugins/Runtime/NetworkPrediction/Config/BaseNetworkPrediction.ini b/Engine/Plugins/Runtime/NetworkPrediction/Config/BaseNetworkPrediction.ini
--- a/Engine/Plugins/Runtime/NetworkPrediction/Config/BaseNetworkPrediction.ini
+++ b/Engine/Plugins/Runtime/NetworkPrediction/Config/BaseNetworkPrediction.ini
@@ -1,5 +1,3 @@
 [/Script/NetworkPrediction.NetworkPredictionSettingsObject]
 Settings=(PreferredTickingPolicy=Fixed,ReplicatedManagerClassOverride=None,FixedTickFrameRate=60,bForceEngineFixTickForcePhysics=True,SimulatedProxyNetworkLOD=ForwardPredict,FixedTickInterpolationBufferedMS=100,IndependentTickInterpolationBufferedMS=100,IndependentTickInterpolationMaxBufferedMS=250)
 +DevHUDs=(HUDName="Debug",Items=((DisplayName="Debug ON",ExecCommand="np2.debug 1",bAutoBack=False),(DisplayName="Debug OFF",ExecCommand="np2.debug 0",bAutoBack=False),(DisplayName="Debug Draw Tolerance ON",ExecCommand="np2.Debug.Tolerance 1",bAutoBack=False),(DisplayName="Debug Draw Tolerance OFF",ExecCommand="np2.debug.Tolerance 0",bAutoBack=False),(DisplayName="Dump",ExecCommand="np2.dump",bAutoBack=False,bRequirePIE=True)),bRequirePIE=False,bRequireNotPIE=False)
-+DevHUDs=(HUDName="Input Macros",Items=((DisplayName="Start/Stop Record",ExecCommand="np2.RecordInputs Temp",bAutoBack=False),(DisplayName="Playback",ExecCommand="np2.PlaybackInputs Temp",bAutoBack=False)),bRequirePIE=False,bRequireNotPIE=False)
-+DevHUDs=(HUDName="Change Map",Items=((DisplayName="Empty",ExecCommand="serverexec servertravel Physics2_Empty"),(DisplayName="Physics Tracks",ExecCommand="serverexec servertravel TestMap_Physics_Tracks"),(DisplayName="Physics Tracks 2",ExecCommand="serverexec servertravel TestMap_Physics_Tracks2"),(DisplayName="Soccer",ExecCommand="serverexec servertravel Soccer"),(DisplayName="Soccer Empty",ExecCommand="serverexec servertravel Physics2_soccerball_empty"),(DisplayName="Portals",ExecCommand="serverexec servertravel Portals"),(DisplayName="Cube DM",ExecCommand="serverexec servertravel TestMap_Cube_DM"),(DisplayName="Swings",ExecCommand="serverexec servertravel TestMap_Swings"),(DisplayName="HoverShips",ExecCommand="serverexec servertravel Physics2_HoverShips")),bRequirePIE=False,bRequireNotPIE=true)
