Network Physics
-multi body support. Must be static/created prior to begin play for now.

#rnx

diff --git a/Engine/Plugins/Runtime/NetworkPrediction/Source/NetworkPrediction/Private/NetworkPhysicsComponent.cpp b/Engine/Plugins/Runtime/NetworkPrediction/Source/NetworkPrediction/Private/NetworkPhysicsComponent.cpp
--- a/Engine/Plugins/Runtime/NetworkPrediction/Source/NetworkPrediction/Private/NetworkPhysicsComponent.cpp
+++ b/Engine/Plugins/Runtime/NetworkPrediction/Source/NetworkPrediction/Private/NetworkPhysicsComponent.cpp
@@ -64,59 +64,42 @@ void UNetworkPhysicsComponent::InitializeComponent()
 		return Pc;
 	};
 
-	UPrimitiveComponent* PrimitiveComponent = nullptr;
 	if (AActor* MyActor = GetOwner())
 	{
-		// Explicitly tagged component
-		if (ManagedComponentTag != NAME_None)
+		TInlineComponentArray<UPrimitiveComponent*> FoundComponents;
+		MyActor->GetComponents<UPrimitiveComponent>(FoundComponents);
+		for (UPrimitiveComponent* FoundComponent : FoundComponents)
 		{
-			if (UPrimitiveComponent* FoundComponent = SelectComponent(MyActor->GetComponentsByTag(UPrimitiveComponent::StaticClass(), ManagedComponentTag)))
-			{
-				PrimitiveComponent = FoundComponent;
-			}
-			else
+			NetworkPhysicsStateArray.Reserve(FoundComponents.Num());
+
+			if (FPhysicsInterface::IsValid(FoundComponent->BodyInstance.ActorHandle))
 			{
-				UE_LOG(LogNetworkPhysics, Warning, TEXT("Actor %s: could not find a valid Primitive Component with Tag %s"), *MyActor->GetPathName(), *ManagedComponentTag.ToString());
-			}
-		}
+				FNetworkPhysicsState& NewElement = NetworkPhysicsStateArray.AddDefaulted_GetRef();
 
-		// Root component
-		if (!PrimitiveComponent && ValidComponent(MyActor->GetRootComponent()))
-		{
-			PrimitiveComponent = CastChecked<UPrimitiveComponent>(MyActor->GetRootComponent());
-		}
+				NewElement.Proxy = FoundComponent->BodyInstance.ActorHandle;
+				NewElement.OwningActor = GetOwner();
 
-		// Any other valid primitive component?
-		if (!PrimitiveComponent)
-		{
-			if (UPrimitiveComponent* FoundComponent = SelectComponent(MyActor->K2_GetComponentsByClass(UPrimitiveComponent::StaticClass())))
-			{
-				PrimitiveComponent = FoundComponent;
-			}
-		}
-	}
+				Manager->RegisterPhysicsProxy(&NewElement);
 
-	if (ensureMsgf(PrimitiveComponent, TEXT("No PrimitiveComponent found on %s"), *GetPathName()))
-	{
-		NetworkPhysicsState.Proxy = PrimitiveComponent->BodyInstance.ActorHandle;
-		NetworkPhysicsState.OwningActor = GetOwner();
-		ensure(NetworkPhysicsState.OwningActor);
+				Manager->RegisterPhysicsProxyDebugDraw(&NewElement, [this, FoundComponent](const UNetworkPhysicsManager::FDrawDebugParams& P)
+				{
+					const float Thickness = 2.f;
+					
+					FBox Box(ForceInit);
+					Box += FoundComponent->CalcBounds(FTransform::Identity).GetBox();
 
-		Manager->RegisterPhysicsProxy(&NetworkPhysicsState);
+					FVector BoxOrigin;
+					FVector BoxExtent;
+					Box.GetCenterAndExtents(BoxOrigin, BoxExtent);
+					const FVector RelativeScale3D = FoundComponent->GetRelativeScale3D();
+					BoxExtent *= RelativeScale3D;
 
-		Manager->RegisterPhysicsProxyDebugDraw(&NetworkPhysicsState, [this](const UNetworkPhysicsManager::FDrawDebugParams& P)
-			{
-				AActor* Actor = this->GetOwner();
-				FBox LocalSpaceBox = Actor->CalculateComponentsBoundingBoxInLocalSpace();
-				const float Thickness = 2.f;
-
-				FVector ActorOrigin;
-				FVector ActorExtent;
-				LocalSpaceBox.GetCenterAndExtents(ActorOrigin, ActorExtent);
-				ActorExtent *= Actor->GetActorScale3D();
-				DrawDebugBox(P.DrawWorld, P.Loc, ActorExtent, P.Rot, P.Color, false, P.Lifetime, 0, Thickness);
-			});
+					DrawDebugBox(P.DrawWorld, P.Loc + BoxOrigin, BoxExtent, P.Rot, P.Color, false, P.Lifetime, 0, Thickness);
+				});
+			}
+		}
 	}
+
 #endif
 }
 
@@ -128,7 +111,10 @@ void UNetworkPhysicsComponent::UninitializeComponent()
 	{
 		if (UNetworkPhysicsManager* Manager = World->GetSubsystem<UNetworkPhysicsManager>())
 		{
-			Manager->UnregisterPhysicsProxy(&NetworkPhysicsState);
+			for (FNetworkPhysicsState& Element : NetworkPhysicsStateArray)
+			{
+				Manager->UnregisterPhysicsProxy(&Element);
+			}
 		}
 	}
 #endif
@@ -147,7 +133,8 @@ APlayerController* UNetworkPhysicsComponent::GetOwnerPC() const
 void UNetworkPhysicsComponent::GetLifetimeReplicatedProps(TArray<FLifetimeProperty>& OutLifetimeProps) const
 {
 	Super::GetLifetimeReplicatedProps(OutLifetimeProps);
-	DOREPLIFETIME( UNetworkPhysicsComponent, NetworkPhysicsState);
+	DOREPLIFETIME( UNetworkPhysicsComponent, NetworkPhysicsStateArray);
+	
 }
 
 // ========================================================================================
diff --git a/Engine/Plugins/Runtime/NetworkPrediction/Source/NetworkPrediction/Public/NetworkPhysicsComponent.h b/Engine/Plugins/Runtime/NetworkPrediction/Source/NetworkPrediction/Public/NetworkPhysicsComponent.h
--- a/Engine/Plugins/Runtime/NetworkPrediction/Source/NetworkPrediction/Public/NetworkPhysicsComponent.h
+++ b/Engine/Plugins/Runtime/NetworkPrediction/Source/NetworkPrediction/Public/NetworkPhysicsComponent.h
@@ -3,6 +3,7 @@
 #pragma once
 
 #include "NetworkPhysics.h"
+#include "Net/Serialization/FastArraySerializer.h"
 #include "NetworkPhysicsComponent.generated.h"
 
 // Generic component that will register an actor's primitive component with the Network Physics system.
@@ -22,19 +23,18 @@ public:
 	virtual void InitializeComponent() override;
 	virtual void UninitializeComponent() override;
 
-	UPROPERTY(Replicated, transient)
-	FNetworkPhysicsState NetworkPhysicsState;
-
 	UFUNCTION(BlueprintPure, Category = "Network Physics")
-	int32 GetNetworkPredictionLOD() const { return NetworkPhysicsState.LocalLOD; }
+	int32 GetNetworkPredictionLOD() const { return NetworkPhysicsStateArray.Num() > 0 ? NetworkPhysicsStateArray[0].LocalLOD : 0; }
 
 	APlayerController* GetOwnerPC() const;
 
 protected:
+	
+	FSingleParticlePhysicsProxy* GetManagedProxy() const { return NetworkPhysicsStateArray.Num() > 0 ? NetworkPhysicsStateArray[0].Proxy : nullptr; }
+	
+private:
 
-	FSingleParticlePhysicsProxy* GetManagedProxy() const { return NetworkPhysicsState.Proxy; }
-
-	// Which component's physics body to manage if there are multiple PrimitiveComponent and not the root component.
-	UPROPERTY(EditDefaultsOnly, Category = "Network Physics")
-	FName ManagedComponentTag=NAME_None;
+	// Array of physics state being managed. Do not mess with this directly, it should only be touched in InitializeComponent
+	UPROPERTY(Replicated, transient)
+	TArray<FNetworkPhysicsState> NetworkPhysicsStateArray;
 };
\ No newline at end of file
