Network Prediction
-Make physics corrections automatically re-apply if a future rewind happens to a prior frame. Previously the original correction would be list

diff --git a/Engine/Plugins/Runtime/NetworkPrediction/Source/NetworkPrediction/Private/NetworkPhysics.cpp b/Engine/Plugins/Runtime/NetworkPrediction/Source/NetworkPrediction/Private/NetworkPhysics.cpp
--- a/Engine/Plugins/Runtime/NetworkPrediction/Source/NetworkPrediction/Private/NetworkPhysics.cpp
+++ b/Engine/Plugins/Runtime/NetworkPrediction/Source/NetworkPrediction/Private/NetworkPhysics.cpp
@@ -203,8 +203,7 @@ struct FNetworkPhysicsRewindCallback : public Chaos::IRewindCallback
 		npCheckSlow(!RewindData->IsResim());
 		
 		TRACE_CPUPROFILER_EVENT_SCOPE(NPA_Physics_TriggerRewindIfNeeded_Internal);
-
-		PendingCorrections.Reset();
+		
 		const int32 EarliestFrame = RewindData->GetEarliestFrame_Internal();
 
 		if (UE_NETWORK_PHYSICS::ForceResimWithoutCorrection() > 0)
@@ -260,6 +259,17 @@ struct FNetworkPhysicsRewindCallback : public Chaos::IRewindCallback
 					continue;
 				}
 
+				FCorrectionSnapshot& CorrectionSnapshot = CorrectionData_Internal[Obj.Frame % CorrectionData_Internal.Num()];
+				if (CorrectionSnapshot.LocalFrame != Obj.Frame)
+				{
+					CorrectionSnapshot.LocalFrame = Obj.Frame;
+					CorrectionSnapshot.RecvData.Reset();
+				}
+				
+				FPhysicsObjectCorrectionState& CorrectionCopy = CorrectionSnapshot.RecvData.AddDefaulted_GetRef();
+				CorrectionCopy.Proxy = Obj.Proxy;
+				CorrectionCopy.Physics = Obj.Physics;
+
 				const auto P = RewindData->GetPastStateAtFrame(*Proxy->GetHandle_LowLevel(), Obj.Frame);
 				const int32 SimulationFrame = Obj.Frame - Snapshot.LocalFrameOffset;
 
@@ -299,9 +309,6 @@ struct FNetworkPhysicsRewindCallback : public Chaos::IRewindCallback
 					RewindToFrame = RewindToFrame == INDEX_NONE ? Obj.Frame : FMath::Min(RewindToFrame, Obj.Frame);
 					ensure(RewindToFrame >= 0);
 
-					PendingCorrections.Emplace(Obj);
-					PendingCorrectionIdx = 0;
-
 					// Send this correction back to GT for debugging
 					if (UE_NETWORK_PHYSICS::Debug > 0)
 					{
@@ -366,6 +373,8 @@ struct FNetworkPhysicsRewindCallback : public Chaos::IRewindCallback
 				}
 			}
 		}
+
+		//UE_CLOG(UE_NETWORK_PHYSICS::LogCorrections && RewindToFrame != INDEX_NONE, LogNetworkPhysics, Log, TEXT("Rewinding to Frame %d"), RewindToFrame);
 		
 		return RewindToFrame;
 	}
@@ -373,41 +382,28 @@ struct FNetworkPhysicsRewindCallback : public Chaos::IRewindCallback
 	void ProcessInputs_Internal(int32 PhysicsStep, const TArray<Chaos::FSimCallbackInputAndObject>& SimCallbackInputs) override
 	{
 		// Apply Physics corrections if necessary
-		if (PendingCorrectionIdx != INDEX_NONE)
-		{		
-			for (; PendingCorrectionIdx < PendingCorrections.Num(); ++PendingCorrectionIdx)
-			{
-				FNetworkPhysicsState& CorrectionState = PendingCorrections[PendingCorrectionIdx];
-				if (CorrectionState.Frame > PhysicsStep)
-				{
-					break;
-				}
 
-				// TODO: This nullcheck is to avoid null access crash in Editor Tests.
-				// Should re-evaluate whether/why this is occurring in the first place.
+		FCorrectionSnapshot& CorrectionSnapshot = CorrectionData_Internal[PhysicsStep % CorrectionData_Internal.Num()];
+		if (CorrectionSnapshot.LocalFrame == PhysicsStep)
+		{
+			for (FPhysicsObjectCorrectionState& CorrectionState : CorrectionSnapshot.RecvData)
+			{
 				if (ensure(CorrectionState.Proxy))
-				{
+				{		
 					if (auto* PT = CorrectionState.Proxy->GetPhysicsThreadAPI())
 					{
-						UE_CLOG(UE_NETWORK_PHYSICS::LogCorrections > 0, LogNetworkPhysics, Log, TEXT("Applying Correction from frame %d (actual step: %d). Location: %s"), CorrectionState.Frame, PhysicsStep, *FVector(CorrectionState.Physics.Location).ToString());
-
 						npEnsure(CorrectionState.Physics.Location.ContainsNaN() == false);
 						npEnsure(CorrectionState.Physics.LinearVelocity.ContainsNaN() == false);
 						npEnsure(CorrectionState.Physics.AngularVelocity.ContainsNaN() == false);
-					
+
 						PT->SetX(CorrectionState.Physics.Location, false);
 						PT->SetV(CorrectionState.Physics.LinearVelocity, false);
 						PT->SetR(CorrectionState.Physics.Rotation, false);
 						PT->SetW(CorrectionState.Physics.AngularVelocity, false);
 
-						//Solver->GetParticles().MarkTransientDirtyParticle(PT->GetProxy()->GetHandle_LowLevel());
-
-						//if (PT->ObjectState() != CorrectionState.Physics.ObjectState)
-						{
-							ensure(CorrectionState.Physics.ObjectState != Chaos::EObjectStateType::Uninitialized);
-							UE_CLOG(UE_NETWORK_PHYSICS::LogCorrections > 0 && PT->ObjectState() != CorrectionState.Physics.ObjectState, LogNetworkPhysics, Log, TEXT("Applying Correction State %d"), CorrectionState.Physics.ObjectState);
-							PT->SetObjectState(CorrectionState.Physics.ObjectState);
-						}
+						ensure(CorrectionState.Physics.ObjectState != Chaos::EObjectStateType::Uninitialized);
+						UE_CLOG(UE_NETWORK_PHYSICS::LogCorrections > 0 && PT->ObjectState() != CorrectionState.Physics.ObjectState, LogNetworkPhysics, Log, TEXT("Applying Correction State %d"), CorrectionState.Physics.ObjectState);
+						PT->SetObjectState(CorrectionState.Physics.ObjectState);
 					}
 				}
 			}
@@ -557,9 +553,24 @@ struct FNetworkPhysicsRewindCallback : public Chaos::IRewindCallback
 
 	TQueue<FRequest> DataRequested;
 
+	// ----------------------------------------------------------------------------------
+	
 	// Objects that we need to correct in PreResimStep_Internal
-	TArray<FNetworkPhysicsState> PendingCorrections;
-	int32 PendingCorrectionIdx = INDEX_NONE;
+	struct FPhysicsObjectCorrectionState
+	{
+		Chaos::FSingleParticlePhysicsProxy* Proxy = nullptr;
+		FBasePhysicsState Physics;
+	};
+
+	struct FCorrectionSnapshot
+	{
+		int32 LocalFrame = INDEX_NONE;
+		TArray<FPhysicsObjectCorrectionState> RecvData;
+	};
+
+	TStaticArray<FCorrectionSnapshot, 64> CorrectionData_Internal;
+
+	// ----------------------------------------------------------------------------------
 
 	Chaos::FRewindData* RewindData=nullptr; // This be made to be accessed off of IRewindCallback
 	Chaos::FPhysicsSolver* Solver = nullptr; //todo: shouldn't have to cache this, should be part of base class
