Network Physics
-Added NetworkID for managed physics particles, removed the non-networked ID and replaced underlying data structure with a TSet

diff --git a/Engine/Plugins/Runtime/NetworkPrediction/Source/NetworkPrediction/Private/NetworkPhysics.cpp b/Engine/Plugins/Runtime/NetworkPrediction/Source/NetworkPrediction/Private/NetworkPhysics.cpp
--- a/Engine/Plugins/Runtime/NetworkPrediction/Source/NetworkPrediction/Private/NetworkPhysics.cpp
+++ b/Engine/Plugins/Runtime/NetworkPrediction/Source/NetworkPrediction/Private/NetworkPhysics.cpp
@@ -1203,27 +1203,22 @@ void UNetworkPhysicsManager::RegisterPhysicsProxy(FNetworkPhysicsState* State)
 		return;
 	}
 
-	State->LocalManagedHandle = ++UniqueHandleCounter;
-	ManagedHandleToIndexMap.Add(UniqueHandleCounter) = ManagedPhysicsStates.EmplaceAtLowestFreeIndex(LastFreeIndex, State);
-}
+	ManagedPhysicsStates.Add(State);
 
-void UNetworkPhysicsManager::UnregisterPhysicsProxy(FNetworkPhysicsState* State)
-{
-	checkSlow(State);
-	const int32 LocalHandle = State->LocalManagedHandle;
+	UWorld* World = GetWorld();
+	checkSlow(World);
 
-	if (LocalHandle != INDEX_NONE)
+	const bool bIsServer = World->GetNetMode() != NM_Client;
+	if (bIsServer)
 	{
-		int32 LocalIndex;
-		if (ensure(ManagedHandleToIndexMap.RemoveAndCopyValue(LocalHandle, LocalIndex)))
-		{
-			ManagedPhysicsStates.RemoveAt(LocalIndex);
-			LastFreeIndex = FMath::Min(LastFreeIndex, LocalIndex);
-		}
-		
-		State->LocalManagedHandle = INDEX_NONE;
+		State->NetworkHandle = ++UniqueHandleCounter;
 	}
+}
 
+void UNetworkPhysicsManager::UnregisterPhysicsProxy(FNetworkPhysicsState* State)
+{
+	checkSlow(State);
+	ManagedPhysicsStates.Remove(State);
 	DrawDebugMap.Remove(State->Proxy);
 }
 
diff --git a/Engine/Plugins/Runtime/NetworkPrediction/Source/NetworkPrediction/Public/NetworkPhysics.h b/Engine/Plugins/Runtime/NetworkPrediction/Source/NetworkPrediction/Public/NetworkPhysics.h
--- a/Engine/Plugins/Runtime/NetworkPrediction/Source/NetworkPrediction/Public/NetworkPhysics.h
+++ b/Engine/Plugins/Runtime/NetworkPrediction/Source/NetworkPrediction/Public/NetworkPhysics.h
@@ -60,16 +60,13 @@ struct FNetworkPhysicsState
 	// Local proxy we are working with. Must be set locally on initialization.
 	Chaos::FSingleParticlePhysicsProxy* Proxy = nullptr;
 
-	// Local handle used by UNetworkPhysicsManager for registering/unregistering
-	int32 LocalManagedHandle = INDEX_NONE;
+	// Networked ID
+	int32 NetworkHandle = INDEX_NONE;
 
 	// The actual physics properties
 	FBasePhysicsState Physics;
 
 	// Frame number associated with this data
-	//				GT		PT
-	//	Client:	   Remote	Local	
-	//	Server:	   Local	Local
 	int32 Frame = 0;
 
 	// This is what networking will diff to tell if things have changed
@@ -91,6 +88,10 @@ struct FNetworkPhysicsState
 		Ar << Physics.LinearVelocity;
 		Ar << Physics.AngularVelocity;
 
+		uint32 UnsignedHandle = (uint32)NetworkHandle;
+		Ar.SerializeIntPacked(UnsignedHandle);
+		NetworkHandle = (int32)UnsignedHandle;
+
 #if NETWORK_PHYSICS_REPLICATE_EXTRAS
 		for (int32 i=(int32)Chaos::FFrameAndPhase::EParticleHistoryPhase::PrePushData; i < (int32)Chaos::FFrameAndPhase::EParticleHistoryPhase::NumPhases; ++i)
 		{
@@ -203,10 +204,8 @@ private:
 	int32 LocalOffset = 0; // Calculated client/server frame offset. ClientFrame = ServerFrame + LocalOffset
 
 	int32 ClientServerMaxFrameDelta = 0;
-
-	TSortedMap<int32, int32> ManagedHandleToIndexMap;
-	TSparseArray<FNetworkPhysicsState*> ManagedPhysicsStates;
-	int32 LastFreeIndex = 0;
+	
+	TSet<FNetworkPhysicsState*> ManagedPhysicsStates;	
 	int32 UniqueHandleCounter = 0;
 
 	TMap<Chaos::FSingleParticlePhysicsProxy*, TUniqueFunction<void(const FDrawDebugParams&)>> DrawDebugMap;
