Network Prediction
-Make NP2 state corrections resiliant to rewinds instead of being lost

diff --git a/Engine/Plugins/Runtime/NetworkPrediction/Source/NetworkPrediction/Private/NetworkPhysics.cpp b/Engine/Plugins/Runtime/NetworkPrediction/Source/NetworkPrediction/Private/NetworkPhysics.cpp
--- a/Engine/Plugins/Runtime/NetworkPrediction/Source/NetworkPrediction/Private/NetworkPhysics.cpp
+++ b/Engine/Plugins/Runtime/NetworkPrediction/Source/NetworkPrediction/Private/NetworkPhysics.cpp
@@ -24,6 +24,7 @@
 #include "GameFramework/PlayerState.h"
 #include "Misc/FileHelper.h"
 #include "ProfilingDebugging/CpuProfilerTrace.h"
+#include "Async/NetworkPredictionAsyncDefines.h"
 
 DEFINE_LOG_CATEGORY(LogNetworkPhysics);
 
@@ -568,7 +569,7 @@ struct FNetworkPhysicsRewindCallback : public Chaos::IRewindCallback
 		TArray<FPhysicsObjectCorrectionState> RecvData;
 	};
 
-	TStaticArray<FCorrectionSnapshot, 64> CorrectionData_Internal;
+	TStaticArray<FCorrectionSnapshot, UE_NP::NumFramesStorage> CorrectionData_Internal;
 
 	// ----------------------------------------------------------------------------------
 
diff --git a/Engine/Plugins/Runtime/NetworkPrediction/Source/NetworkPrediction/Public/Async/NetworkPredictionAsyncInstanceData.h b/Engine/Plugins/Runtime/NetworkPrediction/Source/NetworkPrediction/Public/Async/NetworkPredictionAsyncInstanceData.h
--- a/Engine/Plugins/Runtime/NetworkPrediction/Source/NetworkPrediction/Public/Async/NetworkPredictionAsyncInstanceData.h
+++ b/Engine/Plugins/Runtime/NetworkPrediction/Source/NetworkPrediction/Public/Async/NetworkPredictionAsyncInstanceData.h
@@ -112,23 +112,28 @@ struct TAsyncNetRecvData
 		InputCmdType LatestInputCmd; // latest, unprocessed InputCmd from server
 	};
 
-
-	TArray<FInstance> NetRecvInstances;
-	TBitArray<> NetRecvDirtyMask;
-	TBitArray<> PendingCorrectionMask;
+	TSparseArray<FInstance>	NetRecvInstances;
 
 	void Reset()
 	{
 		NetRecvInstances.Reset();
-		NetRecvDirtyMask.Reset();
-		PendingCorrectionMask.Reset();
 	}
+};
 
-	void MarkIndexDirty(const int32 idx)
+template<typename AsyncModelDef>
+struct TAsyncCorrectionData
+{
+	HOIST_ASYNCMODELDEF_TYPES()
+
+	struct FInstance
 	{
-		NpResizeForIndex(NetRecvInstances, idx);
-		NpResizeAndSetBit(NetRecvDirtyMask, idx, true);
-	}
+		int32 Index = INDEX_NONE;
+		InputCmdType InputCmd;
+		NetStateType NetState;
+	};
+
+	int32 Frame;
+	TArray<FInstance> Instances;
 };
 
 template<typename AsyncModelDef>
@@ -225,12 +230,12 @@ struct TAsyncModelDataStore_Internal
 	TFrameStorage<TAsncFrameSnapshot<AsyncModelDef>> Frames;
 
 	TQueue<TAsyncNetRecvData<AsyncModelDef>> NetRecvQueue;
-	TAsyncNetRecvData<AsyncModelDef> NetRecvData;
+
+	TStaticArray<TAsyncCorrectionData<AsyncModelDef>, UE_NP::NumFramesStorage> Corrections;
 
 	void Reset()
 	{
 		Instances.Reset();
-		NetRecvData.Reset();
 	}
 };
 
diff --git a/Engine/Plugins/Runtime/NetworkPrediction/Source/NetworkPrediction/Public/Async/Services/NetworkPredictionAsyncService_Marshal.inl b/Engine/Plugins/Runtime/NetworkPrediction/Source/NetworkPrediction/Public/Async/Services/NetworkPredictionAsyncService_Marshal.inl
--- a/Engine/Plugins/Runtime/NetworkPrediction/Source/NetworkPrediction/Public/Async/Services/NetworkPredictionAsyncService_Marshal.inl
+++ b/Engine/Plugins/Runtime/NetworkPrediction/Source/NetworkPrediction/Public/Async/Services/NetworkPredictionAsyncService_Marshal.inl
@@ -168,25 +168,25 @@ public:
 
 	void ApplyCorrections(int32 LocalFrame)
 	{
-		for (TConstSetBitIterator<> BitIt(DataStore->NetRecvData.PendingCorrectionMask); BitIt; ++BitIt)
+		TAsyncCorrectionData<AsyncModelDef>& CorrectionData = DataStore->Corrections[LocalFrame % DataStore->Corrections.Num()];
+		if (CorrectionData.Frame != LocalFrame)
 		{
-			const int32 idx = BitIt.GetIndex();
-			const typename TAsyncNetRecvData<AsyncModelDef>::FInstance& RecvData = DataStore->NetRecvData.NetRecvInstances[idx];
-			if (RecvData.Frame == LocalFrame)
-			{
-				TAsncFrameSnapshot<AsyncModelDef>& Snapshot = DataStore->Frames[LocalFrame];
+			return;
+		}
 
-				NpResizeForIndex(Snapshot.InputCmds, idx);
-				NpResizeForIndex(Snapshot.NetStates, idx);
-				
-				Snapshot.InputCmds[idx] = RecvData.InputCmd;
-				Snapshot.NetStates[idx] = RecvData.NetState;
+		for (TAsyncCorrectionData<AsyncModelDef>::FInstance& CorrectionInstance : CorrectionData.Instances)
+		{
+			TAsncFrameSnapshot<AsyncModelDef>& Snapshot = DataStore->Frames[LocalFrame];
 
-				DataStore->NetRecvData.PendingCorrectionMask[idx] = false;
-			}
+			const int32 idx = CorrectionInstance.Index;
+
+			NpResizeForIndex(Snapshot.InputCmds, idx);
+			NpResizeForIndex(Snapshot.NetStates, idx);
+				
+			Snapshot.InputCmds[idx] = CorrectionInstance.InputCmd;
+			Snapshot.NetStates[idx] = CorrectionInstance.NetState;
 		}
 	}
-	
 
 protected:
 
@@ -245,7 +245,7 @@ public:
 		//	Network
 		// ---------------------------------------------------------------------
 
-		if (DataStore_External->PendingNetRecv.NetRecvDirtyMask.Contains(true))
+		if (DataStore_External->PendingNetRecv.NetRecvInstances.Num() > 0)
 		{
 			DataStore_Internal->NetRecvQueue.Enqueue(MoveTemp(DataStore_External->PendingNetRecv));
 			DataStore_External->PendingNetRecv.Reset();
diff --git a/Engine/Plugins/Runtime/NetworkPrediction/Source/NetworkPrediction/Public/Async/Services/NetworkPredictionAsyncService_NetSerialize.inl b/Engine/Plugins/Runtime/NetworkPrediction/Source/NetworkPrediction/Public/Async/Services/NetworkPredictionAsyncService_NetSerialize.inl
--- a/Engine/Plugins/Runtime/NetworkPrediction/Source/NetworkPrediction/Public/Async/Services/NetworkPredictionAsyncService_NetSerialize.inl
+++ b/Engine/Plugins/Runtime/NetworkPrediction/Source/NetworkPrediction/Public/Async/Services/NetworkPredictionAsyncService_NetSerialize.inl
@@ -33,9 +33,12 @@ public:
 	{
 		TAsncInstanceStaticData<AsyncModelDef>& InstanceData = DataStore->Instances.FindChecked(ID);
 
-		DataStore->PendingNetRecv.MarkIndexDirty(InstanceData.Index);
+		if (DataStore->PendingNetRecv.NetRecvInstances.IsValidIndex(InstanceData.Index) == false)
+		{
+			DataStore->PendingNetRecv.NetRecvInstances.Insert(InstanceData.Index, TAsyncNetRecvData<AsyncModelDef>::FInstance());
+		}
 
-		typename TAsyncNetRecvData<AsyncModelDef>::FInstance& NetRecvData = DataStore->PendingNetRecv.NetRecvInstances[InstanceData.Index];
+		typename TAsyncNetRecvData<AsyncModelDef>::FInstance& NetRecvData = DataStore->PendingNetRecv.NetRecvInstances[InstanceData.Index]; //DataStore->PendingNetRecv.NetRecvInstances[InstanceData.Index];
 		NetRecvData.Frame = ThisFrame;
 		NetRecvData.Flags = Flags;
 		NetRecvData.InputCmd.NetSerialize(Ar);
diff --git a/Engine/Plugins/Runtime/NetworkPrediction/Source/NetworkPrediction/Public/Async/Services/NetworkPredictionAsyncService_Reconcile.inl b/Engine/Plugins/Runtime/NetworkPrediction/Source/NetworkPrediction/Public/Async/Services/NetworkPredictionAsyncService_Reconcile.inl
--- a/Engine/Plugins/Runtime/NetworkPrediction/Source/NetworkPrediction/Public/Async/Services/NetworkPredictionAsyncService_Reconcile.inl
+++ b/Engine/Plugins/Runtime/NetworkPrediction/Source/NetworkPrediction/Public/Async/Services/NetworkPredictionAsyncService_Reconcile.inl
@@ -6,8 +6,10 @@
 
 namespace UE_NP {
 
-	
+
+NETSIM_DEVCVAR_SHIPCONST_INT(LogStateCorrections, 0, "np2.LogStateCorrections", "");
 NETSIM_DEVCVAR_SHIPCONST_INT(EnableFutureInputs, 0, "np2.EnableFutureInputs", "");
+NETSIM_DEVCVAR_SHIPCONST_INT(SkipObjStateReconcinle, 0, "np2.SkipObjStateReconcile", "");
 
 class IAsyncReconcileService
 {
@@ -33,79 +35,90 @@ public:
 	{
 		TRACE_CPUPROFILER_EVENT_SCOPE(NPA_Reconcile);
 
-		if (!DataStore->NetRecvQueue.Dequeue(DataStore->NetRecvData))
-		{
-			return INDEX_NONE;
-		}
-		npEnsure(DataStore->NetRecvQueue.IsEmpty()); // We should enqueue exactly one NetRecv snapshot per marshalled input, which should be exactly one call to Reconcile
-		DataStore->NetRecvData.PendingCorrectionMask.Reset();
-
 		int32 RewindFrame = INDEX_NONE;
-		for (TConstSetBitIterator<> BitIt(DataStore->NetRecvData.NetRecvDirtyMask); BitIt; ++BitIt)
-		{
-			const int32 idx = BitIt.GetIndex();
-			const typename TAsyncNetRecvData<AsyncModelDef>::FInstance& RecvData = DataStore->NetRecvData.NetRecvInstances[idx];
-			RecvData.Frame += LocalFrameOffset; // Convert server-frame to local frame
+		TAsyncNetRecvData<AsyncModelDef> NetRecvData;
 
-			if (RecvData.Frame >= EarliestFrame && RecvData.Frame <= LastCompletedStep)
+		while (DataStore->NetRecvQueue.Dequeue(NetRecvData))
+		{	
+			for (auto It = NetRecvData.NetRecvInstances.CreateIterator(); It; ++It)
 			{
-				const TAsncFrameSnapshot<AsyncModelDef>& Snapshot = DataStore->Frames[RecvData.Frame];
+				const int32 idx = It.GetIndex();
+				const typename TAsyncNetRecvData<AsyncModelDef>::FInstance& RecvData = *It;
+				RecvData.Frame += LocalFrameOffset; // Convert server-frame to local frame
 
-				// This is probably not enough, we will need to check some alive flag to make sure this state was active at 
-				// this time or something?
-				bool bShouldReconcile = false;
-				if (Snapshot.InputCmds.IsValidIndex(idx) && Snapshot.NetStates.IsValidIndex(idx))
+				if (RecvData.Frame >= EarliestFrame && RecvData.Frame <= LastCompletedStep)
 				{
-					const InputCmdType& LocalInputCmd = Snapshot.InputCmds[idx];
-					const NetStateType& LocalNetState = Snapshot.NetStates[idx];
+					const TAsncFrameSnapshot<AsyncModelDef>& Snapshot = DataStore->Frames[RecvData.Frame];
 
-					if (LocalInputCmd.ShouldReconcile(RecvData.InputCmd))
+					// This is probably not enough, we will need to check some alive flag to make sure this state was active at 
+					// this time or something?
+					bool bShouldReconcile = false;
+					if (Snapshot.InputCmds.IsValidIndex(idx) && Snapshot.NetStates.IsValidIndex(idx))
 					{
-						//UE_LOG(LogNetworkPrediction, Log, TEXT("Reconcile on InputCmd. %s. Instance: %d. Frame: [%d/%d]"), AsyncModelDef::GetName(), idx, RecvData.Frame - LocalFrameOffset, RecvData.Frame);
-						bShouldReconcile = true;
+						if (SkipObjStateReconcinle() == 0)
+						{
+							const InputCmdType& LocalInputCmd = Snapshot.InputCmds[idx];
+							const NetStateType& LocalNetState = Snapshot.NetStates[idx];
+
+							if (LocalInputCmd.ShouldReconcile(RecvData.InputCmd))
+							{
+								UE_CLOG(LogStateCorrections() > 0, LogNetworkPrediction, Log, TEXT("Reconcile on InputCmd. %s. Instance: %d. Frame: [%d/%d]"), AsyncModelDef::GetName(), idx, RecvData.Frame - LocalFrameOffset, RecvData.Frame);
+								bShouldReconcile = true;
+							}
+							if (LocalNetState.ShouldReconcile(RecvData.NetState))
+							{
+								UE_CLOG(LogStateCorrections() > 0, LogNetworkPrediction, Log, TEXT("Reconcile on NetState. %s. Instance: %d. Frame: [%d/%d]"), AsyncModelDef::GetName(), idx, RecvData.Frame - LocalFrameOffset, RecvData.Frame);
+								bShouldReconcile = true;
+							}
+						}
 					}
-					if (LocalNetState.ShouldReconcile(RecvData.NetState))
+					else
 					{
-						//UE_LOG(LogNetworkPrediction, Log, TEXT("Reconcile on NetState. %s. Instance: %d. Frame: [%d/%d]"), AsyncModelDef::GetName(), idx, RecvData.Frame - LocalFrameOffset, RecvData.Frame);
+						UE_CLOG(LogStateCorrections() > 0, LogNetworkPrediction, Log, TEXT("Reconcile due to invalid state. %s. Instance: %d. Frame: [%d/%d]"), AsyncModelDef::GetName(), idx, RecvData.Frame - LocalFrameOffset, RecvData.Frame);
 						bShouldReconcile = true;
 					}
-				}
-				else
-				{
-					//UE_LOG(LogNetworkPrediction, Log, TEXT("Reconcile due to invalid state. %s. Instance: %d. Frame: [%d/%d]"), AsyncModelDef::GetName(), idx, RecvData.Frame - LocalFrameOffset, RecvData.Frame);
-					bShouldReconcile = true;
-				}
 
-				if (bShouldReconcile)
-				{
-					NpResizeAndSetBit(DataStore->NetRecvData.PendingCorrectionMask, idx);
-					RewindFrame = RewindFrame == INDEX_NONE ? RecvData.Frame : FMath::Min<int32>(RewindFrame, RecvData.Frame);
+					if (bShouldReconcile)
+					{
+						TAsyncCorrectionData<AsyncModelDef>& CorrectionData = DataStore->Corrections[RecvData.Frame % DataStore->Corrections.Num()];
+						if (CorrectionData.Frame != RecvData.Frame)
+						{
+							CorrectionData.Frame = RecvData.Frame;
+							CorrectionData.Instances.Reset();
+						}
+
+						TAsyncCorrectionData<AsyncModelDef>::FInstance& CorrectionInstance = CorrectionData.Instances.AddDefaulted_GetRef();
+						CorrectionInstance.Index = idx;
+						CorrectionInstance.InputCmd = RecvData.InputCmd;
+						CorrectionInstance.NetState = RecvData.NetState;
+					
+						RewindFrame = RewindFrame == INDEX_NONE ? RecvData.Frame : FMath::Min<int32>(RewindFrame, RecvData.Frame);
+					}
 				}
-			}
 
-			// -------------------------------------------------------
-			// Apply all Input corrections for future frames if non-locally controlled.
-			//	doing this in a second pass maybe more cache coherent, but this is simpler for now.
-			// -------------------------------------------------------
+				// -------------------------------------------------------
+				// Apply all Input corrections for future frames if non-locally controlled.
+				//	doing this in a second pass maybe more cache coherent, but this is simpler for now.
+				// -------------------------------------------------------
 			
-			if ((RecvData.Flags & (uint8)EInstanceFlags::LocallyControlled) == 0)
-			{
-				const InputCmdType& SourceInput = EnableFutureInputs() ? RecvData.LatestInputCmd : RecvData.InputCmd;
+				if ((RecvData.Flags & (uint8)EInstanceFlags::LocallyControlled) == 0)
+				{
+					const InputCmdType& SourceInput = EnableFutureInputs() ? RecvData.LatestInputCmd : RecvData.InputCmd;
 
-				const int32 NumFrames = (LastCompletedStep+1) - FMath::Max<int32>(0, RecvData.Frame+1);
-				npEnsureMsgf(NumFrames < DataStore->Frames.Num(), TEXT("Too many input corrections. [%d %d] %d %d"), (LastCompletedStep+1), FMath::Max<int32>(0, RecvData.Frame+1), NumFrames, DataStore->Frames.Num() );
+					const int32 NumFrames = (LastCompletedStep+1) - FMath::Max<int32>(0, RecvData.Frame+1);
+					npEnsureMsgf(NumFrames < DataStore->Frames.Num(), TEXT("Too many input corrections. [%d %d] %d %d"), (LastCompletedStep+1), FMath::Max<int32>(0, RecvData.Frame+1), NumFrames, DataStore->Frames.Num() );
 
-				for (int32 Frame= FMath::Max<int32>(0, RecvData.Frame+1); Frame <= LastCompletedStep+1; ++Frame)
-				{
-					TAsncFrameSnapshot<AsyncModelDef>& FutureSnapshot = DataStore->Frames[Frame];
+					for (int32 Frame= FMath::Max<int32>(0, RecvData.Frame+1); Frame <= LastCompletedStep+1; ++Frame)
+					{
+						TAsncFrameSnapshot<AsyncModelDef>& FutureSnapshot = DataStore->Frames[Frame];
 
-					NpResizeForIndex(FutureSnapshot.InputCmds, idx);
-					FutureSnapshot.InputCmds[idx] = SourceInput;
+						NpResizeForIndex(FutureSnapshot.InputCmds, idx);
+						FutureSnapshot.InputCmds[idx] = SourceInput;
+					}
 				}
-			}			
+			}
 		}
-
-		DataStore->NetRecvData.NetRecvDirtyMask.Reset();
+		
 		return RewindFrame;
 	}
 
