CIS fix for clang

diff --git a/Engine/Plugins/Runtime/NetworkPrediction/Source/NetworkPrediction/Public/Async/Services/NetworkPredictionAsyncService_Marshal.inl b/Engine/Plugins/Runtime/NetworkPrediction/Source/NetworkPrediction/Public/Async/Services/NetworkPredictionAsyncService_Marshal.inl
--- a/Engine/Plugins/Runtime/NetworkPrediction/Source/NetworkPrediction/Public/Async/Services/NetworkPredictionAsyncService_Marshal.inl
+++ b/Engine/Plugins/Runtime/NetworkPrediction/Source/NetworkPrediction/Public/Async/Services/NetworkPredictionAsyncService_Marshal.inl
@@ -174,7 +174,7 @@ public:
 			return;
 		}
 
-		for (TAsyncCorrectionData<AsyncModelDef>::FInstance& CorrectionInstance : CorrectionData.Instances)
+		for (typename TAsyncCorrectionData<AsyncModelDef>::FInstance& CorrectionInstance : CorrectionData.Instances)
 		{
 			TAsncFrameSnapshot<AsyncModelDef>& Snapshot = DataStore->Frames[LocalFrame];
 
diff --git a/Engine/Plugins/Runtime/NetworkPrediction/Source/NetworkPrediction/Public/Async/Services/NetworkPredictionAsyncService_Reconcile.inl b/Engine/Plugins/Runtime/NetworkPrediction/Source/NetworkPrediction/Public/Async/Services/NetworkPredictionAsyncService_Reconcile.inl
--- a/Engine/Plugins/Runtime/NetworkPrediction/Source/NetworkPrediction/Public/Async/Services/NetworkPredictionAsyncService_Reconcile.inl
+++ b/Engine/Plugins/Runtime/NetworkPrediction/Source/NetworkPrediction/Public/Async/Services/NetworkPredictionAsyncService_Reconcile.inl
@@ -87,7 +87,7 @@ public:
 							CorrectionData.Instances.Reset();
 						}
 
-						TAsyncCorrectionData<AsyncModelDef>::FInstance& CorrectionInstance = CorrectionData.Instances.AddDefaulted_GetRef();
+						typename TAsyncCorrectionData<AsyncModelDef>::FInstance& CorrectionInstance = CorrectionData.Instances.AddDefaulted_GetRef();
 						CorrectionInstance.Index = idx;
 						CorrectionInstance.InputCmd = RecvData.InputCmd;
 						CorrectionInstance.NetState = RecvData.NetState;
