Fix for mac editor compile

diff --git a/Engine/Plugins/Runtime/NetworkPrediction/Source/NetworkPrediction/Public/Async/Services/NetworkPredictionAsyncService_NetSerialize.inl b/Engine/Plugins/Runtime/NetworkPrediction/Source/NetworkPrediction/Public/Async/Services/NetworkPredictionAsyncService_NetSerialize.inl
--- a/Engine/Plugins/Runtime/NetworkPrediction/Source/NetworkPrediction/Public/Async/Services/NetworkPredictionAsyncService_NetSerialize.inl
+++ b/Engine/Plugins/Runtime/NetworkPrediction/Source/NetworkPrediction/Public/Async/Services/NetworkPredictionAsyncService_NetSerialize.inl
@@ -35,7 +35,7 @@ public:
 
 		if (DataStore->PendingNetRecv.NetRecvInstances.IsValidIndex(InstanceData.Index) == false)
 		{
-			DataStore->PendingNetRecv.NetRecvInstances.Insert(InstanceData.Index, TAsyncNetRecvData<AsyncModelDef>::FInstance());
+			DataStore->PendingNetRecv.NetRecvInstances.Insert(InstanceData.Index, typename TAsyncNetRecvData<AsyncModelDef>::FInstance());
 		}
 
 		typename TAsyncNetRecvData<AsyncModelDef>::FInstance& NetRecvData = DataStore->PendingNetRecv.NetRecvInstances[InstanceData.Index]; //DataStore->PendingNetRecv.NetRecvInstances[InstanceData.Index];
