Network Physics
-Enable global ParticleID sorting

diff --git a/Engine/Plugins/Runtime/NetworkPrediction/Source/NetworkPrediction/Private/NetworkPhysics.cpp b/Engine/Plugins/Runtime/NetworkPrediction/Source/NetworkPrediction/Private/NetworkPhysics.cpp
--- a/Engine/Plugins/Runtime/NetworkPrediction/Source/NetworkPrediction/Private/NetworkPhysics.cpp
+++ b/Engine/Plugins/Runtime/NetworkPrediction/Source/NetworkPrediction/Private/NetworkPhysics.cpp
@@ -56,6 +56,9 @@ namespace UE_NETWORK_PHYSICS
 
 	NP_DEVCVAR_INT(ForceInputFault, 0, "np2.ForceInputFault", "Force input fault correction");
 
+	NETSIM_DEVCVAR_SHIPCONST_INT(EnableParticleIDSort, 1, "np2.EnableParticleIDSort", "Uses NetworkHandle/SetParticleID to maintain consisten ordering between clients and server for more determinism.");		
+
+
 	// NOTE: These have temporarily been switched to regular CVar style variables, because the NP_DECVAR_* macros
 	// can only make ECVF_Cheat variables, which cannot be hotfixed with DefaultEngine.ini. Instead they use
 	// ConsoleVariables.ini, which isn't loaded in shipping.
@@ -807,6 +810,13 @@ void UNetworkPhysicsManager::PostNetRecv()
 		for (FNetworkPhysicsState* PhysicsState : ManagedPhysicsStates)
 		{
 			check(PhysicsState);
+
+			if (PhysicsState->Proxy && PhysicsState->NetworkHandle != INDEX_NONE && PhysicsState->bParticleIDSet == false && UE_NETWORK_PHYSICS::EnableParticleIDSort() > 0)
+			{
+				PhysicsState->Proxy->GetGameThreadAPI().SetParticleID(Chaos::FParticleID{PhysicsState->NetworkHandle, INDEX_NONE});
+				PhysicsState->bParticleIDSet = true;
+			}
+
 			MaxFrameSeen = FMath::Max(PhysicsState->Frame, MaxFrameSeen);
 			if (PhysicsState->Frame > LatestConfirmedFrame)
 			{
@@ -1213,6 +1223,12 @@ void UNetworkPhysicsManager::RegisterPhysicsProxy(FNetworkPhysicsState* State)
 	if (bIsServer)
 	{
 		State->NetworkHandle = ++UniqueHandleCounter;
+
+		if (UE_NETWORK_PHYSICS::EnableParticleIDSort() > 0)
+		{
+			State->Proxy->GetGameThreadAPI().SetParticleID(Chaos::FParticleID{State->NetworkHandle, INDEX_NONE});
+			State->bParticleIDSet = true;
+		}
 	}
 }
 
diff --git a/Engine/Plugins/Runtime/NetworkPrediction/Source/NetworkPrediction/Public/NetworkPhysics.h b/Engine/Plugins/Runtime/NetworkPrediction/Source/NetworkPrediction/Public/NetworkPhysics.h
--- a/Engine/Plugins/Runtime/NetworkPrediction/Source/NetworkPrediction/Public/NetworkPhysics.h
+++ b/Engine/Plugins/Runtime/NetworkPrediction/Source/NetworkPrediction/Public/NetworkPhysics.h
@@ -62,6 +62,7 @@ struct FNetworkPhysicsState
 
 	// Networked ID
 	int32 NetworkHandle = INDEX_NONE;
+	bool bParticleIDSet = false;
 
 	// The actual physics properties
 	FBasePhysicsState Physics;
