Rewind data can dump out full history. Added a cvar so we can auto dump rewind data on client and server when correction happens

diff --git a/Engine/Plugins/Runtime/NetworkPrediction/Source/NetworkPrediction/Private/NetworkPhysics.cpp b/Engine/Plugins/Runtime/NetworkPrediction/Source/NetworkPrediction/Private/NetworkPhysics.cpp
--- a/Engine/Plugins/Runtime/NetworkPrediction/Source/NetworkPrediction/Private/NetworkPhysics.cpp
+++ b/Engine/Plugins/Runtime/NetworkPrediction/Source/NetworkPrediction/Private/NetworkPhysics.cpp
@@ -168,6 +168,9 @@ namespace UE_NETWORK_PHYSICS
 	}
 }
 
+int32 DumpOnRewind = 0;
+FAutoConsoleVariableRef CVarDumpOnRewind(TEXT("np2.DumpOnRewind"), DumpOnRewind, TEXT("If enabled, automatically dumps full rewind data history before rewinding on both client and server"));
+
 struct FNetworkPhysicsRewindCallback : public Chaos::IRewindCallback
 {
 	bool CompareVec(const FVector& A, const FVector& B, const float D, const TCHAR* Str)
@@ -380,6 +383,11 @@ struct FNetworkPhysicsRewindCallback : public Chaos::IRewindCallback
 
 		//UE_CLOG(UE_NETWORK_PHYSICS::LogCorrections && RewindToFrame != INDEX_NONE, LogNetworkPhysics, Log, TEXT("Rewinding to Frame %d"), RewindToFrame);
 		
+		if(!!DumpOnRewind && RewindToFrame != INDEX_NONE)
+		{
+			RewindData->DumpHistory_Internal(-LastLocalOffset, FString(TEXT("Client")));
+		}
+
 		return RewindToFrame;
 	}
 	
@@ -1297,6 +1305,7 @@ void UNetworkPhysicsManager::TickDrawDebug()
 		// Draw corrections
 		P.Lifetime = 5.f;
 
+		bool bSendToServer = true;
 		while(RewindCallback->StatsForGT.IsEmpty() == false)
 		{
 			FNetworkPhysicsRewindCallback::FStats& Stats = RewindCallback->RecordedStatsGT[RewindCallback->LastRecordedStat++ % RewindCallback->RecordedStatsGT.Num()];
@@ -1304,6 +1313,26 @@ void UNetworkPhysicsManager::TickDrawDebug()
 
 			for (auto& Correction : Stats.Corrections)
 			{
+				//HACK
+				if(bSendToServer && !!DumpOnRewind)
+				{
+					if (GetWorld()->GetNetMode() == NM_Client)
+					{
+						if (APlayerController* PC = GEngine->GetFirstLocalPlayerController(GetWorld()))
+						{
+							UE_LOG(LogChaos, Warning, TEXT("REQUEST SERVER DUMP"));
+							PC->ServerExec(TEXT("np2.RewindDump"));
+						}
+					}
+
+					bSendToServer = false;
+				}
+				else if(DumpOnRewind)
+				{
+					UE_LOG(LogChaos, Warning, TEXT("SKIP SERVER DUMP"));
+				}
+				
+
 				if (UE_NETWORK_PHYSICS::DebugTolerance > 0)
 				{
 					if (FVector::Distance(Correction.LocalState.Location, Correction.AuthorityState.Location) < UE_NETWORK_PHYSICS::DebugDrawTolerance_X)
@@ -1489,6 +1518,16 @@ void UNetworkPhysicsManager::TickDrawDebug()
 	}
 }
 
+void UNetworkPhysicsManager::DumpRewindHistory()
+{
+	using namespace Chaos;
+	FRewindData* RewindData = GetWorld()->GetPhysicsScene()->GetSolver()->GetRewindData();
+	GetWorld()->GetPhysicsScene()->GetSolver()->EnqueueCommandImmediate([RewindData, CapturedOffset = LocalOffset, CapturedIsServer = GetWorld()->GetNetMode() != NM_Client]()
+	{
+		RewindData->DumpHistory_Internal(-CapturedOffset, CapturedIsServer ? FString(TEXT("Server")) : FString(TEXT("Client")) );
+	});
+}
+
 void UNetworkPhysicsManager::DumpDebugHistory()
 {
 	FStringOutputDevice Out;
@@ -1549,6 +1588,23 @@ FAutoConsoleCommandWithWorldAndArgs NpDumpCmd(TEXT("np2.Dump"), TEXT(""),
 	}
 }));
 
+FAutoConsoleCommandWithWorldAndArgs NpRewindDumpCmd(TEXT("np2.RewindDump"), TEXT(""),
+	FConsoleCommandWithWorldAndArgsDelegate::CreateLambda([](const TArray< FString >& Args, UWorld* InWorld)
+		{
+			if (UNetworkPhysicsManager* Manager = InWorld->GetSubsystem<UNetworkPhysicsManager>())
+			{
+				Manager->DumpRewindHistory();
+			}
+
+			if (InWorld->GetNetMode() == NM_Client)
+			{
+				if (APlayerController* PC = GEngine->GetFirstLocalPlayerController(InWorld))
+				{
+					PC->ServerExec(TEXT("np2.RewindDump"));
+				}
+			}
+		}));
+
 FAutoConsoleCommandWithWorldAndArgs NpFloatTestCmd(TEXT("np2.FloatTest"), TEXT(""),
 	FConsoleCommandWithWorldAndArgsDelegate::CreateLambda([](const TArray< FString >& Args, UWorld* InWorld) 
 {
diff --git a/Engine/Plugins/Runtime/NetworkPrediction/Source/NetworkPrediction/Public/NetworkPhysics.h b/Engine/Plugins/Runtime/NetworkPrediction/Source/NetworkPrediction/Public/NetworkPhysics.h
--- a/Engine/Plugins/Runtime/NetworkPrediction/Source/NetworkPrediction/Public/NetworkPhysics.h
+++ b/Engine/Plugins/Runtime/NetworkPrediction/Source/NetworkPrediction/Public/NetworkPhysics.h
@@ -190,6 +190,7 @@ public:
 	void RegisterPhysicsProxyDebugDraw(FNetworkPhysicsState* State, TUniqueFunction<void(const FDrawDebugParams&)>&& Func);
 
 	void DumpDebugHistory();
+	void DumpRewindHistory();
 
 private:
 
diff --git a/Engine/Source/Runtime/Experimental/Chaos/Private/RewindData.cpp b/Engine/Source/Runtime/Experimental/Chaos/Private/RewindData.cpp
--- a/Engine/Source/Runtime/Experimental/Chaos/Private/RewindData.cpp
+++ b/Engine/Source/Runtime/Experimental/Chaos/Private/RewindData.cpp
@@ -5,6 +5,7 @@
 #include "PhysicsProxy/SingleParticlePhysicsProxy.h"
 #include "HAL/IConsoleManager.h"
 #include "Chaos/PBDJointConstraints.h"
+#include "Misc/FileHelper.h"
 
 namespace Chaos
 {
@@ -312,6 +313,28 @@ void FRewindData::FinishFrame()
 	LatestFrame = FMath::Max(LatestFrame,CurFrame);
 }
 
+void FRewindData::DumpHistory_Internal(const int32 FramePrintOffset, const FString& Filename)
+{
+	FStringOutputDevice Out;
+	const int32 EarliestFrame = GetEarliestFrame_Internal();
+	for(int32 Frame = EarliestFrame; Frame < CurFrame; ++Frame)
+	{
+		for (int32 Phase = 0; Phase < FFrameAndPhase::EParticleHistoryPhase::NumPhases; ++Phase)
+		{
+			for(const FDirtyParticleInfo& Info : DirtyParticles)
+			{
+				Out.Logf(TEXT("Frame:%d Phase:%d\n"), Frame + FramePrintOffset, Phase);
+				FGeometryParticleState State = GetPastStateAtFrame(*Info.GetObjectPtr(), Frame, (FFrameAndPhase::EParticleHistoryPhase)Phase);
+				Out.Logf(TEXT("%s\n"), *State.ToString());
+			}
+		}
+	}
+
+	FString Path = FPaths::ProfilingDir() + FString::Printf(TEXT("/RewindData/%s_%d_%d.txt"), *Filename, EarliestFrame + FramePrintOffset, CurFrame - 1 + FramePrintOffset);
+	FFileHelper::SaveStringToFile(Out, *Path);
+	UE_LOG(LogChaos, Warning, TEXT("Saved:%s"), *Path);
+}
+
 CHAOS_API int32 SkipDesyncTest = 0;
 FAutoConsoleVariableRef CVarSkipDesyncTest(TEXT("p.SkipDesyncTest"), SkipDesyncTest, TEXT("Skips hard desync test, this means all particles will assume to be clean except spawning at different times. This is useful for a perf lower bound, not actually correct"));
 
diff --git a/Engine/Source/Runtime/Experimental/Chaos/Public/RewindData.h b/Engine/Source/Runtime/Experimental/Chaos/Public/RewindData.h
--- a/Engine/Source/Runtime/Experimental/Chaos/Public/RewindData.h
+++ b/Engine/Source/Runtime/Experimental/Chaos/Public/RewindData.h
@@ -420,6 +420,33 @@ struct FShapesArrayStateBase
 	}
 };
 
+template <typename T>
+FString ToStringHelper(const T& Val)
+{
+	return Val.ToString();
+}
+
+inline FString ToStringHelper(const FReal Val)
+{
+	return FString::Printf(TEXT("%f"), Val);
+}
+
+inline FString ToStringHelper(const EObjectStateType Val)
+{
+	return FString::Printf(TEXT("%d"), Val);
+}
+
+inline FString ToStringHelper(const bool Val)
+{
+	return FString::Printf(TEXT("%d"), Val);
+}
+
+inline FString ToStringHelper(const int32 Val)
+{
+	return FString::Printf(TEXT("%d"), Val);
+}
+
+
 template <typename TParticle>
 class TShapesArrayState
 {
@@ -637,6 +664,51 @@ public:
 		State = InState;
 	}
 
+	FString ToString() const
+	{
+#undef REWIND_PARTICLE_TO_STR
+#define REWIND_PARTICLE_TO_STR(PropName) Out += FString::Printf(TEXT(#PropName":%s\n"), *ToStringHelper(PropName()));
+		//TODO: use macro to define api and the to string
+		FString Out = FString::Printf(TEXT("ParticleID:[%d%d]\n"), Particle.ParticleID().GlobalID, Particle.ParticleID().LocalID);
+
+		REWIND_PARTICLE_TO_STR(X)
+		REWIND_PARTICLE_TO_STR(R)
+		//REWIND_PARTICLE_TO_STR(Geometry)
+		//REWIND_PARTICLE_TO_STR(UniqueIdx)
+		//REWIND_PARTICLE_TO_STR(SpatialIdx)
+
+		if(Particle.CastToKinematicParticle())
+		{
+			REWIND_PARTICLE_TO_STR(V)
+			REWIND_PARTICLE_TO_STR(W)
+		}
+
+		if(Particle.CastToRigidParticle())
+		{
+			REWIND_PARTICLE_TO_STR(LinearEtherDrag)
+			REWIND_PARTICLE_TO_STR(AngularEtherDrag)
+			REWIND_PARTICLE_TO_STR(MaxLinearSpeedSq)
+			REWIND_PARTICLE_TO_STR(MaxAngularSpeedSq)
+			REWIND_PARTICLE_TO_STR(ObjectState)
+			REWIND_PARTICLE_TO_STR(GravityEnabled)
+			REWIND_PARTICLE_TO_STR(CCDEnabled)
+			REWIND_PARTICLE_TO_STR(CollisionGroup)
+
+			REWIND_PARTICLE_TO_STR(CenterOfMass)
+			REWIND_PARTICLE_TO_STR(RotationOfMass)
+			REWIND_PARTICLE_TO_STR(I)
+			REWIND_PARTICLE_TO_STR(M)
+			REWIND_PARTICLE_TO_STR(InvM)
+
+			REWIND_PARTICLE_TO_STR(F)
+			REWIND_PARTICLE_TO_STR(Torque)
+			REWIND_PARTICLE_TO_STR(LinearImpulse)
+			REWIND_PARTICLE_TO_STR(AngularImpulse)
+		}
+
+		return Out;
+	}
+
 private:
 	const FGeometryParticleHandle& Particle;
 	const FDirtyPropertiesPool& Pool;
@@ -866,6 +938,8 @@ public:
 		return !!EnableResimCache && bResimOptimization ? Managers[CurFrame].ExternalResimCache.Get() : nullptr;
 	}
 
+	void CHAOS_API DumpHistory_Internal(const int32 FramePrintOffset, const FString& Filename = FString(TEXT("Dump")));
+
 	template <typename CreateCache>
 	void AdvanceFrame(FReal DeltaTime, const CreateCache& CreateCacheFunc)
 	{
