Network Prediction Plugin
-Edigrate source changes over from UE4 plus some additional fixes

#rnx

diff --git a/Engine/Plugins/Runtime/NetworkPrediction/Config/BaseEngine.ini b/Engine/Plugins/Runtime/NetworkPrediction/Config/BaseEngine.ini
new file mode 100644
--- /dev/null
+++ b/Engine/Plugins/Runtime/NetworkPrediction/Config/BaseEngine.ini
@@ -0,0 +1,2 @@
+[CoreRedirects]
++ClassRedirects=(OldName="/Script/NetworkPredictionExtras.NetworkPhysicsComponent", NewName="/Script/NetworkPrediction.NetworkPhysicsComponent")
diff --git a/Engine/Plugins/Runtime/NetworkPrediction/Source/NetworkPrediction/NetworkPrediction.Build.cs b/Engine/Plugins/Runtime/NetworkPrediction/Source/NetworkPrediction/NetworkPrediction.Build.cs
--- a/Engine/Plugins/Runtime/NetworkPrediction/Source/NetworkPrediction/NetworkPrediction.Build.cs
+++ b/Engine/Plugins/Runtime/NetworkPrediction/Source/NetworkPrediction/NetworkPrediction.Build.cs
@@ -51,7 +51,6 @@ namespace UnrealBuildTool.Rules
                 PrivateDependencyModuleNames.AddRange(
                 new string[]
                 {
-					"EditorFramework",
                     "UnrealEd",
                 });
             }
diff --git a/Engine/Plugins/Runtime/NetworkPrediction/Source/NetworkPrediction/Private/NetworkPhysics.cpp b/Engine/Plugins/Runtime/NetworkPrediction/Source/NetworkPrediction/Private/NetworkPhysics.cpp
--- a/Engine/Plugins/Runtime/NetworkPrediction/Source/NetworkPrediction/Private/NetworkPhysics.cpp
+++ b/Engine/Plugins/Runtime/NetworkPrediction/Source/NetworkPrediction/Private/NetworkPhysics.cpp
@@ -175,7 +175,10 @@ struct FNetworkPhysicsRewindCallback : public Chaos::IRewindCallback
 		const float Error = FQuat::ErrorAutoNormalize(A, B);
 		const bool b = D == 0 ? A != B : Error > D;
 
-		//UE_LOG(LogTemp, Warning, TEXT("Error: %f"), Error);
+		if (b)
+		{
+			//UE_LOG(LogTemp, Warning, TEXT("%s Error: %f"), Str, Error);
+		}
 
 		UE_CLOG(UE_NETWORK_PHYSICS::LogCorrections > 0 && b && Str, LogNetworkPhysics, Log, TEXT("%s correction. Server: %s. Local: %s. Delta: %f"), Str, *A.ToString(), *B.ToString(), FQuat::ErrorAutoNormalize(A, B));
 		return b;
@@ -278,7 +281,7 @@ struct FNetworkPhysicsRewindCallback : public Chaos::IRewindCallback
 					}
 				}
 #endif
-				
+
 				//UE_LOG(LogTemp, Warning, TEXT("0x%X P.ObjectState(): %d. R: %s [%s]"), (int64)Proxy,  P.ObjectState(), *P.R().ToString(), *Obj.Physics.Rotation.ToString());
 				//UE_LOG(LogTemp, Warning, TEXT("			R: %s [%s]"), *FRotator(P.R()).ToString(), *FRotator(Obj.Physics.Rotation).ToString());
 
@@ -384,30 +387,31 @@ struct FNetworkPhysicsRewindCallback : public Chaos::IRewindCallback
 				// Should re-evaluate whether/why this is occurring in the first place.
 				if (ensure(CorrectionState.Proxy))
 				{
-				if (auto* PT = CorrectionState.Proxy->GetPhysicsThreadAPI())
-				{
-					UE_CLOG(UE_NETWORK_PHYSICS::LogCorrections > 0, LogNetworkPhysics, Log, TEXT("Applying Correction from frame %d (actual step: %d). Location: %s"), CorrectionState.Frame, PhysicsStep, *FVector(CorrectionState.Physics.Location).ToString());
-				
+					if (auto* PT = CorrectionState.Proxy->GetPhysicsThreadAPI())
+					{
+						UE_CLOG(UE_NETWORK_PHYSICS::LogCorrections > 0, LogNetworkPhysics, Log, TEXT("Applying Correction from frame %d (actual step: %d). Location: %s"), CorrectionState.Frame, PhysicsStep, *FVector(CorrectionState.Physics.Location).ToString());
+
 						npEnsure(CorrectionState.Physics.Location.ContainsNaN() == false);
 						npEnsure(CorrectionState.Physics.LinearVelocity.ContainsNaN() == false);
 						npEnsure(CorrectionState.Physics.AngularVelocity.ContainsNaN() == false);
-					PT->SetX(CorrectionState.Physics.Location, false);
-					PT->SetV(CorrectionState.Physics.LinearVelocity, false);
-					PT->SetR(CorrectionState.Physics.Rotation, false);
-					PT->SetW(CorrectionState.Physics.AngularVelocity, false);
 					
+						PT->SetX(CorrectionState.Physics.Location, false);
+						PT->SetV(CorrectionState.Physics.LinearVelocity, false);
+						PT->SetR(CorrectionState.Physics.Rotation, false);
+						PT->SetW(CorrectionState.Physics.AngularVelocity, false);
+
 						//Solver->GetParticles().MarkTransientDirtyParticle(PT->GetProxy()->GetHandle_LowLevel());
 
-					//if (PT->ObjectState() != CorrectionState.Physics.ObjectState)
-					{
-						ensure(CorrectionState.Physics.ObjectState != Chaos::EObjectStateType::Uninitialized);
+						//if (PT->ObjectState() != CorrectionState.Physics.ObjectState)
+						{
+							ensure(CorrectionState.Physics.ObjectState != Chaos::EObjectStateType::Uninitialized);
 							UE_CLOG(UE_NETWORK_PHYSICS::LogCorrections > 0 && PT->ObjectState() != CorrectionState.Physics.ObjectState, LogNetworkPhysics, Log, TEXT("Applying Correction State %d"), CorrectionState.Physics.ObjectState);
-						PT->SetObjectState(CorrectionState.Physics.ObjectState);
+							PT->SetObjectState(CorrectionState.Physics.ObjectState);
+						}
 					}
 				}
 			}
 		}
-		}
 
 		// Marhsall data back to GT based on what was requested for networking
 		if (!RewindData->IsResim())
@@ -668,12 +672,13 @@ void UNetworkPhysicsManager::OnWorldPostInit(UWorld* World, const UWorld::Initia
 
 			if (ensure(NumFrames > 0))
 			{
-			Solver->EnableRewindCapture(NumFrames, true, MakeUnique<FNetworkPhysicsRewindCallback>());
-			RewindCallback = static_cast<FNetworkPhysicsRewindCallback*>(Solver->GetRewindCallback());
-			RewindCallback->RewindData = Solver->GetRewindData();
-			RewindCallback->Solver = Solver;
-			RewindCallback->World = World;
-			RewindCallback->NetworkPhysicsManager = this;
+				Solver->EnableRewindCapture(NumFrames, true, MakeUnique<FNetworkPhysicsRewindCallback>());
+				RewindCallback = static_cast<FNetworkPhysicsRewindCallback*>(Solver->GetRewindCallback());
+				RewindCallback->RewindData = Solver->GetRewindData();
+				RewindCallback->Solver = Solver;
+				RewindCallback->World = World;
+				RewindCallback->NetworkPhysicsManager = this;
+			}
 		}
 	}
 	}
@@ -741,7 +746,7 @@ void UNetworkPhysicsManager::PostNetRecv()
 
 				if (UE_NETWORK_PHYSICS::TimeDilationEnabled > 0)
 				{
-				PhysScene->SetNetworkDeltaTimeScale(RealTimeDilation);
+					PhysScene->SetNetworkDeltaTimeScale(RealTimeDilation);
 				}
 
 				if (LatestAckdClientFrame != INDEX_NONE)
@@ -841,9 +846,9 @@ void UNetworkPhysicsManager::PostNetRecv()
 
 							if (PhysicsState->OwningActor)
 							{
-							PhysicsState->OwningActor->GetReplicatedMovement_Mutable().bRepPhysics = false;
-							PhysicsState->OwningActor->SetReplicateMovement(false);
-						}
+								PhysicsState->OwningActor->GetReplicatedMovement_Mutable().bRepPhysics = false;
+								PhysicsState->OwningActor->SetReplicateMovement(false);
+							}
 						}
 						else
 						{
@@ -1260,7 +1265,7 @@ void UNetworkPhysicsManager::TickDrawDebug()
 				{
 					(*Func)(P);
 				}
-			}
+			}			
 		}
 	}
 
diff --git a/Engine/Plugins/Runtime/NetworkPrediction/Source/NetworkPrediction/Private/NetworkPredictionSettings.cpp b/Engine/Plugins/Runtime/NetworkPrediction/Source/NetworkPrediction/Private/NetworkPredictionSettings.cpp
--- a/Engine/Plugins/Runtime/NetworkPrediction/Source/NetworkPrediction/Private/NetworkPredictionSettings.cpp
+++ b/Engine/Plugins/Runtime/NetworkPrediction/Source/NetworkPrediction/Private/NetworkPredictionSettings.cpp
@@ -217,35 +217,46 @@ namespace UE_NETWORK_PREDICTION
 		}
 
 		// One time init:
+		TWeakObjectPtr<UWorld> WeakWorld(InWorld);
+		bool bRegisteredInput  = false;
+		auto BindInput = [&bRegisteredInput, WeakWorld]()
+		{
+			if (UWorld* World = WeakWorld.Get())
+			{
+				if (ULocalPlayer* LocalPlayer = World->GetFirstLocalPlayerFromController())
+				{
+					//LocalPlayer->Exec(InWorld, *FString::Printf(TEXT("setbind %d \"np2.DevMenu %d\n"), i, i), *GLog);
+
+					LocalPlayer->Exec(World, TEXT("setbind one \"np2.DevMenu 1\""), *GLog);
+					LocalPlayer->Exec(World, TEXT("setbind two \"np2.DevMenu 2\""), *GLog);
+					LocalPlayer->Exec(World, TEXT("setbind three \"np2.DevMenu 3\""), *GLog);
+					LocalPlayer->Exec(World, TEXT("setbind four \"np2.DevMenu 4\""), *GLog);
+					LocalPlayer->Exec(World, TEXT("setbind five \"np2.DevMenu 5\""), *GLog);
+					LocalPlayer->Exec(World, TEXT("setbind six \"np2.DevMenu 6\""), *GLog);
+					LocalPlayer->Exec(World, TEXT("setbind seven \"np2.DevMenu 7\""), *GLog);
+					LocalPlayer->Exec(World, TEXT("setbind eight \"np2.DevMenu 8\""), *GLog);
+					LocalPlayer->Exec(World, TEXT("setbind nine \"np2.DevMenu 9\""), *GLog);
+					LocalPlayer->Exec(World, TEXT("setbind zero \"np2.DevMenu 0\""), *GLog);
+					bRegisteredInput = true;
+				}
+			}
+		};
+
+		if (!bRegisteredInput)
+		{
+			BindInput();
+		}
+
 		static FDelegateHandle Handle;
 		if (Handle.IsValid())
 		{
 			return;
-		}	
+		}
 
 		Handle = AHUD::OnHUDPostRender.AddLambda([](AHUD* HUD, UCanvas* Canvas)
 		{
 			DrawDevHUD(WorldHandlesMap, HUD, Canvas);
 		});
-
-		//for (int32 i=0; i <= 9; ++i)
-		{
-			if (ULocalPlayer* LocalPlayer = InWorld->GetFirstLocalPlayerFromController())
-			{
-				//LocalPlayer->Exec(InWorld, *FString::Printf(TEXT("setbind %d \"np2.DevMenu %d\n"), i, i), *GLog);
-
-				LocalPlayer->Exec(InWorld, TEXT("setbind one \"np2.DevMenu 1\""), *GLog);
-				LocalPlayer->Exec(InWorld, TEXT("setbind two \"np2.DevMenu 2\""), *GLog);
-				LocalPlayer->Exec(InWorld, TEXT("setbind three \"np2.DevMenu 3\""), *GLog);
-				LocalPlayer->Exec(InWorld, TEXT("setbind four \"np2.DevMenu 4\""), *GLog);
-				LocalPlayer->Exec(InWorld, TEXT("setbind five \"np2.DevMenu 5\""), *GLog);
-				LocalPlayer->Exec(InWorld, TEXT("setbind six \"np2.DevMenu 6\""), *GLog);
-				LocalPlayer->Exec(InWorld, TEXT("setbind seven \"np2.DevMenu 7\""), *GLog);
-				LocalPlayer->Exec(InWorld, TEXT("setbind eight \"np2.DevMenu 8\""), *GLog);
-				LocalPlayer->Exec(InWorld, TEXT("setbind nine \"np2.DevMenu 9\""), *GLog);
-				LocalPlayer->Exec(InWorld, TEXT("setbind zero \"np2.DevMenu 0\""), *GLog);
-			}
-		}
 	}
 }
 
diff --git a/Engine/Plugins/Runtime/NetworkPrediction/Source/NetworkPrediction/Private/NetworkPredictionWorldManager.cpp b/Engine/Plugins/Runtime/NetworkPrediction/Source/NetworkPrediction/Private/NetworkPredictionWorldManager.cpp
--- a/Engine/Plugins/Runtime/NetworkPrediction/Source/NetworkPrediction/Private/NetworkPredictionWorldManager.cpp
+++ b/Engine/Plugins/Runtime/NetworkPrediction/Source/NetworkPrediction/Private/NetworkPredictionWorldManager.cpp
@@ -15,15 +15,6 @@
 // Do extra checks to make sure Physics and GameThread (PrimitiveComponent) are in sync at verious points in the rollback process
 #define NP_ENSURE_PHYSICS_GT_SYNC !(UE_BUILD_SHIPPING || UE_BUILD_TEST)
 
-namespace NetworkPredictionWorldManagerCVars
-{
-	static int32 MaxStepMS = 100;
-	static FAutoConsoleVariableRef CVarMaxSimulationStepMS(
-		TEXT("np.MaxSimulationStepMS"),
-		MaxStepMS,
-		TEXT("The maximum time step (in integer milliseconds) that we will allow for simulation frames. Frame delta times will be clamped between 1 MS and this value.."));
-}
-
 UNetworkPredictionWorldManager* UNetworkPredictionWorldManager::ActiveInstance=nullptr;
 
 // -----------------------------------------------------------------------------------------------
@@ -345,12 +336,13 @@ void UNetworkPredictionWorldManager::BeginNewSimulationFrame(UWorld* InWorld, EL
 	{
 		// Update VariableTickState	
 		constexpr int32 MinStepMS = 1;
+		constexpr int32 MaxStepMS = 100;
 
 		VariableTickState.UnspentTimeMS += fEngineFrameDeltaTimeMS;
 		float fDeltaMS = FMath::FloorToFloat(VariableTickState.UnspentTimeMS);
 		VariableTickState.UnspentTimeMS -= fDeltaMS;
 
-		const int32 DeltaSimMS = FMath::Clamp((int32)fDeltaMS, MinStepMS, NetworkPredictionWorldManagerCVars::MaxStepMS);
+		const int32 DeltaSimMS = FMath::Clamp((int32)fDeltaMS, MinStepMS, MaxStepMS);
 
 		FVariableTickState::FFrame& PendingFrameData = VariableTickState.Frames[VariableTickState.PendingFrame];
 		PendingFrameData.DeltaMS = DeltaSimMS;
diff --git a/Engine/Plugins/Runtime/NetworkPrediction/Source/NetworkPrediction/Public/NetworkPredictionPhysics.h b/Engine/Plugins/Runtime/NetworkPrediction/Source/NetworkPrediction/Public/NetworkPredictionPhysics.h
--- a/Engine/Plugins/Runtime/NetworkPrediction/Source/NetworkPrediction/Public/NetworkPredictionPhysics.h
+++ b/Engine/Plugins/Runtime/NetworkPrediction/Source/NetworkPrediction/Public/NetworkPredictionPhysics.h
@@ -161,12 +161,6 @@ private:
 	}
 };
 
-#else
-
-struct NETWORKPREDICTION_API FNetworkPredictionPhysicsState
-{
-};
-
 #endif // WITH_CHAOS
 
 // ------------------------------------------------------------------------------------------------------
