Network Physics
-Bunch of fixes for handling hitches and substepping. Client and servers can withstand hitches without falling out of sync now.

diff --git a/Engine/Plugins/Runtime/NetworkPrediction/Source/NetworkPrediction/Private/Async/NetworkPredictionAsyncWorldManager.cpp b/Engine/Plugins/Runtime/NetworkPrediction/Source/NetworkPrediction/Private/Async/NetworkPredictionAsyncWorldManager.cpp
--- a/Engine/Plugins/Runtime/NetworkPrediction/Source/NetworkPrediction/Private/Async/NetworkPredictionAsyncWorldManager.cpp
+++ b/Engine/Plugins/Runtime/NetworkPrediction/Source/NetworkPrediction/Private/Async/NetworkPredictionAsyncWorldManager.cpp
@@ -108,10 +108,15 @@ void FNetworkPredictionAsyncWorldManager::NetSerializePlayerControllerInputCmds(
 
 					for (int32 Frame = WriteStart; Frame <= WriteEnd; ++Frame)
 					{
-						//UE_LOG(LogTemp, Log, TEXT("[S] De-NetSerialize Cmd %d Frame %d. NumBytes: %d"), Frame + CmdOffset, Frame, PCInputBuffer.Get(Frame + CmdOffset).Num());
+						//UE_LOG(LogTemp, Log, TEXT("[S] De-NetSerialize Cmd %d Frame %d. NumBytes: %d."), Frame + CmdOffset, Frame, PCInputBuffer.Get(Frame + CmdOffset).Num());
 						NetSerializeCmd(Frame + CmdOffset, Frame);
 					}
 				}
+
+				if (Info.LastFrame < PhysicsStep)
+				{
+					UE_LOG(LogNetworkPrediction, Warning, TEXT("Server does not have valid InputCmd for remote player on PhysicsStep %d (LastFrame: %d)"), PhysicsStep, Info.LastFrame);
+				}
 			}
 		}
 	}
diff --git a/Engine/Plugins/Runtime/NetworkPrediction/Source/NetworkPrediction/Private/NetworkPhysics.cpp b/Engine/Plugins/Runtime/NetworkPrediction/Source/NetworkPrediction/Private/NetworkPhysics.cpp
--- a/Engine/Plugins/Runtime/NetworkPrediction/Source/NetworkPrediction/Private/NetworkPhysics.cpp
+++ b/Engine/Plugins/Runtime/NetworkPrediction/Source/NetworkPrediction/Private/NetworkPhysics.cpp
@@ -25,6 +25,7 @@
 #include "Misc/FileHelper.h"
 #include "ProfilingDebugging/CpuProfilerTrace.h"
 #include "Async/NetworkPredictionAsyncDefines.h"
+#include "PhysicsEngine/PhysicsSettings.h"
 
 DEFINE_LOG_CATEGORY(LogNetworkPhysics);
 
@@ -42,6 +43,8 @@ namespace UE_NETWORK_PHYSICS
 	NP_DEVCVAR_FLOAT(V, 1.0f, "np2.Tolerance.V", "Velocity Tolerance");
 	NP_DEVCVAR_FLOAT(W, 1.0f, "np2.Tolerance.W", "Rotational Velocity Tolerance");
 
+	NP_DEVCVAR_INT(ResimCacheEnabled, 1, "np2.ResimCacheEnabled", "Debug mode for in world drawing");
+
 	NP_DEVCVAR_INT(Debug, 0, "np2.Debug", "Debug mode for in world drawing");
 	NP_DEVCVAR_INT(DebugTolerance, 1, "np2.Debug.Tolerance", "If enabled, only draw large corrections in world.");
 	NP_DEVCVAR_FLOAT(DebugDrawTolerance_X, 5.0, "np2.Debug.Tolerance.X",  "Location Debug Drawing Toleration");
@@ -76,6 +79,7 @@ namespace UE_NETWORK_PHYSICS
 	NP_DEVCVAR_FLOAT(MaxTargetNumBufferedCmds, 5.0, "np2.MaxTargetNumBufferedCmds", "");
 	NP_DEVCVAR_FLOAT(MaxTimeDilationMag, 0.01f, "np2.MaxTimeDilationMag", "Maximum time dilation that client will use to slow down / catch up with server");
 	NP_DEVCVAR_FLOAT(TimeDilationAlpha, 0.1f, "np2.TimeDilationAlpha", "");
+	NP_DEVCVAR_FLOAT(TargetNumBufferedCmdsDeltaOnFault, 1.0f, "np2.TargetNumBufferedCmdsDeltaOnFault", "");
 	NP_DEVCVAR_FLOAT(TargetNumBufferedCmds, 1.9f, "np2.TargetNumBufferedCmds", "");
 	NP_DEVCVAR_FLOAT(TargetNumBufferedCmdsAlpha, 0.005f, "np2.TargetNumBufferedCmdsAlpha", "");
 	NP_DEVCVAR_INT(LerpTargetNumBufferedCmdsAggresively, 0, "np2.LerpTargetNumBufferedCmdsAggresively", "Aggresively lerp towards TargetNumBufferedCmds");
@@ -390,6 +394,11 @@ struct FNetworkPhysicsRewindCallback : public Chaos::IRewindCallback
 
 		return RewindToFrame;
 	}
+
+	void SetResimDebugInfo_Internal(const Chaos::FResimDebugInfo& ResimDebugInfo) override
+	{
+
+	}
 	
 	void ProcessInputs_Internal(int32 PhysicsStep, const TArray<Chaos::FSimCallbackInputAndObject>& SimCallbackInputs) override
 	{
@@ -497,11 +506,6 @@ struct FNetworkPhysicsRewindCallback : public Chaos::IRewindCallback
 		FNetworkPhysicsThreadContext::Get().Update(PhysicsStep - this->LastLocalOffset, this->bIsServer);
 #endif
 
-		for (auto& SubSys : NetworkPhysicsManager->SubSystems)
-		{
-			SubSys.Value->ProcessInputs_Internal(PhysicsStep);
-		}
-		
 		for (auto CallbacKInput : SimCallbackInputs)
 		{
 			for (Chaos::ISimCallbackObject* Callback : SimObjectCallbacks)
@@ -536,10 +540,10 @@ struct FNetworkPhysicsRewindCallback : public Chaos::IRewindCallback
 	{
 	}
 
-	void ProcessInputs_External(int32 PhysicsStep, const TArray<Chaos::FSimCallbackInputAndObject>& SimCallbackInputs) override 
+	void InjectInputs_External(int32 PhysicsStep, int32 NumSteps) override
 	{
 		npCheckSlow(NetworkPhysicsManager);
-		NetworkPhysicsManager->ProcessInputs_External(PhysicsStep, SimCallbackInputs);
+		NetworkPhysicsManager->InjectInputs_External(PhysicsStep, NumSteps);
 	}
 
 	void RegisterRewindableSimCallback_Internal(Chaos::ISimCallbackObject* Callback) override
@@ -690,7 +694,23 @@ void UNetworkPhysicsManager::OnWorldPostInit(UWorld* World, const UWorld::Initia
 
 			if (ensure(NumFrames > 0))
 			{
-				Solver->EnableRewindCapture(NumFrames, true, MakeUnique<FNetworkPhysicsRewindCallback>());
+				bool bEnableResimCache = true;
+				if (UE_NETWORK_PHYSICS::ResimCacheEnabled == 0)
+				{
+					bEnableResimCache = false;
+				}
+				else if (UE_NETWORK_PHYSICS::ResimCacheEnabled == 2)
+				{
+					static int32 InitCount = 0;
+					if (World->GetName().Contains(TEXT("Untitled")) == false)
+					{
+						bEnableResimCache = (++InitCount % 2 == 0);
+					}
+				}
+
+				UE_LOG(LogNetworkPhysics, Warning, TEXT("UNetworkPhysicsManager EnablingRewindCapture. bEnableResimCache: %d"), bEnableResimCache);
+				Solver->EnableRewindCapture(NumFrames, bEnableResimCache, MakeUnique<FNetworkPhysicsRewindCallback>());
+				
 				RewindCallback = static_cast<FNetworkPhysicsRewindCallback*>(Solver->GetRewindCallback());
 				RewindCallback->RewindData = Solver->GetRewindData();
 				RewindCallback->Solver = Solver;
@@ -1025,9 +1045,9 @@ void UNetworkPhysicsManager::PreNetSend(float DeltaSeconds)
 #endif
 }
 
-void UNetworkPhysicsManager::ProcessInputs_External(int32 PhysicsStep, const TArray<Chaos::FSimCallbackInputAndObject>& SimCallbackInputs)
+void UNetworkPhysicsManager::InjectInputs_External(int32 PhysicsStep, int32 NumSteps)
 {
-	TRACE_CPUPROFILER_EVENT_SCOPE(NPA_Physics_ProcessInputs_External);
+	TRACE_CPUPROFILER_EVENT_SCOPE(NPA_Physics_InjectInputs_External);
 
 	UWorld* World = GetWorld();
 	const ENetMode NetMode = World->GetNetMode();
@@ -1049,7 +1069,10 @@ void UNetworkPhysicsManager::ProcessInputs_External(int32 PhysicsStep, const TAr
 	
 	if (NetMode != NM_Client)
 	{
-		ProcessClientInputBuffers_External(PhysicsStep, SimCallbackInputs);
+		for (int32 StepNum=0; StepNum < NumSteps; ++StepNum)
+		{
+			ProcessClientInputBuffers_External(PhysicsStep + StepNum);
+		}
 	}
 
 	// ----------------------------------------------------------------------
@@ -1059,7 +1082,7 @@ void UNetworkPhysicsManager::ProcessInputs_External(int32 PhysicsStep, const TAr
 	bool bOutSentClientInput = false;
 	for (auto& Pair : SubSystems)
 	{
-		Pair.Value->ProcessInputs_External(PhysicsStep, LocalOffset, bOutSentClientInput);
+		Pair.Value->InjectInputs_External(PhysicsStep, NumSteps, LocalOffset, bOutSentClientInput);
 	}
 
 	// ----------------------------------------------------------------------
@@ -1073,40 +1096,12 @@ void UNetworkPhysicsManager::ProcessInputs_External(int32 PhysicsStep, const TAr
 		if (APlayerController* PC = World->GetFirstPlayerController())
 		{
 			TArray<uint8> SendData;
-
-			// ----------------------------------------------------------------------------------------------
-			// REMOVE Once Mock/Hardcoded example is gone
-			// ----------------------------------------------------------------------------------------------
-			{
-				FNetBitWriter Writer(nullptr, 256 << 3);
-
-				for (Chaos::FSimCallbackInputAndObject Pair :  SimCallbackInputs)
-				{
-					if (Pair.Input->NetSendInputCmd(Writer))
-					{
-						break;
-					}
-				}
-
-			
-			
-				if (Writer.IsError() == false)
-				{
-					int32 NumBytes = (int32)Writer.GetNumBytes();
-					SendData = MoveTemp(*const_cast<TArray<uint8>*>(Writer.GetBuffer()));
-					SendData.SetNum(NumBytes);
-				}
-			}
-
-			// ----------------------------------------------------------------------------------------------
-			// ----------------------------------------------------------------------------------------------
-			// ----------------------------------------------------------------------------------------------
 			PC->PushClientInput(PhysicsStep, SendData);
 		}
 	}
 }
 
-void UNetworkPhysicsManager::ProcessClientInputBuffers_External(int32 PhysicsStep, const TArray<Chaos::FSimCallbackInputAndObject>& SimCallbackInputs)
+void UNetworkPhysicsManager::ProcessClientInputBuffers_External(int32 PhysicsStep)
 {
 	UWorld* World = GetWorld();
 	constexpr int32 MaxBufferedCmds = 16;
@@ -1148,7 +1143,7 @@ void UNetworkPhysicsManager::ProcessClientInputBuffers_External(int32 PhysicsSte
 					// No Cmds to process, enter fault state. Increment TargetNumBufferedCmds each time this happens.
 					// TODO: We should have something to bring this back down (which means skipping frames) we don't want temporary poor conditions to cause permanent high input buffering
 					FrameInfo.bFault = true;
-					FrameInfo.TargetNumBufferedCmds = FMath::Min(FrameInfo.TargetNumBufferedCmds+1.f, UE_NETWORK_PHYSICS::MaxTargetNumBufferedCmds);
+					FrameInfo.TargetNumBufferedCmds = FMath::Min(FrameInfo.TargetNumBufferedCmds + UE_NETWORK_PHYSICS::TargetNumBufferedCmdsDeltaOnFault, UE_NETWORK_PHYSICS::MaxTargetNumBufferedCmds);
 
 					UE_CLOG(FrameInfo.LastProcessedInputFrame != INDEX_NONE, LogNetworkPhysics, Warning, TEXT("[Remote.Input] ENTERING fault. New Target: %.2f. (Client) Input: %d. (Server) Local Frame:"), FrameInfo.TargetNumBufferedCmds, FrameInfo.LastProcessedInputFrame, FrameInfo.LastLocalFrame);
 					return;
@@ -1184,31 +1179,6 @@ void UNetworkPhysicsManager::ProcessClientInputBuffers_External(int32 PhysicsSte
 			};
 
 			UpdateLastProcessedInputFrame();
-
-
-			// ----------------------------------------------------------------------------------------------
-			// REMOVE Once Mock/Hardcoded example is gone
-			// ----------------------------------------------------------------------------------------------
-			if (FrameInfo.LastProcessedInputFrame != INDEX_NONE)
-			{
-				const TArray<uint8>& Data = InputBuffer.Get(FrameInfo.LastProcessedInputFrame);
-				if (Data.Num() > 0)
-				{
-					uint8* RawData = const_cast<uint8*>(Data.GetData());
-					FNetBitReader Ar(nullptr, RawData, ((int64)Data.Num()) << 3);
-						
-					for (Chaos::FSimCallbackInputAndObject Pair :  SimCallbackInputs)
-					{
-						if (Pair.Input->NetRecvInputCmd(PC, Ar))
-						{
-							break;
-						}
-					}
-				}
-			}
-			// ----------------------------------------------------------------------------------------------
-			// ----------------------------------------------------------------------------------------------
-			// ----------------------------------------------------------------------------------------------
 		}
 	}
 }
@@ -1616,4 +1586,24 @@ FAutoConsoleCommandWithWorldAndArgs NpFloatTestCmd(TEXT("np2.FloatTest"), TEXT("
 		UE_LOG(LogTemp, Log, TEXT("%.4f -> %d -> %.4f (%.4f error"), F, Compressed, UnCompressed, UnCompressed - F);
 	}
 
+}));
+
+
+FAutoConsoleCommandWithWorldAndArgs NpSetMaxSubStepsCmd(TEXT("np2.SetMaxSubsteps"), TEXT(""),
+	FConsoleCommandWithWorldAndArgsDelegate::CreateLambda([](const TArray< FString >& Args, UWorld* InWorld) 
+{
+	int32 NewMaxSubSteps = 0;
+	if (Args.Num() > 0)
+	{
+		LexFromString(NewMaxSubSteps, *Args[0]);
+	}
+
+	if (NewMaxSubSteps <= 0)
+	{
+		return;
+	}
+
+	UE_LOG(LogNetworkPhysics, Log, TEXT("NewMaxSubSteps: %d"), NewMaxSubSteps);
+	GetMutableDefault<UPhysicsSettings>()->MaxSubsteps = NewMaxSubSteps;
+	
 }));
\ No newline at end of file
diff --git a/Engine/Plugins/Runtime/NetworkPrediction/Source/NetworkPrediction/Public/Async/NetworkPredictionAsyncDefines.h b/Engine/Plugins/Runtime/NetworkPrediction/Source/NetworkPrediction/Public/Async/NetworkPredictionAsyncDefines.h
--- a/Engine/Plugins/Runtime/NetworkPrediction/Source/NetworkPrediction/Public/Async/NetworkPredictionAsyncDefines.h
+++ b/Engine/Plugins/Runtime/NetworkPrediction/Source/NetworkPrediction/Public/Async/NetworkPredictionAsyncDefines.h
@@ -28,4 +28,10 @@ const int32 FrameStorageGrowth = NP_FRAME_STORAGE_GROWTH;
 
 const int32 InputCmdBufferSize = NP_FRAME_INPUTCMD_BUFFER_SIZE;
 
+#ifndef NP_INLINE_SIMOBJ_INPUTS
+	#define NP_INLINE_SIMOBJ_INPUTS 3
+#endif
+
+const int32 InlineSimObjInputs = NP_INLINE_SIMOBJ_INPUTS;
+
 };
\ No newline at end of file
diff --git a/Engine/Plugins/Runtime/NetworkPrediction/Source/NetworkPrediction/Public/Async/NetworkPredictionAsyncSimCallback.h b/Engine/Plugins/Runtime/NetworkPrediction/Source/NetworkPrediction/Public/Async/NetworkPredictionAsyncSimCallback.h
--- a/Engine/Plugins/Runtime/NetworkPrediction/Source/NetworkPrediction/Public/Async/NetworkPredictionAsyncSimCallback.h
+++ b/Engine/Plugins/Runtime/NetworkPrediction/Source/NetworkPrediction/Public/Async/NetworkPredictionAsyncSimCallback.h
@@ -7,6 +7,22 @@
 
 namespace UE_NP {
 
+// This is the Chaos's Simcallback Input object for the NP system.
+// It holds data that is marshaled GT -> PT each time a PhysicsStep task is kicked off.
+// The data in here must be associated with an absolute (local) frame number.
+
+// This is because of "substepping"
+//	In this context, by substepping we mean multiple PhysicsStep tasks are kicked off in a single GT tick
+//	(e.g, we hitched, had a GT delta time of 99ms and kicked off three PhysicsStep tasks at 33ms steps each)
+//
+//	When this happens, the Chaos::FSimCallbackInput is shared between all PT tasks. 
+//	In other ords, the FSimCallbackInput contains all marshaled data for N physics steps. There is not a unique FSimCallbackInput per PhysicsStep when substepping.
+//	(and this makes sense in the context of the GT. The GT has no "chance" to add data to the N+1 steps, since they are all kicked off at the same time).
+//	But for Networking, we need to coordinate steps and InputCmds between clients and Server. Clients and servers will not substep "together". We need the behavior
+//	of the simulation to be exactly the same whether we run 1 PhysicsStep at a time or whether we substep.
+//
+//	Therefore the data in here is stored in the FDataStoreInternal structure and keyed off the local frame.
+
 struct FNetworkPredictionSimCallbackInput : public Chaos::FSimCallbackInput
 {
 	bool NetSendInputCmd(FNetBitWriter& Ar) override { return true; } // FIXME: remove these functions	
@@ -14,24 +30,94 @@ struct FNetworkPredictionSimCallbackInput : public Chaos::FSimCallbackInput
 
 	void Reset()
 	{
-		DataStoreCollection.Reset();
+		for (FDataStoreInternal& Item : DataStoreCollections)
+		{
+			Item.Frame = INDEX_NONE;
+			Item.StoreCollection.Reset();
+		}
+	}
+
+	// Get the next available data store for writing. This data store doesn't have a frame assigned to it yet.
+	template<typename AsyncModelDef>
+	TAsyncModelDataStore_Input<AsyncModelDef>* GetPendingDataStore_External()
+	{
+		for (FDataStoreInternal& Item : DataStoreCollections)
+		{
+			if (Item.Frame == INDEX_NONE)
+			{
+				return Item.StoreCollection.GetDataStore<TAsyncModelDataStore_Input<AsyncModelDef>, AsyncModelDef>();
+			}
+		}
+		
+		return DataStoreCollections.AddDefaulted_GetRef().StoreCollection.GetDataStore<TAsyncModelDataStore_Input<AsyncModelDef>, AsyncModelDef>();
+	}
+
+	// Get the next available data store and assign a final frame number to it.	
+	void FinalizePendingDataStore_External(int32 Frame)
+	{
+		FDataStoreInternal* Found = nullptr;
+		for (FDataStoreInternal& Item : DataStoreCollections)
+		{
+			if (Item.Frame == INDEX_NONE)
+			{
+				Found = &Item;
+				break;
+			}
+		}
+
+		if (!Found)
+		{
+			Found = &DataStoreCollections.AddDefaulted_GetRef();
+		}
+
+		npCheckSlow(Found);
+		Found->Frame = Frame;
 	}
 
+	// Get data store for reading on the PT, by frame number
 	template<typename AsyncModelDef>
-	TAsyncModelDataStore_Input<AsyncModelDef>* GetDataStore()
+	const TAsyncModelDataStore_Input<AsyncModelDef>* GetDataStoreChecked(int32 Frame) const
 	{
-		return DataStoreCollection.GetDataStore<TAsyncModelDataStore_Input<AsyncModelDef>, AsyncModelDef>();
+		return const_cast<FNetworkPredictionSimCallbackInput*>(this)->GetDataStoreCheckedImpl<AsyncModelDef>(Frame);
 	}
 
 	template<typename AsyncModelDef>
-	const TAsyncModelDataStore_Input<AsyncModelDef>* GetDataStore() const
+	TAsyncModelDataStore_Input<AsyncModelDef>* GetDataStoreChecked(int32 Frame)
 	{
-		return DataStoreCollection.GetDataStore<TAsyncModelDataStore_Input<AsyncModelDef>, AsyncModelDef>();
+		return GetDataStoreCheckedImpl<AsyncModelDef>(Frame);
 	}
 
 private:
 
-	FAsyncDataStoreCollection DataStoreCollection;
+	struct FDataStoreInternal
+	{
+		int32 Frame = INDEX_NONE;
+		FAsyncDataStoreCollection StoreCollection;
+	};
+
+
+	TArray<FDataStoreInternal, TInlineAllocator<InlineSimObjInputs>> DataStoreCollections;
+
+
+	template<typename AsyncModelDef>
+	TAsyncModelDataStore_Input<AsyncModelDef>* GetDataStoreCheckedImpl(int32 Frame)
+	{
+		for (FDataStoreInternal& Item : DataStoreCollections)
+		{
+			if (Item.Frame == Frame)
+			{
+				return Item.StoreCollection.GetDataStore<TAsyncModelDataStore_Input<AsyncModelDef>, AsyncModelDef>();
+			}
+		}
+		
+		for (FDataStoreInternal& Item : DataStoreCollections)
+		{
+			UE_LOG(LogNetworkPrediction, Warning, TEXT("   Found: %d"), Item.Frame);
+		}
+
+		npEnsureMsgf(false, TEXT("Could not find DataStore asigned to frame %d. 0x%X"), Frame, (int64)this);
+		return nullptr;
+	}
 };
 
 
diff --git a/Engine/Plugins/Runtime/NetworkPrediction/Source/NetworkPrediction/Public/Async/NetworkPredictionAsyncWorldManager.h b/Engine/Plugins/Runtime/NetworkPrediction/Source/NetworkPrediction/Public/Async/NetworkPredictionAsyncWorldManager.h
--- a/Engine/Plugins/Runtime/NetworkPrediction/Source/NetworkPrediction/Public/Async/NetworkPredictionAsyncWorldManager.h
+++ b/Engine/Plugins/Runtime/NetworkPrediction/Source/NetworkPrediction/Public/Async/NetworkPredictionAsyncWorldManager.h
@@ -44,6 +44,10 @@ public:
 		}
 
 		Chaos::FPhysicsSolver* PhysicsSolver = static_cast<Chaos::FPhysicsSolver*>(this->GetSolver());
+		if (PhysicsSolver->IsShuttingDown())
+		{
+			return;
+		}
 
 		const int32 LocalFrame = PhysicsSolver->GetCurrentFrame();
 		const int32 SimulationFrame = LocalFrame - LocalFrameOffset;
@@ -313,17 +317,24 @@ public:
 			if (ControllerInfo && ControllerInfo->InputSource == ENetworkPredictionAsyncInputSource::Buffered)
 			{
 				const int32 SignedFutureDelta = ControllerInfo->LastFrame - GlobalOutput_External.SimulationFrame;
-				if (SignedFutureDelta > 0 && npEnsure(SignedFutureDelta < 16))
+				if (SignedFutureDelta > 0)
 				{
-					FutureDelta = (uint8)SignedFutureDelta;
-					FutureFrame = ControllerInfo->LastFrame;
-					PendingInputFrame = FutureFrame; // This is telling the NetSend call below to send this frame from the PendingInputCmd
+					if (SignedFutureDelta < InputCmdBufferSize)
+					{
+						FutureDelta = (uint8)SignedFutureDelta;
+						FutureFrame = ControllerInfo->LastFrame;
+						PendingInputFrame = FutureFrame; // This is telling the NetSend call below to send this frame from the PendingInputCmd
+					}
+					else
+					{
+						UE_LOG(LogNetworkPrediction, Warning, TEXT("Too many buffered input cmds detected for client. LastFrame: %d SimulationFrame: %d. (Delta: %d)"), ControllerInfo->LastFrame, GlobalOutput_External.SimulationFrame, SignedFutureDelta);
+					}
 				}
 			}
 			else
 			{
 				const int32 SignedFutureDelta = this->NextPhysicsStep - GlobalOutput_External.SimulationFrame;
-				if (npEnsure(SignedFutureDelta > 0) && npEnsure(SignedFutureDelta < 16))
+				if (npEnsure(SignedFutureDelta > 0) && npEnsure(SignedFutureDelta < InputCmdBufferSize))
 				{
 					FutureDelta = (uint8)SignedFutureDelta;
 				}
@@ -404,44 +415,53 @@ public:
 	{
 	}
 
-	void ProcessInputs_External(int32 PhysicsStep, int32 LocalFrameOffset, bool& bOutSendClientInputCmd) override
+	void InjectInputs_External(int32 PhysicsStep, int32 NumSteps, int32 LocalFrameOffset, bool& bOutSendClientInputCmd) override
 	{
 		TRACE_CPUPROFILER_EVENT_SCOPE(NPA_ProcessInputs_External);
 
 		npEnsureMsgf(NextPhysicsStep == 0 || (NextPhysicsStep == PhysicsStep), TEXT("Unexpected PhysicsStep %d (NextPhysicsStep: %d)"), PhysicsStep, NextPhysicsStep);
-
-		NextPhysicsStep = PhysicsStep+1;
+		NextPhysicsStep = PhysicsStep;
 
 		// Prep the next AsyncInput
 		npCheckSlow(AsyncCallback);
 		FNetworkPredictionSimCallbackInput* AsyncInput = AsyncCallback->GetProducerInputData_External();
 		npCheckSlow(AsyncInput);
-		
+
 		AsyncCallback->PushLatestLocalFrameOffset(LocalFrameOffset);
-		
-		// Sync Networked InputCmds
-		NetSerializePlayerControllerInputCmds(PhysicsStep, [&](FNetworkPredictionAsyncID ID, int32 PendingInputBufferFrame, FArchive& Ar)
+
+		for (int32 StepNum=0; StepNum < NumSteps; ++StepNum)
 		{
-			FGlobalInstanceData& Data = GlobalInstanceData_External.FindChecked(ID);
-			ForEachSim<INetSerializeService>(Data, NetSerializeServices, [&](INetSerializeService* Service)
+			// Whatever pending data is in the AsyncInput is now assigned to this PhysicsStep
+			AsyncInput->FinalizePendingDataStore_External(PhysicsStep);
+		
+			// Sync Networked InputCmds
+			NetSerializePlayerControllerInputCmds(PhysicsStep, [&](FNetworkPredictionAsyncID ID, int32 PendingInputBufferFrame, FArchive& Ar)
 			{
-				Service->NetSerializeInputCmd(ID, PendingInputBufferFrame, Ar);
-				if (Ar.IsSaving())
+				FGlobalInstanceData& Data = GlobalInstanceData_External.FindChecked(ID);
+				ForEachSim<INetSerializeService>(Data, NetSerializeServices, [&](INetSerializeService* Service)
 				{
-					bOutSendClientInputCmd = true;
-				}
+					Service->NetSerializeInputCmd(ID, PendingInputBufferFrame, Ar);
+					if (Ar.IsSaving())
+					{
+						bOutSendClientInputCmd = true;
+					}
+				});
 			});
-		});
 
-		// Marshal data from GT to PT input
-		for (TUniquePtr<IAsyncMarshalService_External>& ServicePtr : MarshalServices.Array)
-		{
-			if (IAsyncMarshalService_External* Service = ServicePtr.Get())
+			// Marshal data from GT to PT input
+			for (TUniquePtr<IAsyncMarshalService_External>& ServicePtr : MarshalServices.Array)
 			{
-				Service->MarshalInput(PhysicsStep, AsyncInput);
+				if (IAsyncMarshalService_External* Service = ServicePtr.Get())
+				{
+					Service->MarshalInput(PhysicsStep, AsyncInput);
+				}
 			}
+
+			PhysicsStep++;
 		}
 
+		NextPhysicsStep = PhysicsStep;
+
 		// Consume all output data. Note we are consuming all SimCallback Output objects here even though we really only want the latest -
 		// but there maybe objects that were destroyed and not present in the final output, so this is trying to not miss them.
 		while (Chaos::TSimCallbackOutputHandle<FNetworkPredictionSimCallbackOutput> AsyncOutput = AsyncCallback->PopFutureOutputData_External())
@@ -464,11 +484,6 @@ public:
 		return AsyncCallback->Reconcile_Internal(LastCompletedStep);
 	}
 
-	void ProcessInputs_Internal(int32 PhysicsStep) override
-	{
-
-	}
-
 	void PreResimStep_Internal(int32 PhysicsStep, bool bFirst) override
 	{
 	}
diff --git a/Engine/Plugins/Runtime/NetworkPrediction/Source/NetworkPrediction/Public/Async/Services/NetworkPredictionAsyncService_Marshal.inl b/Engine/Plugins/Runtime/NetworkPrediction/Source/NetworkPrediction/Public/Async/Services/NetworkPredictionAsyncService_Marshal.inl
--- a/Engine/Plugins/Runtime/NetworkPrediction/Source/NetworkPrediction/Public/Async/Services/NetworkPredictionAsyncService_Marshal.inl
+++ b/Engine/Plugins/Runtime/NetworkPrediction/Source/NetworkPrediction/Public/Async/Services/NetworkPredictionAsyncService_Marshal.inl
@@ -80,7 +80,7 @@ public:
 	{
 		TRACE_CPUPROFILER_EVENT_SCOPE(NPA_MarshalPreSimulate);
 
-		const TAsyncModelDataStore_Input<AsyncModelDef>* InputData = SimInput->GetDataStore<AsyncModelDef>();
+		const TAsyncModelDataStore_Input<AsyncModelDef>* InputData = SimInput->GetDataStoreChecked<AsyncModelDef>(LocalFrame);
 		TAsncFrameSnapshot<AsyncModelDef>& Snapshot = DataStore->Frames[LocalFrame];
 
 		/*
@@ -98,39 +98,44 @@ public:
 		}
 		*/
 
-		// New instances
-		for (const typename TAsyncModelDataStore_Input<AsyncModelDef>::FNewInstance& NewInstance : InputData->NewInstances)
+		if (InputData)
 		{
-			DataStore->Instances.FindOrAdd(NewInstance.ID) = NewInstance.StaticData;
+			// New instances
+			for (const typename TAsyncModelDataStore_Input<AsyncModelDef>::FNewInstance& NewInstance : InputData->NewInstances)
+			{
+				DataStore->Instances.FindOrAdd(NewInstance.ID) = NewInstance.StaticData;
 
-			NpResizeForIndex(Snapshot.InputCmds, NewInstance.StaticData.Index);
-			NpResizeForIndex(Snapshot.NetStates, NewInstance.StaticData.Index);
+				NpResizeForIndex(Snapshot.InputCmds, NewInstance.StaticData.Index);
+				NpResizeForIndex(Snapshot.NetStates, NewInstance.StaticData.Index);
 
-			Snapshot.NetStates[NewInstance.StaticData.Index] = NewInstance.NetState;
-		}
+				Snapshot.NetStates[NewInstance.StaticData.Index] = NewInstance.NetState;
+			}
 
-		// State Mods
-		for (const TAsyncLocalStateMod<AsyncModelDef>& LocalMod : InputData->LocalStateMods)
-		{
-			LocalMod.Func(DataStore->Instances.FindChecked(LocalMod.ID).LocalState);
-		}
+			// State Mods
+			for (const TAsyncLocalStateMod<AsyncModelDef>& LocalMod : InputData->LocalStateMods)
+			{
+				LocalMod.Func(DataStore->Instances.FindChecked(LocalMod.ID).LocalState);
+			}
 
-		for (const TAsyncNetStateMod<AsyncModelDef>& NetMod : InputData->NetStateMods)
-		{
-			NetMod.Func(Snapshot.NetStates[NetMod.Index]);
-		}
+			for (const TAsyncNetStateMod<AsyncModelDef>& NetMod : InputData->NetStateMods)
+			{
+				NetMod.Func(Snapshot.NetStates[NetMod.Index]);
+			}
 
-		// New Local InputCmds
-		for (const typename TAsyncModelDataStore_Input<AsyncModelDef>::FLocalInputCmd& LocalInput : InputData->LocalInputCmds)
-		{
-			Snapshot.InputCmds[LocalInput.Index] = LocalInput.InputCmd;
-		}
+			// New Local InputCmds
+			for (const typename TAsyncModelDataStore_Input<AsyncModelDef>::FLocalInputCmd& LocalInput : InputData->LocalInputCmds)
+			{
+				//UE_LOG(LogTemp, Warning, TEXT("[%d] Presimulate localInputcmd %s. SimInput: 0x%X"), LocalFrame, *LocalInput.InputCmd.ToString(), (int64)SimInput);
 
-		// Deletes
-		for (const FNetworkPredictionAsyncID& ID : InputData->DeletedInstances)
-		{
-			// I'm not sure if this is enough. Maybe we should track the 'kill frame' on the PT rather than removing anything from data structures
-			DataStore->Instances.Remove(ID);
+				Snapshot.InputCmds[LocalInput.Index] = LocalInput.InputCmd;
+			}
+
+			// Deletes
+			for (const FNetworkPredictionAsyncID& ID : InputData->DeletedInstances)
+			{
+				// I'm not sure if this is enough. Maybe we should track the 'kill frame' on the PT rather than removing anything from data structures
+				DataStore->Instances.Remove(ID);
+			}
 		}
 
 		// Apply Corrections
@@ -155,12 +160,15 @@ public:
 			// But even in resim we have to copy the AP input cmds forward due to how we are synchronizing with the GT and correction above.
 			// Using the Input's LocalInputCmds is a bit abusive but is the simplest approach since the data is right there. See notes at 
 			// the top of this header
-			for (const typename TAsyncModelDataStore_Input<AsyncModelDef>::FLocalInputCmd& LocalInput : InputData->LocalInputCmds)
+			if (InputData)
 			{
-				const int32 idx = LocalInput.Index;
-				if (Snapshot.InputCmds.IsValidIndex(idx) && NextSnapshot.InputCmds.IsValidIndex(idx))
+				for (const typename TAsyncModelDataStore_Input<AsyncModelDef>::FLocalInputCmd& LocalInput : InputData->LocalInputCmds)
 				{
-					NextSnapshot.InputCmds[idx] = Snapshot.InputCmds[idx];
+					const int32 idx = LocalInput.Index;
+					if (Snapshot.InputCmds.IsValidIndex(idx) && NextSnapshot.InputCmds.IsValidIndex(idx))
+					{
+						NextSnapshot.InputCmds[idx] = Snapshot.InputCmds[idx];
+					}
 				}
 			}
 		}
@@ -223,22 +231,25 @@ public:
 	// Marshal external data to SimInput
 	void MarshalInput(const int32 Frame, FNetworkPredictionSimCallbackInput* SimInput) final override
 	{
-		TAsyncModelDataStore_Input<AsyncModelDef>* InputData = SimInput->GetDataStore<AsyncModelDef>();
+		TAsyncModelDataStore_Input<AsyncModelDef>* InputData = SimInput->GetDataStoreChecked<AsyncModelDef>(Frame);
 
 		// ---------------------------------------------------------------------
 		//	InputCmds
 		// ---------------------------------------------------------------------
 
-		for (TAsyncPendingInputCmdPtr<AsyncModelDef>& Info : DataStore_External->ActivePendingInputCmds)
+		if (InputData)
 		{
-			npCheckSlow(Info.PendingInputCmd != nullptr);
-			InputData->LocalInputCmds.Emplace(Info.Index, *Info.PendingInputCmd);
-		}
+			for (TAsyncPendingInputCmdPtr<AsyncModelDef>& Info : DataStore_External->ActivePendingInputCmds)
+			{
+				npCheckSlow(Info.PendingInputCmd != nullptr);
+				InputData->LocalInputCmds.Emplace(Info.Index, *Info.PendingInputCmd);
+			}
 
-		for (auto& MapIt : DataStore_External->PendingInputCmdBuffers)
-		{
-			TAsyncPendingInputCmdBuffer<AsyncModelDef>& PendingInputs = MapIt.Value.Get();
-			InputData->LocalInputCmds.Emplace(PendingInputs.Index, PendingInputs.Buffer[Frame % PendingInputs.Buffer.Num()]);
+			for (auto& MapIt : DataStore_External->PendingInputCmdBuffers)
+			{
+				TAsyncPendingInputCmdBuffer<AsyncModelDef>& PendingInputs = MapIt.Value.Get();
+				InputData->LocalInputCmds.Emplace(PendingInputs.Index, PendingInputs.Buffer[Frame % PendingInputs.Buffer.Num()]);
+			}
 		}
 
 		// ---------------------------------------------------------------------
@@ -284,7 +295,7 @@ public:
 		}
 		else
 		{
-			AsyncInput->GetDataStore<AsyncModelDef>()->LocalStateMods.Emplace(ID, MoveTemp(Func));
+			AsyncInput->GetPendingDataStore_External<AsyncModelDef>()->LocalStateMods.Emplace(ID, MoveTemp(Func));
 		}
 	}
 	
@@ -300,7 +311,7 @@ public:
 
 		npCheckSlow(AsyncInput);
 		const int32 idx = DataStore_External->Instances.FindChecked(ID).Index;
-		AsyncInput->GetDataStore<AsyncModelDef>()->NetStateMods.Emplace(idx, MoveTemp(Func));
+		AsyncInput->GetPendingDataStore_External<AsyncModelDef>()->NetStateMods.Emplace(idx, MoveTemp(Func));
 	}
 
 private:
diff --git a/Engine/Plugins/Runtime/NetworkPrediction/Source/NetworkPrediction/Public/Async/Services/NetworkPredictionAsyncService_Registration.inl b/Engine/Plugins/Runtime/NetworkPrediction/Source/NetworkPrediction/Public/Async/Services/NetworkPredictionAsyncService_Registration.inl
--- a/Engine/Plugins/Runtime/NetworkPrediction/Source/NetworkPrediction/Public/Async/Services/NetworkPredictionAsyncService_Registration.inl
+++ b/Engine/Plugins/Runtime/NetworkPrediction/Source/NetworkPrediction/Public/Async/Services/NetworkPredictionAsyncService_Registration.inl
@@ -65,7 +65,7 @@ public:
 				DataStore->InActivePendingInputCmds.RemoveAtSwap(idx, 1, false);
 			}
 
-			TAsyncModelDataStore_Input<AsyncModelDef>* InputDataStore = Input->GetDataStore<AsyncModelDef>();
+			TAsyncModelDataStore_Input<AsyncModelDef>* InputDataStore = Input->GetPendingDataStore_External<AsyncModelDef>();
 			InputDataStore->DeletedInstances.Add(ID);
 		}
 	}
@@ -95,7 +95,7 @@ public:
 		}
 
 		// Remap any deferred local state mods
-		TAsyncModelDataStore_Input<AsyncModelDef>* InputData = AsyncInput->GetDataStore<AsyncModelDef>();
+		TAsyncModelDataStore_Input<AsyncModelDef>* InputData = AsyncInput->GetPendingDataStore_External<AsyncModelDef>();
 		npCheckSlow(InputData);
 
 		// note: going in reverse would be better since we could use RemoveAtSwap but preserving order that mods are applied is important and that would reverse them
@@ -150,7 +150,7 @@ private:
 		NpResizeForIndex(DataStore->LatestSnapshot.InputCmds, FreeIndice);
 
 		// Marshal via Input
-		TAsyncModelDataStore_Input<AsyncModelDef>* InputDataStore = Input->GetDataStore<AsyncModelDef>();
+		TAsyncModelDataStore_Input<AsyncModelDef>* InputDataStore = Input->GetPendingDataStore_External<AsyncModelDef>();
 		npCheckSlow(InputDataStore);
 
 		InputDataStore->NewInstances.Add({ID, InstanceData, MoveTemp(InitialValueNet)});
diff --git a/Engine/Plugins/Runtime/NetworkPrediction/Source/NetworkPrediction/Public/NetworkPhysics.h b/Engine/Plugins/Runtime/NetworkPrediction/Source/NetworkPrediction/Public/NetworkPhysics.h
--- a/Engine/Plugins/Runtime/NetworkPrediction/Source/NetworkPrediction/Public/NetworkPhysics.h
+++ b/Engine/Plugins/Runtime/NetworkPrediction/Source/NetworkPrediction/Public/NetworkPhysics.h
@@ -133,11 +133,10 @@ public:
 	virtual void PostNetRecv(UWorld* World, int32 LocalOffset, int32 LastProcessedFrame) = 0;
 	virtual void PreNetSend(UWorld* World, float DeltaSeconds) = 0;
 
-	// Finalizes InputCmds, marshalls them to PT and Network.
-	virtual void ProcessInputs_External(int32 PhysicsStep, int32 LocalFrameOffset, bool& bOutSendClientInputCmd) { }
-	
-	// Records "Final" Inputs for a frame and marshalls them back to GT for networking
-	virtual void ProcessInputs_Internal(int32 PhysicsStep) { }
+	// Finalizes InputCmds, marshalls them to PT and Network. TODO: Remove
+	virtual void ProcessInputs_External(int32 PhysicsStep, int32 LocalFrameOffset, bool& bOutSendClientInputCmd) { }	
+
+	virtual void InjectInputs_External(int32 PhysicsStep, int32 NumSteps, int32 LocalFrameOffset, bool& bOutSendClientInputCmd) { }
 
 	// ShouldReconcile
 	virtual int32 TriggerRewindIfNeeded_Internal(int32 LastCompletedStep) { return INDEX_NONE; }
@@ -215,9 +214,9 @@ private:
 	TArray<TPair<FName, TUniquePtr<INetworkPhysicsSubsystem>>> SubSystems;
 
 	friend struct FNetworkPhysicsRewindCallback;
-	void ProcessInputs_External(int32 PhysicsStep, const TArray<Chaos::FSimCallbackInputAndObject>& SimCallbackInputs);
+	void InjectInputs_External(int32 PhysicsStep, int32 NumSteps);
 
-	void ProcessClientInputBuffers_External(int32 PhysicsStep, const TArray<Chaos::FSimCallbackInputAndObject>& SimCallbackInputs);
+	void ProcessClientInputBuffers_External(int32 PhysicsStep);
 
 	// Historic GT recording for GT only (compiled out. probably temp until Insights integration is back)
 	struct FDebugSnapshot
