Network Prediction
-Add NetCore to build.cs

diff --git a/Engine/Plugins/Runtime/NetworkPrediction/Source/NetworkPrediction/NetworkPrediction.Build.cs b/Engine/Plugins/Runtime/NetworkPrediction/Source/NetworkPrediction/NetworkPrediction.Build.cs
--- a/Engine/Plugins/Runtime/NetworkPrediction/Source/NetworkPrediction/NetworkPrediction.Build.cs
+++ b/Engine/Plugins/Runtime/NetworkPrediction/Source/NetworkPrediction/NetworkPrediction.Build.cs
@@ -24,6 +24,7 @@ namespace UnrealBuildTool.Rules
 				{
 					"Core",
 					"CoreUObject",
+					"NetCore",
                     "Engine",
                     "RenderCore",
 					"PhysicsCore",
