Network Prediction Physics
-Fix case where server would process wrong remote client InputCmd after a hitch

diff --git a/Engine/Plugins/Runtime/NetworkPrediction/Source/NetworkPrediction/Private/Async/NetworkPredictionAsyncWorldManager.cpp b/Engine/Plugins/Runtime/NetworkPrediction/Source/NetworkPrediction/Private/Async/NetworkPredictionAsyncWorldManager.cpp
--- a/Engine/Plugins/Runtime/NetworkPrediction/Source/NetworkPrediction/Private/Async/NetworkPredictionAsyncWorldManager.cpp
+++ b/Engine/Plugins/Runtime/NetworkPrediction/Source/NetworkPrediction/Private/Async/NetworkPredictionAsyncWorldManager.cpp
@@ -102,7 +102,7 @@ void FNetworkPredictionAsyncWorldManager::NetSerializePlayerControllerInputCmds(
 				else
 				{
 					// De-NetSerialize all new InputCmds from PC's buffer ["Client Frame" based] to our internal typed buffer ["local PhysicsStep Frame" based]
-					const int32 CmdOffset = PCFrameInfo.LastProcessedInputFrame - PhysicsStep; // Map PC frame to local Physics Step
+					const int32 CmdOffset = PCFrameInfo.LastProcessedInputFrame - PCFrameInfo.LastLocalFrame; // Map PC frame to local Physics Step
 					const int32 WriteStart = FMath::Max(PhysicsStep, Info.LastFrame+1); // Start at the current PhysicsStep or highest frame we haven't processed
 					const int32 WriteEnd = PCInputBuffer.HeadFrame() - CmdOffset; // go until we reach the end of valid data in the PC input buffer
 
@@ -138,7 +138,7 @@ void FNetworkPredictionAsyncWorldManager::NetSerializePlayerControllerInputCmds(
 				// We will probably need settings or some way to not send client->server at extremely high rates
 				// for now, its every game thread frame
 				TArray<uint8> SendData;
-				if (Writer.IsError() == false)
+				if (npEnsure(Writer.IsError() == false))
 				{
 					int32 NumBytes = (int32)Writer.GetNumBytes();
 					SendData = MoveTemp(*const_cast<TArray<uint8>*>(Writer.GetBuffer()));
