Network Physics
-Call SpawnProxyIfNeeded when applying rewind data so that newly spawned objects can be corrected "in the past"

diff --git a/Engine/Plugins/Runtime/NetworkPrediction/Source/NetworkPrediction/Private/NetworkPhysics.cpp b/Engine/Plugins/Runtime/NetworkPrediction/Source/NetworkPrediction/Private/NetworkPhysics.cpp
--- a/Engine/Plugins/Runtime/NetworkPrediction/Source/NetworkPrediction/Private/NetworkPhysics.cpp
+++ b/Engine/Plugins/Runtime/NetworkPrediction/Source/NetworkPrediction/Private/NetworkPhysics.cpp
@@ -44,6 +44,7 @@ namespace UE_NETWORK_PHYSICS
 	NP_DEVCVAR_FLOAT(W, 1.0f, "np2.Tolerance.W", "Rotational Velocity Tolerance");
 
 	NP_DEVCVAR_INT(ResimCacheEnabled, 1, "np2.ResimCacheEnabled", "Debug mode for in world drawing");
+	NP_DEVCVAR_INT(SpawnProxyIfNeededEnabled, 1, "np2.SpawnProxyIfNeededEnabled", "Call SpawnProxyIfNeededEnabled to deal with spawning issues");
 
 	NP_DEVCVAR_INT(Debug, 0, "np2.Debug", "Debug mode for in world drawing");
 	NP_DEVCVAR_INT(DebugTolerance, 1, "np2.Debug.Tolerance", "If enabled, only draw large corrections in world.");
@@ -417,6 +418,11 @@ struct FNetworkPhysicsRewindCallback : public Chaos::IRewindCallback
 						npEnsure(CorrectionState.Physics.LinearVelocity.ContainsNaN() == false);
 						npEnsure(CorrectionState.Physics.AngularVelocity.ContainsNaN() == false);
 
+						if (UE_NETWORK_PHYSICS::SpawnProxyIfNeededEnabled > 0)
+						{
+							RewindData->SpawnProxyIfNeeded(*CorrectionState.Proxy);
+						}
+
 						PT->SetX(CorrectionState.Physics.Location, false);
 						PT->SetV(CorrectionState.Physics.LinearVelocity, false);
 						PT->SetR(CorrectionState.Physics.Rotation, false);
@@ -426,6 +432,10 @@ struct FNetworkPhysicsRewindCallback : public Chaos::IRewindCallback
 						UE_CLOG(UE_NETWORK_PHYSICS::LogCorrections > 0 && PT->ObjectState() != CorrectionState.Physics.ObjectState, LogNetworkPhysics, Log, TEXT("Applying Correction State %d"), CorrectionState.Physics.ObjectState);
 						PT->SetObjectState(CorrectionState.Physics.ObjectState);
 					}
+					else
+					{
+						UE_LOG(LogNetworkPhysics, Warning, TEXT("No valid PT available for proxy."));
+					}
 				}
 			}
 		}
