Network Prediction WIP

#rnx

diff --git a/Engine/Plugins/Runtime/NetworkPrediction/Source/NetworkPrediction/Private/NetworkPhysics.cpp b/Engine/Plugins/Runtime/NetworkPrediction/Source/NetworkPrediction/Private/NetworkPhysics.cpp
--- a/Engine/Plugins/Runtime/NetworkPrediction/Source/NetworkPrediction/Private/NetworkPhysics.cpp
+++ b/Engine/Plugins/Runtime/NetworkPrediction/Source/NetworkPrediction/Private/NetworkPhysics.cpp
@@ -16,6 +16,8 @@
 #include "Engine/Canvas.h"
 #include "Debug/ReporterGraph.h"
 #include "Misc/ScopeExit.h"
+#include "GameFramework/Actor.h"
+#include "NetworkPredictionCheck.h"
 #include "PhysicsProxy/SingleParticlePhysicsProxy.h"
 
 DEFINE_LOG_CATEGORY(LogNetworkPhysics);
@@ -44,6 +46,12 @@ namespace UE_NETWORK_PHYSICS
 
 	int32 FixedLocalFrameOffsetTolerance = 3;
 	FAutoConsoleVariableRef CVarFixedLocalFrameOffsetTolerance(TEXT("np2.FixedLocalFrameOffsetTolerance"), FixedLocalFrameOffsetTolerance, TEXT("When > 0, use hardcoded frame offset on client from head"));
+
+	int32 EnableLOD = 1;
+	FAutoConsoleVariableRef CVarEnableLOD(TEXT("np2.EnableLOD"), EnableLOD, TEXT("Enable local LOD mode"));
+
+	float LODDistance = 2400.f;
+	FAutoConsoleVariableRef CVarLODDistance(TEXT("np2.LODDistance"), LODDistance, TEXT("Simple distance based LOD"));
 }
 
 struct FNetworkPhysicsRewindCallback : public Chaos::IRewindCallback
@@ -51,19 +59,19 @@ struct FNetworkPhysicsRewindCallback : public Chaos::IRewindCallback
 	bool CompareVec(const FVector& A, const FVector& B, const float D, const TCHAR* Str)
 	{
 		const bool b = FVector::DistSquared(A, B) > D * D;
-		UE_CLOG(UE_NETWORK_PHYSICS::bLogCorrections && b, LogNetworkPhysics, Log, TEXT("%s correction. Server: %s. Local: %s. Delta: %s (%.4f)"), Str, *A.ToString(), *B.ToString(), *(A-B).ToString(), (A-B).Size());
+		UE_CLOG(UE_NETWORK_PHYSICS::bLogCorrections && b && Str, LogNetworkPhysics, Log, TEXT("%s correction. Server: %s. Local: %s. Delta: %s (%.4f)"), Str, *A.ToString(), *B.ToString(), *(A-B).ToString(), (A-B).Size());
 		return  b;
 	}
 	bool CompareQuat(const FQuat& A, const FQuat& B, const float D, const TCHAR* Str)
 	{
 		const bool b = FQuat::ErrorAutoNormalize(A, B) > D;
-		UE_CLOG(UE_NETWORK_PHYSICS::bLogCorrections && b, LogNetworkPhysics, Log, TEXT("%s correction. Server: %s. Local: %s. Delta: %f"), Str, *A.ToString(), *B.ToString(), FQuat::ErrorAutoNormalize(A, B));
+		UE_CLOG(UE_NETWORK_PHYSICS::bLogCorrections && b && Str, LogNetworkPhysics, Log, TEXT("%s correction. Server: %s. Local: %s. Delta: %f"), Str, *A.ToString(), *B.ToString(), FQuat::ErrorAutoNormalize(A, B));
 		return b;
 	}
-	bool CompareObjState(Chaos::EObjectStateType A, Chaos::EObjectStateType B)
+	bool CompareObjState(Chaos::EObjectStateType A, Chaos::EObjectStateType B, const TCHAR* Str)
 	{
 		const bool b = A != B;
-		UE_CLOG(UE_NETWORK_PHYSICS::bLogCorrections && b, LogNetworkPhysics, Log, TEXT("ObjectState correction. Server: %d. Local: %d."), A, B);
+		UE_CLOG(UE_NETWORK_PHYSICS::bLogCorrections && b && Str, LogNetworkPhysics, Log, TEXT("%s correction. Server: %d. Local: %d."), Str, A, B);
 		return b;
 	}
 
@@ -75,6 +83,7 @@ struct FNetworkPhysicsRewindCallback : public Chaos::IRewindCallback
 			return INDEX_NONE;
 		}
 
+		/*
 		if (UE_NETWORK_PHYSICS::bForceResim)
 		{
 			DataFromNetwork.Empty();
@@ -85,6 +94,7 @@ struct FNetworkPhysicsRewindCallback : public Chaos::IRewindCallback
 			}
 			return INDEX_NONE;
 		}
+		*/
 
 		FStats Stats;
 		Stats.LatestSimFrame = RewindData->CurrentFrame();
@@ -108,6 +118,8 @@ struct FNetworkPhysicsRewindCallback : public Chaos::IRewindCallback
 		int32 RewindToFrame = INDEX_NONE;
 		while (DataFromNetwork.Dequeue(Snapshot))
 		{
+			// Detect change in local offset. This usually results in a bunch of corrections
+			// that are the fault of networking, not physics/user code. Note this for debugging purposes.
 			if (Snapshot.LocalFrameOffset != LastLocalOffset)
 			{
 				LastLocalOffset = Snapshot.LocalFrameOffset;
@@ -127,33 +139,39 @@ struct FNetworkPhysicsRewindCallback : public Chaos::IRewindCallback
 
 				if (Proxy->GetHandle_LowLevel() == nullptr)
 				{
+					UE_LOG(LogTemp, Warning, TEXT("No valid LowLEvel handle yet..."));
 					continue;
 				}
 
 				const auto P = RewindData->GetPastStateAtFrame(*Proxy->GetHandle_LowLevel(), Obj.Frame);
 				const int32 SimulationFrame = Obj.Frame - Snapshot.LocalFrameOffset;
-				//if (CompareObjState(Obj.ObjectState, P.ObjectState()))
+				
+				if (CompareObjState(Obj.Physics.ObjectState, P.ObjectState(), TEXT("DEBUG ObjectState")))
+				//if (CompareVec(Obj.Physics.Location, P.X(), UE_NETWORK_PHYSICS::X, TEXT("Location")))
 				{
-					//UE_LOG(LogTemp, Warning, TEXT("ObjectState Correction. Obj.Frame: %d. CurrentFramE: %d"), Obj.Frame, RewindData->CurrentFrame());
-					/*
-					for (int32 F = Obj.Frame; F <= RewindData->CurrentFrame(); ++F)
+					UE_LOG(LogTemp, Warning, TEXT("ObjectState Correction. Obj.Frame: %d. CurrentFramE: %d"), Obj.Frame, RewindData->CurrentFrame());
+					//UE_LOG(LogTemp, Warning, TEXT("Location Correction. Obj.Frame: %d. CurrentFrame: %d"), Obj.Frame, RewindData->CurrentFrame());
+					
+					for (int32 F = Obj.Frame-1; F <= RewindData->CurrentFrame(); ++F)
 					{
 						const auto PP = RewindData->GetPastStateAtFrame(*Proxy->GetHandle_LowLevel(), F);
-						UE_LOG(LogTemp, Warning, TEXT("    [%d/%d] --> %s"), F, F-Snapshot.LocalFrameOffset, *FRotator(PP.R()).ToString());
+						////UE_LOG(LogTemp, Warning, TEXT("    [%d/%d] --> %s"), F, F-Snapshot.LocalFrameOffset, *FRotator(PP.R()).ToString());
+						//UE_LOG(LogTemp, Warning, TEXT("    [%d/%d] --> %s"), F, F-Snapshot.LocalFrameOffset, *FVector(PP.X()).ToString());
+						UE_LOG(LogTemp, Warning, TEXT("    [%d/%d] --> %d"), F, F-Snapshot.LocalFrameOffset, (int32)PP.ObjectState());
 					}
 					UE_LOG(LogTemp, Warning, TEXT(""));
-					*/
+					
 				}
 
 				Stats.MinFrameChecked = FMath::Min(Stats.MaxFrameChecked, Obj.Frame);
 				Stats.MaxFrameChecked = FMath::Max(Stats.MaxFrameChecked, Obj.Frame);
 				Stats.NumChecked++;
 				
-				if (CompareObjState(Obj.Physics.ObjectState, P.ObjectState()) ||
+				if (CompareObjState(Obj.Physics.ObjectState, P.ObjectState(), TEXT("ObjectState")) ||
 					CompareVec(Obj.Physics.Location, P.X(), UE_NETWORK_PHYSICS::X, TEXT("Location")) ||
 					CompareVec(Obj.Physics.LinearVelocity, P.V(), UE_NETWORK_PHYSICS::V, TEXT("Velocity")) ||
 					CompareVec(Obj.Physics.AngularVelocity, P.W(), UE_NETWORK_PHYSICS::V, TEXT("Angular Velocity")) ||
-					CompareQuat(Obj.Physics.Rotation, P.R(), UE_NETWORK_PHYSICS::R, TEXT("Rotation")))
+					CompareQuat(Obj.Physics.Rotation, P.R(), UE_NETWORK_PHYSICS::R, TEXT("Rotation")) || UE_NETWORK_PHYSICS::bForceResim)
 				{
 
 					UE_CLOG(UE_NETWORK_PHYSICS::bLogCorrections, LogNetworkPhysics, Log, TEXT("Rewind Needed. SimFrame: %d. Obj.Frame: %d. LastCompletedStep: %d."), SimulationFrame, Obj.Frame, LastCompletedStep);
@@ -185,6 +203,17 @@ struct FNetworkPhysicsRewindCallback : public Chaos::IRewindCallback
 			}
 		}
 
+		
+		for (auto& SubSys : NetworkPhysicsManager->SubSystems)
+		{
+			const int32 CallbackFrame = SubSys.Value->TriggerRewindIfNeeded_Internal(LastCompletedStep);
+			if (CallbackFrame != INDEX_NONE)
+			{
+				Stats.NumSubsystemCorrections++;
+				RewindToFrame = RewindToFrame == INDEX_NONE ? CallbackFrame : FMath::Min(RewindToFrame, CallbackFrame);
+			}
+		}
+
 		// Calling code will ensure if we return a bad frame
 		if (!(RewindToFrame > INDEX_NONE && RewindToFrame < LastCompletedStep-1))
 		{
@@ -192,7 +221,7 @@ struct FNetworkPhysicsRewindCallback : public Chaos::IRewindCallback
 			return INDEX_NONE;
 		}
 		
-		ResimEndFrame = RewindToFrame;
+		ResimEndFrame = LastCompletedStep;
 		return RewindToFrame;
 	}
 	
@@ -200,7 +229,7 @@ struct FNetworkPhysicsRewindCallback : public Chaos::IRewindCallback
 	{
 		// Apply Physics corrections if necessary
 		if (PendingCorrectionIdx != INDEX_NONE)
-		{
+		{		
 			for (; PendingCorrectionIdx < PendingCorrections.Num(); ++PendingCorrectionIdx)
 			{
 				FNetworkPhysicsState& CorrectionState = PendingCorrections[PendingCorrectionIdx];
@@ -211,14 +240,14 @@ struct FNetworkPhysicsRewindCallback : public Chaos::IRewindCallback
 
 				if (auto* PT = CorrectionState.Proxy->GetPhysicsThreadAPI())
 				{
-					UE_LOG(LogNetworkPhysics, Log, TEXT("Applying Correction from frame %d (actual step: %d)"), CorrectionState.Frame, PhysicsStep);
-
-					PT->SetX(CorrectionState.Physics.Location);
-					PT->SetV(CorrectionState.Physics.LinearVelocity);
-					PT->SetR(CorrectionState.Physics.Rotation);
-					PT->SetW(CorrectionState.Physics.AngularVelocity);
-
-					if (PT->ObjectState() != CorrectionState.Physics.ObjectState)
+					UE_LOG(LogNetworkPhysics, Log, TEXT("Applying Correction from frame %d (actual step: %d). Location: %s"), CorrectionState.Frame, PhysicsStep, *FVector(CorrectionState.Physics.Location).ToString());
+				
+					PT->SetX(CorrectionState.Physics.Location, false);
+					PT->SetV(CorrectionState.Physics.LinearVelocity, false);
+					PT->SetR(CorrectionState.Physics.Rotation, false);
+					PT->SetW(CorrectionState.Physics.AngularVelocity, false);
+					
+					//if (PT->ObjectState() != CorrectionState.Physics.ObjectState)
 					{
 						ensure(CorrectionState.Physics.ObjectState != Chaos::EObjectStateType::Uninitialized);
 						UE_LOG(LogNetworkPhysics, Log, TEXT("Applying Correction State %d"), CorrectionState.Physics.ObjectState);
@@ -228,8 +257,6 @@ struct FNetworkPhysicsRewindCallback : public Chaos::IRewindCallback
 			}
 		}
 
-		
-
 		// Marhsall data back to GT based on what was requested for networking (this should be server only)
 		FRequest Request;
 		while (DataRequested.Dequeue(Request));
@@ -258,7 +285,6 @@ struct FNetworkPhysicsRewindCallback : public Chaos::IRewindCallback
 				{
 					//UE_LOG(LogTemp, Warning, TEXT("[Server] Sleeping Frame: %d"), PhysicsStep);
 				}
-
 				//UE_LOG(LogTemp, Warning, TEXT("[Server] %d Rotation: %s"), PhysicsStep, *FRotator(Obj.Physics.Rotation).ToString());
 			}
 		}
@@ -268,7 +294,10 @@ struct FNetworkPhysicsRewindCallback : public Chaos::IRewindCallback
 			DataFromPhysics.Enqueue(MoveTemp(Snapshot));
 		}
 
-		// ------------------------------------------------------------------------
+		for (auto& SubSys : NetworkPhysicsManager->SubSystems)
+		{
+			SubSys.Value->ProcessInputs_Internal(PhysicsStep);
+		}
 		
 		for (auto CallbacKInput : SimCallbackInputs)
 		{
@@ -279,7 +308,7 @@ struct FNetworkPhysicsRewindCallback : public Chaos::IRewindCallback
 					Callback->ApplyCorrections_Internal(PhysicsStep, CallbacKInput.Input);
 					break;
 				}
-			}			
+			}
 		}
 	}
 
@@ -294,6 +323,10 @@ struct FNetworkPhysicsRewindCallback : public Chaos::IRewindCallback
 			}
 		}
 
+		for (auto& SubSys : NetworkPhysicsManager->SubSystems)
+		{
+			SubSys.Value->PreResimStep_Internal(PhysicsStep, bFirst);
+		}
 	}
 
 	void PostResimStep_Internal(int32 PhysicsStep) override
@@ -306,110 +339,8 @@ struct FNetworkPhysicsRewindCallback : public Chaos::IRewindCallback
 
 	void ProcessInputs_External(int32 PhysicsStep, const TArray<Chaos::FSimCallbackInputAndObject>& SimCallbackInputs) override 
 	{
-		// Finalize our GT input for this PhysicsStep.
-		if (World->GetNetMode() == NM_Client)
-		{
-			// Client:
-
-		
-
-			if (APlayerController* PC = World->GetFirstPlayerController())
-			{
-				// TODO: Broadcast delegate to generate new InputCmd from user code and call RPC to replicate it
-				FNetBitWriter Writer(nullptr, 256 << 3);
-
-				for (Chaos::FSimCallbackInputAndObject Pair :  SimCallbackInputs)
-				{
-					if (Pair.Input->NetSendInputCmd(Writer))
-					{
-						break;
-					}
-				}
-
-				// Calling the RPC right here is bad. There probably is some intermediate buffer that these finalized commands should go to and the network
-				// sending can happen elsewhere. That is also where move combining / delta compressiong type logic could happen.
-				TArray<uint8> SendData;
-				if (Writer.IsError() == false)
-				{
-					int32 NumBytes = (int32)Writer.GetNumBytes();
-					SendData = MoveTemp(*const_cast<TArray<uint8>*>(Writer.GetBuffer()));
-					SendData.SetNum(NumBytes);
-				}
-				
-				PC->PushClientInput(PhysicsStep, SendData);
-			}
-		}
-		else
-		{
-			// Server: look at buffered received ClientInputFrames and consume (usually) one
-			constexpr int32 MaxBufferedCmds = 64;
-
-			for (FConstPlayerControllerIterator Iterator = World->GetPlayerControllerIterator(); Iterator; ++Iterator)
-			{
-				if (APlayerController* PC = Iterator->Get())
-				{
-					APlayerController::FServerFrameInfo& FrameInfo = PC->GetServerFrameInfo();
-					APlayerController::FInputCmdBuffer& InputBuffer = PC->GetInputBuffer();
-
-					auto UpdateLastProcessedInputFrame = [&]()
-					{
-						const int32 NumBufferedInputCmds = InputBuffer.HeadFrame() - FrameInfo.LastProcessedInputFrame;
-
-						// Check Overflow
-						if (NumBufferedInputCmds > MaxBufferedCmds)
-						{
-							UE_LOG(LogNetworkPhysics, Warning, TEXT("[Remote.Input] overflow %d %d -> %d"), InputBuffer.HeadFrame(), FrameInfo.LastProcessedInputFrame, NumBufferedInputCmds);
-							FrameInfo.LastProcessedInputFrame = InputBuffer.HeadFrame() - MaxBufferedCmds + 1;
-						}
-
-						// Check fault - we are waiting for Cmds to reach FaultLimit before continuing
-						if (FrameInfo.bFault)
-						{
-							if (NumBufferedInputCmds < FrameInfo.FaultLimit)
-							{
-								// Skip this because it is in fault. We will use the prev input for this frame.
-								UE_CLOG(FrameInfo.LastProcessedInputFrame != INDEX_NONE, LogNetworkPhysics, Warning, TEXT("[Remote.Input] in fault. Reusing Inputcmd. (Client) Input: %d. (Server) Local Frame:"), FrameInfo.LastProcessedInputFrame, FrameInfo.LastLocalFrame);
-								return;
-							}
-							FrameInfo.bFault = false;
-						}
-						else if (NumBufferedInputCmds <= 0)
-						{
-							// No Cmds to process, enter fault state. Increment FaultLimit each time this happens.
-							// TODO: We should have something to bring this back down (which means skipping frames) we don't want temporary poor conditions to cause permanent high input buffering
-							FrameInfo.bFault = true;
-							FrameInfo.FaultLimit = FMath::Min(FrameInfo.FaultLimit+1, MaxBufferedCmds-1);
-
-							UE_CLOG(FrameInfo.LastProcessedInputFrame != INDEX_NONE, LogNetworkPhysics, Warning, TEXT("[Remote.Input] ENTERING fault. New Limit: %d. (Client) Input: %d. (Server) Local Frame:"), FrameInfo.FaultLimit, FrameInfo.LastProcessedInputFrame, FrameInfo.LastLocalFrame);
-							return;
-						}
-
-						FrameInfo.LastProcessedInputFrame++;
-						FrameInfo.LastLocalFrame = PhysicsStep;
-					};
-
-					UpdateLastProcessedInputFrame();
-
-					if (FrameInfo.LastProcessedInputFrame != INDEX_NONE)
-					{
-						const TArray<uint8>& Data = InputBuffer.Get(FrameInfo.LastProcessedInputFrame);
-						if (Data.Num() > 0)
-						{
-							uint8* RawData = const_cast<uint8*>(Data.GetData());
-							FNetBitReader Ar(nullptr, RawData, ((int64)Data.Num()) << 3);
-						
-							for (Chaos::FSimCallbackInputAndObject Pair :  SimCallbackInputs)
-							{
-								if (Pair.Input->NetRecvInputCmd(PC, Ar))
-								{
-									break;
-								}
-							}
-						}
-					}
-				}
-			}
-		}
+		npCheckSlow(NetworkPhysicsManager);
+		NetworkPhysicsManager->ProcessInputs_External(PhysicsStep, SimCallbackInputs);
 	}
 
 	void RegisterRewindableSimCallback_Internal(Chaos::ISimCallbackObject* Callback) override
@@ -449,6 +380,7 @@ struct FNetworkPhysicsRewindCallback : public Chaos::IRewindCallback
 	// ----------------------
 
 	TArray<Chaos::ISimCallbackObject*> SimObjectCallbacks;
+	UNetworkPhysicsManager* NetworkPhysicsManager = nullptr;
 
 
 	// Debugging
@@ -537,6 +469,7 @@ void UNetworkPhysicsManager::OnWorldPostInit(UWorld* World, const UWorld::Initia
 			RewindCallback->RewindData = Solver->GetRewindData();
 			RewindCallback->Solver = Solver;
 			RewindCallback->World = World;
+			RewindCallback->NetworkPhysicsManager = this;
 		}
 	}
 #endif
@@ -638,6 +571,17 @@ void UNetworkPhysicsManager::PostNetRecv()
 			}
 		}
 
+		FVector LocalViewLoc;
+		FRotator LocalViewRot;
+		const float LODDistSq = UE_NETWORK_PHYSICS::LODDistance * UE_NETWORK_PHYSICS::LODDistance;
+		if (UE_NETWORK_PHYSICS::EnableLOD > 0)
+		{
+			if (APlayerController* PC = GetWorld()->GetFirstPlayerController())
+			{
+				PC->GetPlayerViewPoint(LocalViewLoc, LocalViewRot);
+			}
+		}
+
 		// Go through the physics states that we are managing and look for new network updates and marshal them to PT with correct local frame
 		for (FNetworkPhysicsState* PhysicsState : ManagedPhysicsStates)
 		{
@@ -645,11 +589,60 @@ void UNetworkPhysicsManager::PostNetRecv()
 			MaxFrameSeen = FMath::Max(PhysicsState->Frame, MaxFrameSeen);
 			if (PhysicsState->Frame > LatestConfirmedFrame)
 			{
-				const int32 LocalFrame = PhysicsState->Frame + LocalOffset;
-				if (LocalFrame > 0 && PhysicsState->Physics.ObjectState != Chaos::EObjectStateType::Uninitialized)
+				auto CreateMarshalledCopy = [&]()
+				{
+					const int32 LocalFrame = PhysicsState->Frame + LocalOffset;
+					if (LocalFrame > 0 && PhysicsState->Physics.ObjectState != Chaos::EObjectStateType::Uninitialized)
+					{
+						FNetworkPhysicsState& MarshalledCopy = Snapshot.Objects.Emplace_GetRef(FNetworkPhysicsState(*PhysicsState));
+						MarshalledCopy.Frame = LocalFrame;
+					}
+				};
+
+				if (UE_NETWORK_PHYSICS::EnableLOD > 0)
 				{
-					FNetworkPhysicsState& MarshalledCopy = Snapshot.Objects.Emplace_GetRef(FNetworkPhysicsState(*PhysicsState));
-					MarshalledCopy.Frame = LocalFrame;
+					if (PhysicsState->Proxy)
+					{
+						int32 CalculatedLOD = 0;
+						const FVector Location = PhysicsState->Proxy->GetGameThreadAPI().X();
+						if (FVector::DistSquared(Location, LocalViewLoc) > LODDistSq)
+						{
+							CalculatedLOD = 1;
+						}
+					
+						if (PhysicsState->LocalLOD != CalculatedLOD)
+						{
+							PhysicsState->LocalLOD = CalculatedLOD;
+						}
+
+						if (PhysicsState->LocalLOD == 0)
+						{
+							CreateMarshalledCopy();
+
+							PhysicsState->OwningActor->GetReplicatedMovement_Mutable().bRepPhysics = false;
+							PhysicsState->OwningActor->SetReplicateMovement(false);
+						}
+						else
+						{
+							if (ensure(PhysicsState->OwningActor))
+							{
+								FRepMovement& RepMovement = PhysicsState->OwningActor->GetReplicatedMovement_Mutable();
+								RepMovement.Location = PhysicsState->Physics.Location;
+								RepMovement.Rotation = FRotator(PhysicsState->Physics.Rotation);
+								RepMovement.LinearVelocity = PhysicsState->Physics.LinearVelocity;
+								RepMovement.AngularVelocity = PhysicsState->Physics.AngularVelocity;
+								RepMovement.bSimulatedPhysicSleep = (PhysicsState->Physics.ObjectState == Chaos::EObjectStateType::Sleeping);
+								
+								RepMovement.bRepPhysics = true;
+								PhysicsState->OwningActor->SetReplicateMovement(true);
+								PhysicsState->OwningActor->OnRep_ReplicatedMovement();
+							}
+						}
+					}
+				}
+				else
+				{
+					CreateMarshalledCopy();
 				}
 			}
 		}
@@ -750,6 +743,149 @@ void UNetworkPhysicsManager::PreNetSend(float DeltaSeconds)
 #endif
 }
 
+void UNetworkPhysicsManager::ProcessInputs_External(int32 PhysicsStep, const TArray<Chaos::FSimCallbackInputAndObject>& SimCallbackInputs)
+{
+	UWorld* World = GetWorld();
+	const ENetMode NetMode = World->GetNetMode();
+
+	// ----------------------------------------------------------------------
+	//	Server: Do remote client input buffering logic
+	//	This should probably be moved into PlayerController itself, with more options.
+	//	We want this to happen before calling into INetworkPhysicsSubsystem::ProcessInputs_External.
+	// ----------------------------------------------------------------------
+	
+	if (NetMode != NM_Client)
+	{
+		constexpr int32 MaxBufferedCmds = 64;
+		for (FConstPlayerControllerIterator Iterator = World->GetPlayerControllerIterator(); Iterator; ++Iterator)
+		{
+			if (APlayerController* PC = Iterator->Get())
+			{
+				APlayerController::FServerFrameInfo& FrameInfo = PC->GetServerFrameInfo();
+				APlayerController::FInputCmdBuffer& InputBuffer = PC->GetInputBuffer();
+
+				auto UpdateLastProcessedInputFrame = [&]()
+				{
+					const int32 NumBufferedInputCmds = InputBuffer.HeadFrame() - FrameInfo.LastProcessedInputFrame;
+
+					// Check Overflow
+					if (NumBufferedInputCmds > MaxBufferedCmds)
+					{
+						UE_LOG(LogNetworkPhysics, Warning, TEXT("[Remote.Input] overflow %d %d -> %d"), InputBuffer.HeadFrame(), FrameInfo.LastProcessedInputFrame, NumBufferedInputCmds);
+						FrameInfo.LastProcessedInputFrame = InputBuffer.HeadFrame() - MaxBufferedCmds + 1;
+					}
+
+					// Check fault - we are waiting for Cmds to reach FaultLimit before continuing
+					if (FrameInfo.bFault)
+					{
+						if (NumBufferedInputCmds < FrameInfo.FaultLimit)
+						{
+							// Skip this because it is in fault. We will use the prev input for this frame.
+							UE_CLOG(FrameInfo.LastProcessedInputFrame != INDEX_NONE, LogNetworkPhysics, Warning, TEXT("[Remote.Input] in fault. Reusing Inputcmd. (Client) Input: %d. (Server) Local Frame:"), FrameInfo.LastProcessedInputFrame, FrameInfo.LastLocalFrame);
+							return;
+						}
+						FrameInfo.bFault = false;
+					}
+					else if (NumBufferedInputCmds <= 0)
+					{
+						// No Cmds to process, enter fault state. Increment FaultLimit each time this happens.
+						// TODO: We should have something to bring this back down (which means skipping frames) we don't want temporary poor conditions to cause permanent high input buffering
+						FrameInfo.bFault = true;
+						FrameInfo.FaultLimit = FMath::Min(FrameInfo.FaultLimit+1, MaxBufferedCmds-1);
+
+						UE_CLOG(FrameInfo.LastProcessedInputFrame != INDEX_NONE, LogNetworkPhysics, Warning, TEXT("[Remote.Input] ENTERING fault. New Limit: %d. (Client) Input: %d. (Server) Local Frame:"), FrameInfo.FaultLimit, FrameInfo.LastProcessedInputFrame, FrameInfo.LastLocalFrame);
+						return;
+					}
+
+					FrameInfo.LastProcessedInputFrame++;
+					FrameInfo.LastLocalFrame = PhysicsStep;
+				};
+
+				UpdateLastProcessedInputFrame();
+
+
+				// ----------------------------------------------------------------------------------------------
+				// REMOVE Once Mock/Hardcoded exmaple is gone
+				// ----------------------------------------------------------------------------------------------
+				if (FrameInfo.LastProcessedInputFrame != INDEX_NONE)
+				{
+					const TArray<uint8>& Data = InputBuffer.Get(FrameInfo.LastProcessedInputFrame);
+					if (Data.Num() > 0)
+					{
+						uint8* RawData = const_cast<uint8*>(Data.GetData());
+						FNetBitReader Ar(nullptr, RawData, ((int64)Data.Num()) << 3);
+						
+						for (Chaos::FSimCallbackInputAndObject Pair :  SimCallbackInputs)
+						{
+							if (Pair.Input->NetRecvInputCmd(PC, Ar))
+							{
+								break;
+							}
+						}
+					}
+				}
+				// ----------------------------------------------------------------------------------------------
+				// ----------------------------------------------------------------------------------------------
+				// ----------------------------------------------------------------------------------------------
+
+			}
+		}
+	}
+
+	// ----------------------------------------------------------------------
+	//	Call into Subsystems. This is where they will marshall input to the PT
+	// ----------------------------------------------------------------------
+
+	bool bOutSentClientInput = false;
+	for (auto& Pair : SubSystems)
+	{
+		Pair.Value->ProcessInputs_External(PhysicsStep, LocalOffset, bOutSentClientInput);
+	}
+
+	// ----------------------------------------------------------------------
+	// If we are a client and none of the subsystems sent ClientInput on their own,
+	// then we should flush empty input to the server so that we can have an accurate
+	// LocalFrameOffset. This is the critical piece of information for determining how
+	// far ahead of the server we should be.
+	// ----------------------------------------------------------------------
+	if (NetMode == NM_Client && !bOutSentClientInput)
+	{
+		if (APlayerController* PC = World->GetFirstPlayerController())
+		{
+			TArray<uint8> SendData;
+
+			// ----------------------------------------------------------------------------------------------
+			// REMOVE Once Mock/Hardcoded exmaple is gone
+			// ----------------------------------------------------------------------------------------------
+			{
+				FNetBitWriter Writer(nullptr, 256 << 3);
+
+				for (Chaos::FSimCallbackInputAndObject Pair :  SimCallbackInputs)
+				{
+					if (Pair.Input->NetSendInputCmd(Writer))
+					{
+						break;
+					}
+				}
+
+			
+			
+				if (Writer.IsError() == false)
+				{
+					int32 NumBytes = (int32)Writer.GetNumBytes();
+					SendData = MoveTemp(*const_cast<TArray<uint8>*>(Writer.GetBuffer()));
+					SendData.SetNum(NumBytes);
+				}
+			}
+
+			// ----------------------------------------------------------------------------------------------
+			// ----------------------------------------------------------------------------------------------
+			// ----------------------------------------------------------------------------------------------
+			PC->PushClientInput(PhysicsStep, SendData);
+		}
+	}
+}
+
 void UNetworkPhysicsManager::RegisterPhysicsProxy(FNetworkPhysicsState* State)
 {
 	checkSlow(State);
diff --git a/Engine/Plugins/Runtime/NetworkPrediction/Source/NetworkPrediction/Public/NetworkPhysics.h b/Engine/Plugins/Runtime/NetworkPrediction/Source/NetworkPrediction/Public/NetworkPhysics.h
--- a/Engine/Plugins/Runtime/NetworkPrediction/Source/NetworkPrediction/Public/NetworkPhysics.h
+++ b/Engine/Plugins/Runtime/NetworkPrediction/Source/NetworkPrediction/Public/NetworkPhysics.h
@@ -16,6 +16,9 @@ struct FNetworkPhysicsRewindCallback;
 class FMockObjectManager;
 class FSingleParticlePhysicsProxy;
 
+namespace Chaos { struct FSimCallbackInputAndObject; }
+
+// FIXME: use FRigidBodyState instead
 struct FBasePhysicsState
 {
 	Chaos::EObjectStateType ObjectState = Chaos::EObjectStateType::Uninitialized;
@@ -67,6 +70,10 @@ struct FNetworkPhysicsState
 		Ar << Physics.AngularVelocity;
 		return true;
 	}
+
+	// LOD: should probably be moved out of this struct
+	int32 LocalLOD = 0;
+	AActor* OwningActor = nullptr;
 };
 
 template<>
@@ -86,6 +93,18 @@ public:
 	virtual ~INetworkPhysicsSubsystem() = default;
 	virtual void PostNetRecv(UWorld* World, int32 LocalOffset, int32 LastProcessedFrame) = 0;
 	virtual void PreNetSend(UWorld* World, float DeltaSeconds) = 0;
+
+	// Finalizes InputCmds, marshalls them to PT and Network.
+	virtual void ProcessInputs_External(int32 PhysicsStep, int32 LocalFrameOffset, bool& bOutSendClientInputCmd) { }
+	
+	// Records "Final" Inputs for a frame and marshalls them back to GT for networking
+	virtual void ProcessInputs_Internal(int32 PhysicsStep) { }
+
+	// ShouldReconcile
+	virtual int32 TriggerRewindIfNeeded_Internal(int32 LastCompletedStep) { return INDEX_NONE; }
+
+	// Applies Corrections
+	virtual void PreResimStep_Internal(int32 PhysicsStep, bool bFirst) { }
 };
 
 
@@ -152,4 +171,7 @@ private:
 	TMap<FSingleParticlePhysicsProxy*, TUniqueFunction<void(const FDrawDebugParams&)>> DrawDebugMap;
 	
 	TArray<TPair<FName, TUniquePtr<INetworkPhysicsSubsystem>>> SubSystems;
+
+	friend struct FNetworkPhysicsRewindCallback;
+	void ProcessInputs_External(int32 PhysicsStep, const TArray<Chaos::FSimCallbackInputAndObject>& SimCallbackInputs);
 };
\ No newline at end of file
diff --git a/Engine/Plugins/Runtime/NetworkPredictionExtras/Source/NetworkPredictionExtras/Private/MockNetworkPhysics.cpp b/Engine/Plugins/Runtime/NetworkPredictionExtras/Source/NetworkPredictionExtras/Private/MockNetworkPhysics.cpp
--- a/Engine/Plugins/Runtime/NetworkPredictionExtras/Source/NetworkPredictionExtras/Private/MockNetworkPhysics.cpp
+++ b/Engine/Plugins/Runtime/NetworkPredictionExtras/Source/NetworkPredictionExtras/Private/MockNetworkPhysics.cpp
@@ -116,28 +116,86 @@ namespace UE_NETWORK_PHYSICS
 
 	bool MockDebug=false;
 	FAutoConsoleVariableRef CVarMockDebug(TEXT("np2.Mock.Debug"), MockDebug, TEXT("Enabled spammy log debugging of mock physics object state"));
+
+	bool bEnableMock=true;
+	FAutoConsoleVariableRef CVarbEnableMock(TEXT("np2.Mock.Enable"), bEnableMock, TEXT("Enable Mock implementation"));
+		
+	// Mock Movement tweaking
+	float DragK=200.f;
+	FAutoConsoleVariableRef CVarDragK(TEXT("np2.Mock.DragK"), DragK, TEXT("Drag Coefficient (higher=more drag)"));
+
+	float MovementK=2.5f;
+	FAutoConsoleVariableRef CVarMovementK(TEXT("np2.Mock.MovementK"), MovementK, TEXT("Movement Coefficient (higher=faster movement)"));
+
+	float TurnK=100000.f;
+	FAutoConsoleVariableRef CVarTurnK(TEXT("np2.Mock.TurnK"), TurnK, TEXT("Coefficient for automatic turning (higher=quicker turning)"));
+
+	float TurnDampK=20.f;
+	FAutoConsoleVariableRef CVarTurnDampK(TEXT("np2.Mock.TurnDampK"), TurnDampK, TEXT("Coefficient for damping portion of turn. Higher=more damping but too higher will lead to instability."));
 }
 
 void FMockManagedState::AsyncTick(UWorld* World, Chaos::FPhysicsSolver* Solver, const float DeltaSeconds, const int32 SimulationFrame, const int32 LocalStorageFrame)
 {
 	// FIXME: PT_State is the only thing that should be writable here
-	// FIXME: DeltaSeconds/TotalSecond are local time and therefor worthless. This needs to be server frame/time
-
 	if (ensure(Proxy))
 	{
 		if (auto* PT = Proxy->GetPhysicsThreadAPI())
 		{
-			if (InputCmd.bBrakesPressed)
+			const float UpDot = FVector::DotProduct(PT->R().GetUpVector(), FVector::UpVector);
+			if (PT_State.RecoveryFrame == 0)
+			{
+				if (UpDot < 0.2f)
+				{
+					PT_State.RecoveryFrame = SimulationFrame;
+				}
+			}
+			
+			if (PT_State.RecoveryFrame != 0)
 			{
-				PT->SetV( Chaos::FVec3(0.f));				
+				if (UpDot > 0.7f)
+				{
+					// Recovered
+					PT_State.RecoveryFrame = 0;
+				}
+				else
+				{
+					UE_LOG(LogTemp, Warning, TEXT("In Air"));
+
+					FRotator Rot = PT->R().Rotator();
+					const float DeltaRoll = FRotator::NormalizeAxis( -1.f * (Rot.Roll + (PT->W().X * UE_NETWORK_PHYSICS::TurnDampK)));
+					const float DeltaPitch = FRotator::NormalizeAxis( -1.f * (Rot.Pitch + (PT->W().Y * UE_NETWORK_PHYSICS::TurnDampK)));
+
+					PT->AddTorque(FVector(DeltaRoll, DeltaPitch, 0.f) * UE_NETWORK_PHYSICS::TurnK * 1.5f);
+					PT->AddForce(FVector(0.f, 0.f, 600.f));
+				}
+			}
+
+			else if (InputCmd.bBrakesPressed)
+			{
+				FVector NewV = PT->V();
+				if (NewV.SizeSquared2D() < 1.f)
+				{
+					PT->SetV( Chaos::FVec3(0.f, 0.f, NewV.Z));
+				}
+				else
+				{
+					PT->SetV( Chaos::FVec3(NewV.X * 0.8f, NewV.Y * 0.8f, NewV.Z));
+				}
+				/*
+
+				PT->SetV( Chaos::FVec3(0.f));
 				PT->SetW( Chaos::FVec3(0.f));
 
 				UE_CLOG(UE_NETWORK_PHYSICS::MockDebug, LogNetworkPhysics, Log, TEXT("[%d/%d] Applied Break. Rot was: %s"), SimulationFrame, LocalStorageFrame, *PT->R().Rotator().ToString());
+				
 				Chaos::FRotation3 NewR = PT->R();
 				FRotator Rot = NewR.Rotator();
 				Rot.Pitch = 0.f;
 				Rot.Roll = 0.f;
-				PT->SetR( FQuat(Rot));
+				//PT->SetR( FQuat(Rot));				
+				//PT->SetR(FQuat(FRotator(-90.f, -71.f, 39.f)));
+				PT->SetR(FQuat(FRotator(-90.f, -71.f, 39.f)));
+				*/
 			}
 			else
 			{
@@ -146,8 +204,24 @@ void FMockManagedState::AsyncTick(UWorld* World, Chaos::FPhysicsSolver* Solver,
 					//UE_LOG(LogTemp, Warning, TEXT("Applied Force. Frame: %d"), LocalFrame);
 					//UE_LOG(LogTemp, Warning, TEXT("0x%X ForceMultiplier: %f (Rand: %d)"), (int64)Proxy, GT_State.ForceMultiplier, GT_State.RandValue);
 					//UE_CLOG(PC == nullptr && InputCmd.Force.SizeSquared() > 0.f, LogTemp, Warning, TEXT("0x%X ForceMultiplier: %f (Rand: %d)"), (int64)Proxy, GT_State.ForceMultiplier, GT_State.RandValue);
-					PT->AddForce(InputCmd.Force * GT_State.ForceMultiplier);
+					
+					PT->AddForce(InputCmd.Force * GT_State.ForceMultiplier * UE_NETWORK_PHYSICS::MovementK);
+
+					const float CurrentYaw = PT->R().Rotator().Yaw + (PT->W().Z * UE_NETWORK_PHYSICS::TurnDampK);
+					const float DesiredYaw = InputCmd.Force.Rotation().Yaw;
+					const float DeltaYaw = FRotator::NormalizeAxis( DesiredYaw - CurrentYaw );
+					
+					PT->AddTorque(FVector(0.f, 0.f, DeltaYaw * UE_NETWORK_PHYSICS::TurnK));
+				}
+
+				FVector V = PT->V();
+				V.Z = 0.f;
+				if (V.SizeSquared() > 0.1f)
+				{
+					FVector Drag = -1.f * V * UE_NETWORK_PHYSICS::DragK;
+					PT->AddForce(Drag);
 				}
+
 				
 				if (FMath::Abs<float>(InputCmd.Turn) > 0.001f)
 				{
@@ -755,6 +829,11 @@ void UNetworkPhysicsComponent::InitializeComponent()
 #if WITH_CHAOS
 	Super::InitializeComponent();
 
+	if (!UE_NETWORK_PHYSICS::bEnableMock)
+	{
+		return;
+	}
+
 	UWorld* World = GetWorld();	
 	checkSlow(World);
 	UNetworkPhysicsManager* Manager = World->GetSubsystem<UNetworkPhysicsManager>();
@@ -779,6 +858,9 @@ void UNetworkPhysicsComponent::InitializeComponent()
 	if (ensureMsgf(PrimitiveComponent, TEXT("No PrimitiveComponent found on %s"), *GetPathName()))
 	{
 		NetworkPhysicsState.Proxy = PrimitiveComponent->BodyInstance.ActorHandle;
+		NetworkPhysicsState.OwningActor = GetOwner();
+		ensure(NetworkPhysicsState.OwningActor);
+
 		Manager->RegisterPhysicsProxy(&NetworkPhysicsState);
 
 		Manager->RegisterPhysicsProxyDebugDraw(&NetworkPhysicsState, [this](const UNetworkPhysicsManager::FDrawDebugParams& P)
@@ -809,6 +891,13 @@ void UNetworkPhysicsComponent::InitializeComponent()
 void UNetworkPhysicsComponent::EndPlay(const EEndPlayReason::Type EndPlayReason)
 {
 	Super::EndPlay(EndPlayReason);
+#if WITH_CHAOS
+
+	if (!UE_NETWORK_PHYSICS::bEnableMock)
+	{
+		return;
+	}
+
 	//UE_LOG(LogTemp, Warning, TEXT("EndPlay %s"), *GetPathNameSafe(this));
 	if (UWorld* World = GetWorld())
 	{
@@ -822,11 +911,18 @@ void UNetworkPhysicsComponent::EndPlay(const EEndPlayReason::Type EndPlayReason)
 			Manager->UnregisterPhysicsProxy(&NetworkPhysicsState);
 		}
 	}
+#endif
 }
 
 void UNetworkPhysicsComponent::TickComponent(float DeltaTime, enum ELevelTick TickType, FActorComponentTickFunction *ThisTickFunction)
 {
 	Super::TickComponent(DeltaTime, TickType, ThisTickFunction);
+#if WITH_CHAOS
+
+	if (!UE_NETWORK_PHYSICS::bEnableMock)
+	{
+		return;
+	}
 
 	APlayerController* PC = GetOwnerPC();
 	if (PC && PC->IsLocalController())
@@ -838,6 +934,7 @@ void UNetworkPhysicsComponent::TickComponent(float DeltaTime, enum ELevelTick Ti
 	ReplicatedManagedState.PC = PC;
 	InManagedState.PC = PC;
 	OutManagedState.PC = PC;
+#endif
 }
 
 APlayerController* UNetworkPhysicsComponent::GetOwnerPC() const
@@ -857,7 +954,6 @@ void UNetworkPhysicsComponent::GetLifetimeReplicatedProps(TArray<FLifetimeProper
 	DOREPLIFETIME( UNetworkPhysicsComponent, ReplicatedManagedState);
 }
 
-
 // ============================================================================================================
 // ============================================================================================================
 // ============================================================================================================
diff --git a/Engine/Plugins/Runtime/NetworkPredictionExtras/Source/NetworkPredictionExtras/Public/MockNetworkPhysics.h b/Engine/Plugins/Runtime/NetworkPredictionExtras/Source/NetworkPredictionExtras/Public/MockNetworkPhysics.h
--- a/Engine/Plugins/Runtime/NetworkPredictionExtras/Source/NetworkPredictionExtras/Public/MockNetworkPhysics.h
+++ b/Engine/Plugins/Runtime/NetworkPredictionExtras/Source/NetworkPredictionExtras/Public/MockNetworkPhysics.h
@@ -142,17 +142,22 @@ struct FMockState_PT
 	UPROPERTY(BlueprintReadWrite, Category="Mock Object")
 	int32 CheckSum = 0;
 
+	// Frame we started "in air recovery" on
+	UPROPERTY(BlueprintReadWrite, Category="Mock Object")
+	int32 RecoveryFrame = 0;
+
 	bool NetSerialize(FArchive& Ar, class UPackageMap* Map, bool& bOutSuccess)
 	{
 		Ar << JumpCooldownMS;
 		Ar << JumpCount;
-		Ar << CheckSum;		
+		Ar << CheckSum;	
+		Ar << RecoveryFrame;
 		return true;
 	}
 
 	bool ShouldReconcile(const FMockState_PT& AuthState) const
 	{
-		return JumpCooldownMS != AuthState.JumpCooldownMS || JumpCount != AuthState.JumpCount;
+		return JumpCooldownMS != AuthState.JumpCooldownMS || JumpCount != AuthState.JumpCount || RecoveryFrame != AuthState.RecoveryFrame;
 	}
 };
 
@@ -318,6 +323,13 @@ public:
 	UFUNCTION(BlueprintPure, Category = "Mock Input")
 	FMockState_PT GetMockState_PT() const { return OutManagedState.PT_State; }
 
+	// ----------------------------------------------------------------
+	// Misc API
+	// ----------------------------------------------------------------
+
+	UFUNCTION(BlueprintPure, Category = "Network Physics")
+	int32 GetNetworkPredictionLOD() const { return NetworkPhysicsState.LocalLOD; }
+
 protected:
 
 	// Managed state should not be publically exposed
