Network Physics
-Disable LOD by default since transitions are still causing issues

#rnx

diff --git a/Engine/Plugins/Runtime/NetworkPrediction/Source/NetworkPrediction/Private/NetworkPhysics.cpp b/Engine/Plugins/Runtime/NetworkPrediction/Source/NetworkPrediction/Private/NetworkPhysics.cpp
--- a/Engine/Plugins/Runtime/NetworkPrediction/Source/NetworkPrediction/Private/NetworkPhysics.cpp
+++ b/Engine/Plugins/Runtime/NetworkPrediction/Source/NetworkPrediction/Private/NetworkPhysics.cpp
@@ -47,7 +47,7 @@ namespace UE_NETWORK_PHYSICS
 	int32 FixedLocalFrameOffsetTolerance = 3;
 	FAutoConsoleVariableRef CVarFixedLocalFrameOffsetTolerance(TEXT("np2.FixedLocalFrameOffsetTolerance"), FixedLocalFrameOffsetTolerance, TEXT("When > 0, use hardcoded frame offset on client from head"));
 
-	int32 EnableLOD = 1;
+	int32 EnableLOD = 0;
 	FAutoConsoleVariableRef CVarEnableLOD(TEXT("np2.EnableLOD"), EnableLOD, TEXT("Enable local LOD mode"));
 
 	float LODDistance = 2400.f;
