Network Prediction
Async generalized version WIP

#rnx

diff --git a/Engine/Plugins/Runtime/NetworkPrediction/Source/NetworkPrediction/Private/NetworkPredictionAsync.cpp b/Engine/Plugins/Runtime/NetworkPrediction/Source/NetworkPrediction/Private/NetworkPredictionAsync.cpp
new file mode 100644
--- /dev/null
+++ b/Engine/Plugins/Runtime/NetworkPrediction/Source/NetworkPrediction/Private/NetworkPredictionAsync.cpp
@@ -0,0 +1,258 @@
+// Copyright Epic Games, Inc. All Rights Reserved.
+
+#include "NetworkPredictionAsync.h"
+#include "Net/UnrealNetwork.h"
+#include "PhysicsProxy/SingleParticlePhysicsProxy.h"
+
+// ---------------------------------------------------------------------
+
+namespace UE_NETWORK_PHYSICS
+{
+	bool TempMockDebug=true;
+	FAutoConsoleVariableRef CVarTempMockDebug(TEXT("np2.Mock.Debug"), TempMockDebug, TEXT("Enabled spammy log debugging of mock physics object state"));
+
+	bool bEnableAsync=false;
+	FAutoConsoleVariableRef CVarbEnableAsync(TEXT("np2.Async.Enable"), bEnableAsync, TEXT("Enable Async implementation"));
+}
+
+struct FTempMockObjModlDef : public UE_NP::TNetworkPredictionModelDefAsync<FTempMockObjModlDef>
+{
+	using InputCmdType = FTempMockInputCmd;
+	using ObjStateType = FTempMockObject;
+	
+	using ObjKeyType = UTempMockComponent*;
+	static const TCHAR* GetName() { return TEXT("TempMockObj"); }
+};
+
+bool FNetworkPredictionProxyAsync::NetSerialize(FArchive& Ar, class UPackageMap* Map, bool& bOutSuccess)
+{
+	if (!UE_NETWORK_PHYSICS::bEnableAsync)
+	{
+		bOutSuccess = true;
+		return true;
+	}
+
+	if (ensure(NetSerializeFunc))
+	{
+		NetSerializeFunc(Ar);
+		bOutSuccess = true;
+		return true;
+	}
+
+	bOutSuccess = false;
+	return false;
+}
+
+void FTempMockObject::SimulationTick(const FTempMockInputCmd* InputCmd, FTempMockObject* SimObject, const float DeltaSeconds, const int32 Frame, const int32 LocalStorageFrame)
+{
+	// UE_LOG(LogTemp, Warning, TEXT("SimulationTick. %d. Proxy: 0x%X"), Frame, (int64)SimObject->Proxy);
+	
+	if (ensure(SimObject->Proxy))
+	{
+		if (auto* PT = SimObject->Proxy->GetPhysicsThreadAPI())
+		{
+			if (InputCmd->bBrakesPressed)
+			{
+				PT->SetV( Chaos::FVec3(0.f));				
+				PT->SetW( Chaos::FVec3(0.f));
+
+				UE_CLOG(UE_NETWORK_PHYSICS::TempMockDebug, LogNetworkPhysics, Log, TEXT("[%d/%d] Applied Break. Rot was: %s"), Frame, LocalStorageFrame, *PT->R().Rotator().ToString());
+				Chaos::FRotation3 NewR = PT->R();
+				FRotator Rot = NewR.Rotator();
+				Rot.Pitch = 0.f;
+				Rot.Roll = 0.f;
+				PT->SetR( FQuat(Rot));
+			}
+			else
+			{
+				if (InputCmd->Force.SizeSquared() > 0.001f)
+				{
+					//UE_LOG(LogTemp, Warning, TEXT("Applied Force. Frame: %d"), LocalFrame);
+					//UE_LOG(LogTemp, Warning, TEXT("0x%X ForceMultiplier: %f (Rand: %d)"), (int64)Proxy, GT_State.ForceMultiplier, GT_State.RandValue);
+					//UE_CLOG(PC == nullptr && InputCmd.Force.SizeSquared() > 0.f, LogTemp, Warning, TEXT("0x%X ForceMultiplier: %f (Rand: %d)"), (int64)Proxy, GT_State.ForceMultiplier, GT_State.RandValue);
+					PT->AddForce(InputCmd->Force * SimObject->ForceMultiplier);
+				}
+				
+				if (FMath::Abs<float>(InputCmd->Turn) > 0.001f)
+				{
+					PT->AddTorque(Chaos::FVec3(0.0f, 0.0f, InputCmd->Turn * SimObject->ForceMultiplier * 10.f));
+				}				
+			}
+
+			if (InputCmd->bJumpedPressed)
+			{
+				SimObject->JumpCount++;
+				UE_CLOG(UE_NETWORK_PHYSICS::TempMockDebug, LogNetworkPhysics, Log, TEXT("[%d/%d] 0x%X. bJumpedPressed: %d. Count: %d"), Frame, LocalStorageFrame, (int64)SimObject, InputCmd->bJumpedPressed, SimObject->JumpCount);
+			}
+
+			//UE_CLOG(PC == nullptr, LogTemp, Warning, TEXT("NP [%d] bJumpedPressed: %d. Count: %d"), LocalFrame, InputCmd.bJumpedPressed, PT_State.JumpCount);
+			if (InputCmd->bJumpedPressed && SimObject->JumpCooldownMS == 0)
+			{
+				//UE_LOG(LogTemp, Warning, TEXT("0x%X Jump! Total: %d"), (int64)Proxy, (int32)(TotalSeconds*1000.f));
+				
+				PT->SetLinearImpulse( Chaos::FVec3(0.f, 0.f, 115000.f) );
+
+				SimObject->JumpCooldownMS = 1000;
+				UE_CLOG(UE_NETWORK_PHYSICS::TempMockDebug, LogNetworkPhysics, Log, TEXT("[%d/%d] Applied Jump and reset cooldown"), Frame, LocalStorageFrame);
+			}
+			else
+			{
+				SimObject->JumpCooldownMS = FMath::Max( SimObject->JumpCooldownMS - (int32)(DeltaSeconds* 1000.f), 0);
+				if (SimObject->JumpCooldownMS != 0)
+				{
+					UE_CLOG(UE_NETWORK_PHYSICS::TempMockDebug, LogNetworkPhysics, Log, TEXT("[%d/%d] JumpCount: %d. JumpCooldown: %d"), Frame, LocalStorageFrame, SimObject->JumpCount, SimObject->JumpCooldownMS);
+				}
+			}
+		}
+	}
+}
+
+
+UTempMockComponent::UTempMockComponent()
+{
+	bWantsInitializeComponent = true;
+	bAutoActivate = true;
+	SetIsReplicatedByDefault(true);
+
+	PrimaryComponentTick.bCanEverTick = true;
+	PrimaryComponentTick.TickGroup = TG_DuringPhysics;
+	PrimaryComponentTick.bStartWithTickEnabled = true;
+}
+
+void UTempMockComponent::InitializeComponent()
+{
+	Super::InitializeComponent();
+
+#if WITH_CHAOS
+	if (!UE_NETWORK_PHYSICS::bEnableAsync)
+	{
+		return;
+	}
+
+	UWorld* World = GetWorld();	
+	checkSlow(World);
+	UNetworkPhysicsManager* Manager = World->GetSubsystem<UNetworkPhysicsManager>();
+	if (!Manager)
+	{
+		return;
+	}
+
+	UPrimitiveComponent* PrimitiveComponent = nullptr;
+	if (AActor* MyActor = GetOwner())
+	{
+		if (UPrimitiveComponent* RootPrimitive = Cast<UPrimitiveComponent>(MyActor->GetRootComponent()))
+		{
+			PrimitiveComponent = RootPrimitive;
+		}
+		else if (UPrimitiveComponent* FoundPrimitive = MyActor->FindComponentByClass<UPrimitiveComponent>())
+		{
+			PrimitiveComponent = FoundPrimitive;
+		}
+	}
+
+	if (ensureMsgf(PrimitiveComponent, TEXT("No PrimitiveComponent found on %s"), *GetPathName()))
+	{
+		NetworkPhysicsState.Proxy = PrimitiveComponent->BodyInstance.ActorHandle;
+		NetworkPhysicsState.OwningActor = GetOwner();
+		ensure(NetworkPhysicsState.OwningActor);
+
+		Manager->RegisterPhysicsProxy(&NetworkPhysicsState);
+
+		Manager->RegisterPhysicsProxyDebugDraw(&NetworkPhysicsState, [this](const UNetworkPhysicsManager::FDrawDebugParams& P)
+		{
+			AActor* Actor = this->GetOwner();
+			FBox LocalSpaceBox = Actor->CalculateComponentsBoundingBoxInLocalSpace();
+			const float Thickness = 2.f;
+
+			FVector ActorOrigin;
+			FVector ActorExtent;
+			LocalSpaceBox.GetCenterAndExtents(ActorOrigin, ActorExtent);
+			ActorExtent *= Actor->GetActorScale3D();
+			DrawDebugBox(P.DrawWorld, P.Loc, ActorExtent, P.Rot, P.Color, false, P.Lifetime, 0, Thickness);
+		});
+
+		
+		UE_NP::TNetworkPredictionAsyncSystemManager<FTempMockObjModlDef>* MockManager = UE_NP::TNetworkPredictionAsyncSystemManager<FTempMockObjModlDef>::Get(World);
+		checkSlow(MockManager);
+
+		FTempMockObject InitialState;
+		InitialState.Proxy = PrimitiveComponent->BodyInstance.ActorHandle;
+		MockManager->RegisterNewInstance(this, InitialState);
+
+		ReplicatedObject.InitProxy(MockManager, this);
+	}
+
+#endif
+}
+
+void UTempMockComponent::EndPlay(const EEndPlayReason::Type EndPlayReason)
+{
+	Super::EndPlay(EndPlayReason);
+
+#if WITH_CHAOS
+	if (!UE_NETWORK_PHYSICS::bEnableAsync)
+	{
+		return;
+	}
+
+	if (UWorld* World = GetWorld())
+	{
+		if (UNetworkPhysicsManager* Manager = World->GetSubsystem<UNetworkPhysicsManager>())
+		{
+			Manager->UnregisterPhysicsProxy(&NetworkPhysicsState);
+
+			if (auto* MockManager = UE_NP::TNetworkPredictionAsyncSystemManager<FTempMockObjModlDef>::Get(World))
+			{
+				MockManager->UnregisterInstance(this);
+			}
+		}
+	}
+#endif
+}
+
+void UTempMockComponent::GetLifetimeReplicatedProps(TArray<FLifetimeProperty>& OutLifetimeProps) const
+{
+	Super::GetLifetimeReplicatedProps(OutLifetimeProps);
+	DOREPLIFETIME( UTempMockComponent, NetworkPhysicsState);
+	DOREPLIFETIME( UTempMockComponent, ReplicatedObject);
+}
+
+void UTempMockComponent::TickComponent(float DeltaTime, enum ELevelTick TickType, FActorComponentTickFunction *ThisTickFunction)
+{
+	Super::TickComponent(DeltaTime, TickType, ThisTickFunction);
+
+#if WITH_CHAOS
+	if (!UE_NETWORK_PHYSICS::bEnableAsync)
+	{
+		return;
+	}
+
+	APlayerController* PC = GetOwnerPC();
+	
+	// Detect changes in possession (this isn't ideal: would rather register a delegate to tell us when this happens on the pawn but it doesn't exist)
+	if (CachedPC != PC)
+	{
+		if (UE_NP::TNetworkPredictionAsyncSystemManager<FTempMockObjModlDef>* MockManager = UE_NP::TNetworkPredictionAsyncSystemManager<FTempMockObjModlDef>::Get(GetWorld()))
+		{
+			MockManager->RegisterController(this, PC, (PC && PC->IsLocalController()) ? &PendingInputCmd : nullptr);
+		}
+		CachedPC = PC;
+	}
+
+	if (PC && PC->IsLocalController())
+	{
+		// Broadcast out a delegate. The user will write to PendingInputCmd
+		OnGeneratedLocalInputCmd.Broadcast();
+	}
+#endif
+}
+
+APlayerController* UTempMockComponent::GetOwnerPC() const
+{
+	if (APawn* PawnOwner = Cast<APawn>(GetOwner()))
+	{
+		return Cast<APlayerController>(PawnOwner->GetController());
+	}
+
+	return nullptr;
+}
\ No newline at end of file
diff --git a/Engine/Plugins/Runtime/NetworkPrediction/Source/NetworkPrediction/Public/NetworkPredictionAsync.h b/Engine/Plugins/Runtime/NetworkPrediction/Source/NetworkPrediction/Public/NetworkPredictionAsync.h
new file mode 100644
--- /dev/null
+++ b/Engine/Plugins/Runtime/NetworkPrediction/Source/NetworkPrediction/Public/NetworkPredictionAsync.h
@@ -0,0 +1,1138 @@
+// Copyright Epic Games, Inc. All Rights Reserved.
+
+#pragma once
+#include "Engine/World.h"
+#include "Physics/Experimental/PhysScene_Chaos.h"
+#include "PBDRigidsSolver.h"
+#include "ChaosSolversModule.h"
+#include "RewindData.h"
+#include "EngineUtils.h"
+#include "NetworkPhysics.h"
+#include "NetworkPredictionCheck.h"
+#include "NetworkPredictionLog.h"
+#include "Algo/Sort.h"
+#include "Algo/BinarySearch.h"
+#include "GameFramework/PlayerController.h"
+#include "Chaos/SimCallbackObject.h"
+
+#include "NetworkPredictionAsync.generated.h"
+
+class APlayerController;
+
+// Generalized version of Async NetworkPrediction. WIP
+namespace UE_NP {
+
+// ---------------------------------------------------
+
+template<typename T>
+struct TNetworkPredictionModelDefAsync
+{
+	// --------------------------------------------------------------
+	// Option 1: A single "GameObjType" for the entire thing. Don't impose any rules about client/server of GT/PT ownership.
+	//	Just marshal deltas or TFunctions to modify state. Let user implemented functions carve things up for effeciency.
+	// --------------------------------------------------------------
+
+	using InputCmdType = void;
+	using ObjStateType = void;
+	using OutStateType = void;
+
+	using ObjKeyType = AActor*;
+	using ControllerType = APlayerController*;
+
+	struct ManagedObjType
+	{
+		typename T::InputCmdType Input;
+		typename T::ObjStateType State;
+	};
+};
+
+// GT view of latest state
+template<typename ModelDef>
+struct TStateDataView
+{
+	using InputCmdType = typename ModelDef::InputCmdType;
+	using ObjStateType = typename ModelDef::ObjStateType;
+	using OutStateType = typename ModelDef::OutStateType;
+
+	int32 Frame;
+	InputCmdType* InputCmd;	// Will usually point into local 
+	ObjStateType* ObjState;
+};
+
+// Something like this is needed and should be managed by the async physics system for persistent PT storage
+// Note that the way this is used now, is TPersistentStorage<TSortedMap<ObjKey, ObjState>> - e.g, the whole snapshot
+// We could maybe provide a way to do sparse allocated per object storage.
+template<typename T>
+struct TPersistentStorage
+{
+	int32 GetTailFrame() const { return FMath::Max(0, HeadFrame - Buffer.Num() + 1); }
+	int32 GetHeadFrame() const { return HeadFrame; }
+
+	// Const function that returns the frame that is safe to write to.
+	// (Something like) this is the only thing ISimCallbackObject should really be using
+	T* GetWritable() const { return const_cast<T*>(&Buffer[HeadFrame % Buffer.Num()]); }
+
+	T* Get(int32 Frame)
+	{
+		if (Frame <= HeadFrame && Frame >= GetTailFrame())
+		{
+			return &Buffer[Frame % Buffer.Num()];
+		}
+
+		return nullptr;
+	}
+	
+	void IncrementFrame(int32 Frame, const bool bAutoCopyPre = true)
+	{
+		if (bAutoCopyPre && (HeadFrame+1 == Frame))
+		{
+			T* Pre = GetWritable();
+			HeadFrame++;
+			T* Next = GetWritable();
+			*Next = *Pre;
+		}
+		else
+		{
+			HeadFrame = Frame;
+		}
+	}
+
+	void RollbackToFrame(int32 Frame)
+	{
+		npEnsure(Frame >= GetTailFrame());
+		HeadFrame = Frame;
+	}
+
+private:
+
+	int32 HeadFrame = 0;
+	TStaticArray<T, 64> Buffer;
+};
+
+template<typename T, int32 Size=16>
+struct TStorageBuffer
+{
+	int32 HeadFrame() const { return LastWritten; }
+	int32 TailFrame() const { return FMath::Max(0, LastWritten - Buffer.Num() + 1); }
+
+	T& Write(int32 Frame) { LastWritten = Frame; return Buffer[Frame % Buffer.Num()]; }
+	const T& Get(int32 Frame) const { return Buffer[Frame % Buffer.Num()]; }
+	T& Get(int32 Frame) { return Buffer[Frame % Buffer.Num()]; }
+
+private:
+	int32 LastWritten = INDEX_NONE;
+	TStaticArray<T, Size> Buffer;
+};
+
+
+
+inline FString ToString(AActor* ObjKey)
+{
+	return GetNameSafe(ObjKey);
+}
+
+inline FString ToString(void* ObjKey)
+{
+	return FString::Printf(TEXT("0x%X"), (int64)ObjKey);
+}
+
+inline int32 ClampRewindFrame(const int32 CurrentRewindToFrame, const int32 NewRewindFrame)
+{
+	return CurrentRewindToFrame == INDEX_NONE ? NewRewindFrame : FMath::Min(CurrentRewindToFrame, NewRewindFrame);
+}
+
+// return ifs PC (ControllerType) is locally controlled, called only from client context.
+// E.g, we won't call this version on server. So the existence of a non null PC implies local controlled.
+inline bool IsLocallyControlled_Client(APlayerController* PC)
+{
+	return PC != nullptr;
+}
+
+
+template<typename ModelDef>
+struct TSnapshot
+{
+	using InputCmdType = typename ModelDef::InputCmdType;
+	using ObjStateType = typename ModelDef::ObjStateType;
+	using ManagedObjType = typename ModelDef::ManagedObjType;
+
+	using ObjKeyType = typename ModelDef::ObjKeyType;
+	using ControllerType = typename ModelDef::ControllerType;
+
+	int32 Frame = INDEX_NONE;
+	int32 LocalFrameOffset = 0;
+	TSortedMap<ObjKeyType, ManagedObjType> Objects;
+
+	void Reset()
+	{
+		Frame = INDEX_NONE;
+		LocalFrameOffset = 0;
+		Objects.Reset();
+	}
+};
+
+template<typename ModelDef>
+struct TReconcileSnapshot
+{
+	using InputCmdType = typename ModelDef::InputCmdType;
+	using ObjStateType = typename ModelDef::ObjStateType;
+	using ManagedObjType = typename ModelDef::ManagedObjType;
+
+	using ObjKeyType = typename ModelDef::ObjKeyType;
+	using ControllerType = typename ModelDef::ControllerType;
+
+	int32 Frame = INDEX_NONE;
+	int32 LocalFrameOffset = 0;
+
+	TSortedMap<ObjKeyType, TTuple<ManagedObjType, TArray<InputCmdType>>> Objects;
+};
+
+// ---------------------------------------------------
+
+template<typename ModelDef>
+struct TNetworkPredictionAsyncInput: public Chaos::FSimCallbackInput
+{
+	using InputCmdType = typename ModelDef::InputCmdType;
+	using ObjStateType = typename ModelDef::ObjStateType;
+	using ManagedObjType = typename ModelDef::ManagedObjType;
+
+	using ObjKeyType = typename ModelDef::ObjKeyType;
+	using ControllerType = typename ModelDef::ControllerType;
+
+	TArray<TPair<ObjKeyType, TUniqueFunction<void(ObjStateType&)>>> MarshalledMods; // Mods the GT want to make on our object state
+	TArray<ObjKeyType> MarshalledDeletes; // GT informing us to delete these objects from PT
+	TArray<TTuple<ObjKeyType, InputCmdType>> MarshalledInputCmds; // InputCmds from the GT, to be used in this frame
+
+	bool NetSendInputCmd(FNetBitWriter& Ar) override // FIXME: remove these functions
+	{
+		return true;
+	}
+
+	bool NetRecvInputCmd(APlayerController* PC, FNetBitReader& Ar) override // FIXME: remove these functions
+	{
+		return true;
+	}
+
+	void Reset()
+	{
+		MarshalledMods.Reset();
+		MarshalledDeletes.Reset();
+	}
+};
+
+template<typename ModelDef>
+struct TNetworkPredictionAsyncOutput: public Chaos::FSimCallbackOutput
+{
+	void Reset()
+	{
+		
+	}
+};
+
+
+
+template<typename ModelDef>
+class TNetworkPredictionAsyncCallback: public Chaos::TSimCallbackObject<TNetworkPredictionAsyncInput<ModelDef>,TNetworkPredictionAsyncOutput<ModelDef>>
+{
+public:
+
+	using InputCmdType = typename ModelDef::InputCmdType;
+	using ObjStateType = typename ModelDef::ObjStateType;
+	using ManagedObjType = typename ModelDef::ManagedObjType;
+
+	using ObjKeyType = typename ModelDef::ObjKeyType;
+	using ControllerType = typename ModelDef::ControllerType;
+
+	virtual void OnPreSimulate_Internal() override
+	{
+		const TNetworkPredictionAsyncInput<ModelDef>* Input = this->GetConsumerInput_Internal();
+		if (Input == nullptr)
+		{
+			return;
+		}
+		
+		Chaos::FPhysicsSolver* PhysicsSolver = static_cast<Chaos::FPhysicsSolver*>(this->GetSolver());
+		const int32 LocalFrame = PhysicsSolver->GetCurrentFrame();
+		const int32 SimulationFrame = LocalFrame - LocalFrameOffset_Internal; // FIXME: should be part of ISimCallbackObject/PersistentStorage API somehow
+		
+		TSortedMap<ObjKeyType, ManagedObjType>* ObjectMap = PersistentStorage_Internal->GetWritable();
+		npCheckSlow(ObjectMap);
+	
+		// -----------------------------------------
+		// 1. Apply data from GT inputs
+		// -----------------------------------------
+		for (auto& ModIt : Input->MarshalledMods)
+		{
+			const ObjKeyType& ObjKey = ModIt.Key;
+			ManagedObjType& InputObj = ObjectMap->FindOrAdd(ObjKey);
+			ModIt.Value(InputObj.State); // maybe frame should be passed here too?
+		}
+
+		for (const TTuple<ObjKeyType, InputCmdType>& InputCmd : Input->MarshalledInputCmds)
+		{
+			ManagedObjType& InputObj = ObjectMap->FindOrAdd(InputCmd.Key);
+			InputObj.Input = InputCmd.Value;
+		}
+
+		for (const ObjKeyType& ObjKey : Input->MarshalledDeletes)
+		{
+			ObjectMap->Remove(ObjKey);
+		}
+
+		// -----------------------------------------
+		// 2. Network Sync: take snapshot (server) and apply corrections (client)
+		// -----------------------------------------
+
+		if (NetworkSyncFunc)
+		{
+			NetworkSyncFunc(LocalFrame);
+		}
+
+		// -----------------------------------------
+		// 3. Actually Run the Tick
+		// -----------------------------------------			
+
+		const float DeltaTimeSeconds = this->GetDeltaTime_Internal();
+		for (auto& MapIt : *ObjectMap)
+		{
+			ManagedObjType& GameObj = MapIt.Value;
+			ObjStateType::SimulationTick(&GameObj.Input, &GameObj.State, DeltaTimeSeconds, SimulationFrame, LocalFrame);
+		}
+	}
+
+	virtual void OnContactModification_Internal(const TArrayView<Chaos::FPBDCollisionConstraintHandleModification>& Modifications) override
+	{
+	}
+
+	int32 TriggerRewindIfNeeded_Internal(int32 LastCompletedStep)
+	{
+		return INDEX_NONE;
+	}
+
+	void ApplyCorrections_Internal(int32 PhysicsStep, Chaos::FSimCallbackInput* Input) override
+	{
+		// Not needed?
+	}
+
+	void FirstPreResimStep_Internal(int32 PhysicsStep) override
+	{
+		// Not needed?
+	}		
+
+	const TPersistentStorage<TSortedMap<ObjKeyType, ManagedObjType>>* PersistentStorage_Internal = nullptr;
+	TUniqueFunction<void(int32)> NetworkSyncFunc;
+	int32 LocalFrameOffset_Internal = 0;
+};
+
+// -----------------------------------------------------------------------------------------
+// ModelDef based way to define a gameplay simulation that is:
+//	1. Async: It runs on a ISimCallbackObject which you do not have to define.
+//	2. Rollback Networking: Controller and Object State are networked and trigger rollbacks/resims.
+//  3. Object based. Meaning this simulation is intended to manage many instances of objects with distinct lifetimes.
+//
+//
+// The way the gameplay simulation is defined is:
+//	1. Server-Authoratative ObjectState (the server is authoratative over this state)
+//	2. Client-Authoratative InputCmd (a controlling client is authoratative over this state) [This can be optional]
+//	3. SimulationTick function: a static function that is given input/output state that runs within this framework for you.
+// 
+// 
+// If someone *doesn't* want to use TNetworkPredictionAsyncSystemManager to define thier gameplay simulation, this is what they are left with:
+//	1. If they just define an ISimCallbackObject, they will still particpate in physics rollbacks. E.g they will resim frames but not trigger them.
+//		* They can marshall their data via Input/Output objects however they want between threads.
+//		* They can run their 'Tick' however they want (E.g, their OnPreSimulate_Internal logic)
+//		* They *must* use a form of Persistent Storage which is rollback friendly (TODO)
+//			* This is enforced via const on the ISimCallback API (TODO)
+//			* Persistent storage will provide an API to encapsulate the historic buffering when rollback is running (TODO)
+//
+//	2. If they do want to trigger rewinds, they will implement their own INetworkPhysicsSubsystem.
+//
+//	
+//
+// 
+//	We should support some generic form of networking these with custom channels
+//	replicate the entire snapshots atomically. Use unreal replication to bind 
+//		this will never work due to raw pointers etc that will be used?
+//
+// -----------------------------------------------------------------------------------------
+template<typename ModelDef>
+class TNetworkPredictionAsyncSystemManager: public INetworkPhysicsSubsystem
+{
+public:
+
+	using InputCmdType = typename ModelDef::InputCmdType;
+	using ObjStateType = typename ModelDef::ObjStateType;
+	using ManagedObjType = typename ModelDef::ManagedObjType;
+
+	using ObjKeyType = typename ModelDef::ObjKeyType;
+	using ControllerType = typename ModelDef::ControllerType;
+
+	static FName GetName() { return FName(ModelDef::GetName()); }
+	static TNetworkPredictionAsyncSystemManager<ModelDef>* Get(UWorld* World)
+	{
+		if (UNetworkPhysicsManager* NetworkPhysicsManager = World->GetSubsystem<UNetworkPhysicsManager>())
+		{
+			if (TNetworkPredictionAsyncSystemManager<ModelDef>* Existing = NetworkPhysicsManager->GetSubsystem<TNetworkPredictionAsyncSystemManager<ModelDef>>())
+			{
+				return Existing;
+			}
+
+			TNetworkPredictionAsyncSystemManager<ModelDef>* NewInstance = NetworkPhysicsManager->RegisterSubsystem<TNetworkPredictionAsyncSystemManager<ModelDef>>(MakeUnique<TNetworkPredictionAsyncSystemManager<ModelDef>>(World));
+			return NewInstance;
+		}
+		ensure(false);
+		return nullptr;
+	}
+	
+	TNetworkPredictionAsyncSystemManager(UWorld* InWorld)
+	{
+		World = InWorld;
+		if (ensure(World))
+		{
+			FPhysScene* PhysScene = World->GetPhysicsScene();
+			if (ensureAlways(PhysScene))
+			{
+				Chaos::FPhysicsSolver* Solver = PhysScene->GetSolver();
+				if (ensureAlways(Solver))
+				{
+					if (Solver->GetRewindCallback())
+					{
+						AsyncCallback = Solver->CreateAndRegisterSimCallbackObject_External<TNetworkPredictionAsyncCallback<ModelDef>>(true, true);
+						AsyncCallback->PersistentStorage_Internal = &PersistentStorage_Internal;
+						AsyncCallback->NetworkSyncFunc = [this](int32 F) { this->NetworkSyncPersistentStorage_Internal(F); };
+					}
+					else
+					{
+						UE_LOG(LogNetworkPhysics, Warning, TEXT("Rewind not enabled on Physics solver. FMockObjectManager will be disabled"));
+					}
+				}
+			}
+		}
+	}
+
+	~TNetworkPredictionAsyncSystemManager()
+	{
+
+	}
+
+	void PostNetRecv(UWorld* InWorld, int32 LocalOffset, int32 LastProcessedFrame) override
+	{
+		
+	}
+
+	void PreNetSend(UWorld* InWorld, float DeltaSeconds) override
+	{
+		if (SnapshotsForGT_External.IsEmpty() == false)
+		{
+			while(SnapshotsForGT_External.Dequeue(LatestSnapshot));	
+
+			//OutputMap
+		}
+	}
+
+	// Serialize the authoratative object state
+	//		-Receiving side must place it somewhere to be reconciled
+	//		-Sending side should send latest available copies
+	// Fixme: want differnet options here for serializing the Input/State structs (and maybe support for Output struct too)
+	void NetSerializeGameObject(FArchive& Ar, ObjKeyType ObjKey)
+	{
+		if (Ar.IsSaving())
+		{
+			// 1. Frame
+			Ar << LatestSnapshot.Frame;
+			if (LatestSnapshot.Frame == -1)
+			{
+				return;
+			}
+
+			// 2. Sim State
+			if (ManagedObjType* Obj = LatestSnapshot.Objects.Find(ObjKey))
+			{
+				Obj->State.NetSerialize(Ar);
+			}
+			else
+			{
+				UE_LOG(LogNetworkPrediction, Warning, TEXT("No copy of object %s found in LatestSnapshot. Frame: %d"), *UE_NP::ToString(ObjKey), LatestSnapshot.Frame);
+				ManagedObjType DefaultObj;
+				DefaultObj.State.NetSerialize(Ar);
+			}
+
+			// 3. Controller State
+			uint8 NumInputs = 0;
+			if (auto* ControllerData = ControllerInputMap.Find(ObjKey))
+			{
+				auto& InputBuffer = ControllerData->template Get<2>();
+				npEnsureMsgf(InputBuffer.TailFrame() <= LatestSnapshot.Frame, TEXT("Stored InputCmd no longer valid for replication. Snapshot Frame: %d. Tail Frame: %d"), LatestSnapshot.Frame, InputBuffer.TailFrame());
+
+				if (InputBuffer.HeadFrame() == INDEX_NONE)
+				{
+					NumInputs = 0;
+					Ar << NumInputs;
+				}
+				else
+				{
+					const int32 MaxFutureInputs = 6;
+					NumInputs = FMath::Clamp<int32>(InputBuffer.HeadFrame() - LatestSnapshot.Frame + 1, 0, MaxFutureInputs);
+					Ar << NumInputs;
+
+					const int32 LastInput = LatestSnapshot.Frame + NumInputs;
+					for (int32 Frame = LatestSnapshot.Frame; Frame < LastInput; ++Frame)
+					{
+						npEnsureSlow(Frame <= InputBuffer.HeadFrame());
+						InputBuffer.Get(Frame).NetSerialize(Ar);
+					}
+				}
+			}
+			else
+			{
+				NumInputs = 0;
+				Ar << NumInputs;
+			}
+
+			NET_CHECKSUM(Ar);
+		}
+		else
+		{
+			// 1. Frame
+			int32 Frame = 0;
+			Ar << Frame;
+			if (Frame == -1)
+			{
+				return;
+			}
+
+			// 2. Sim State
+			auto& Recv = NetRecvData.FindOrAdd(ObjKey);
+			Recv.template Get<int32>() = Frame;
+			Recv.template Get<ManagedObjType>().State.NetSerialize(Ar);
+			
+			// 3. Controller State
+			TArray<InputCmdType>& FutureInputs = Recv.template Get<TArray<InputCmdType>>();
+
+			uint8 NumInputs = 0;
+			Ar << NumInputs;
+
+			FutureInputs.Reset(NumInputs);
+			FutureInputs.AddDefaulted(NumInputs);
+
+			if (NumInputs > 0)
+			{
+				Recv.template Get<ManagedObjType>().Input.NetSerialize(Ar);
+			}
+
+			for (int32 i=1; i < NumInputs; ++i)
+			{
+				FutureInputs[i].NetSerialize(Ar);
+			}
+
+			NET_CHECKSUM(Ar);
+		}
+	}
+
+	void ProcessInputs_External(int32 PhysicsStep, int32 LocalFrameOffset, bool& bOutSendClientInputCmd) override
+	{
+		const ENetMode NetMode = World->GetNetMode();
+
+		npCheckSlow(AsyncCallback);
+		TNetworkPredictionAsyncInput<ModelDef>* AsyncInput = AsyncCallback->GetProducerInputData_External();
+		npCheckSlow(AsyncInput);
+
+		// ---------------------------------------------------------------------
+		//	InputCmds
+		// ---------------------------------------------------------------------
+		
+		for (auto& MapIt : ControllerInputMap)
+		{
+			const ObjKeyType& ObjKey = MapIt.Key;
+			ControllerType& Controller = MapIt.Value.template Get<ControllerType>();
+			InputCmdType* PendingInputCmd = MapIt.Value.template Get<InputCmdType*>();
+			TStorageBuffer<InputCmdType, 16>& InputBuffer = MapIt.Value.template Get<2>();
+
+			APlayerController* PC = Controller;
+			if (!npEnsure(PC))
+			{
+				continue;
+			}
+
+			if (PendingInputCmd)
+			{
+				// Locally controlled, capture the PendingInputcmd
+				InputBuffer.Write(PhysicsStep) = *PendingInputCmd;
+
+				if (NetMode == NM_Client)
+				{
+					FNetBitWriter Writer(nullptr, 256 << 3);
+					PendingInputCmd->NetSerialize(Writer);
+
+					// Calling the RPC right here is bad. There probably is some intermediate buffer that these finalized commands should go to and the network
+					// sending can happen elsewhere. That is also where move combining / delta compressiong type logic could happen.
+					TArray<uint8> SendData;
+					if (Writer.IsError() == false)
+					{
+						int32 NumBytes = (int32)Writer.GetNumBytes();
+						SendData = MoveTemp(*const_cast<TArray<uint8>*>(Writer.GetBuffer()));
+						SendData.SetNum(NumBytes);
+					}
+				
+					PC->PushClientInput(PhysicsStep, SendData);
+					bOutSendClientInputCmd = true;
+
+				}
+			}
+			else if (npEnsure(NetMode != NM_Client))
+			{
+				// Server: Non locally controlled. look for a Cmd from the PlayerController InputCmd Bus
+				APlayerController::FServerFrameInfo& FrameInfo = PC->GetServerFrameInfo();
+				APlayerController::FInputCmdBuffer& PCInputBuffer = PC->GetInputBuffer();
+
+				auto NetSerializeCmd = [&PCInputBuffer](InputCmdType& InputCmd, int32 ClientFrmae)
+				{
+					InputCmd = InputCmdType();
+					const TArray<uint8>& Data = PCInputBuffer.Get(ClientFrmae);
+					if (Data.Num() > 0)
+					{
+						uint8* RawData = const_cast<uint8*>(Data.GetData());
+						FNetBitReader Ar(nullptr, RawData, ((int64)Data.Num()) << 3);
+						InputCmd.NetSerialize(Ar);
+					}
+				};
+
+				if (FrameInfo.LastProcessedInputFrame != INDEX_NONE)
+				{
+					if (FrameInfo.bFault)
+					{
+						// While in fault just write the current input to the buffer and don't bother with future inputs
+						if (InputBuffer.HeadFrame() != PhysicsStep)
+						{
+							UE_LOG(LogTemp, Log, TEXT("[S] In Fault. Cmd %d Frame %d"), FrameInfo.LastProcessedInputFrame, PhysicsStep);
+							NetSerializeCmd(InputBuffer.Write(PhysicsStep), FrameInfo.LastProcessedInputFrame);
+						}
+					}
+					else
+					{
+						// De-NetSerialize all new InputCmds from PC's buffer to our internal typed buffer
+						// We are going ahead of where we need to be so that we can replicate unprocessed
+						// cmds as'FutureInputs'.
+						const int32 CmdOffset = FrameInfo.LastProcessedInputFrame - PhysicsStep;
+						const int32 WriteStart = FMath::Max(PhysicsStep, InputBuffer.HeadFrame()+1);
+						const int32 WriteEnd = PCInputBuffer.HeadFrame() - CmdOffset;
+
+						for (int32 Frame = WriteStart; Frame <= WriteEnd; ++Frame)
+						{
+							//UE_LOG(LogTemp, Log, TEXT("[S] De-NetSerialize Cmd %d Frame %d. NumBytes: %d"), Frame + CmdOffset, Frame, PCInputBuffer.Get(Frame + CmdOffset).Num());
+							NetSerializeCmd(InputBuffer.Write(Frame), Frame + CmdOffset);
+						}
+					}
+				}
+				else
+				{
+					InputBuffer.Write(PhysicsStep) = InputCmdType();
+				}
+			}
+
+			npEnsureMsgf(InputBuffer.HeadFrame() >= PhysicsStep, TEXT("No local input found for frame %d. (Head: %d)"), PhysicsStep, InputBuffer.HeadFrame());
+			AsyncInput->MarshalledInputCmds.Emplace_GetRef(ObjKey, InputBuffer.Get(PhysicsStep));
+		}
+
+		// ---------------------------------------------------------------------
+		//	Networking Recv
+		// ---------------------------------------------------------------------
+
+		// Coalesce the loose NetRecvData into frame-based snapshots
+		TArray<TReconcileSnapshot<ModelDef>> CoalescedSnapshots;
+		for (auto& MapIt : NetRecvData)
+		{
+			const int32 ServerFrame = MapIt.Value.template Get<int32>();
+			const int32 Frame = ServerFrame + LocalFrameOffset;
+
+			TReconcileSnapshot<ModelDef>* DestSnapshot = Algo::FindBy(CoalescedSnapshots, Frame, &TReconcileSnapshot<ModelDef>::Frame);
+			if (!DestSnapshot)
+			{
+				DestSnapshot = &CoalescedSnapshots.Emplace_GetRef(TReconcileSnapshot<ModelDef>{Frame, LocalFrameOffset});
+			}
+
+			DestSnapshot->Objects.Add(MapIt.Key) = MakeTuple(MoveTemp(MapIt.Value.template Get<ManagedObjType>()), MoveTemp(MapIt.Value.template Get<TArray<InputCmdType>>()));
+		}
+		NetRecvData.Reset();
+		Algo::SortBy(CoalescedSnapshots, &TReconcileSnapshot<ModelDef>::Frame);
+		for (TReconcileSnapshot<ModelDef>& Snapshot : CoalescedSnapshots)
+		{
+			SnapshotForReconcile_Internal.Enqueue(MoveTemp(Snapshot));
+		}
+
+		// ---------------------------------------------------------------------
+		//	Object Data
+		// ---------------------------------------------------------------------
+		
+		// Marshall updated ControllerMap if necessary		
+		if (bControllerMapDirty)
+		{
+			TArray<ObjKeyType> LocallyControlledObjs;
+			for (auto& MapIt : ControllerInputMap)
+			{
+				if (MapIt.Value.template Get<InputCmdType*>() != nullptr)
+				{
+					LocallyControlledObjs.Add(MapIt.Key);
+				}
+			}
+
+			MarshalledLocalObjs.Enqueue(MoveTemp(LocallyControlledObjs));
+			bControllerMapDirty = false;
+		}
+	}
+
+	int32 TriggerRewindIfNeeded_Internal(int32 LastCompletedStep) override
+	{
+		// -----------------------------------------------------
+		//	Update locally controlled set
+		// -----------------------------------------------------
+		if (!MarshalledLocalObjs.IsEmpty())
+		{
+			TArray<ObjKeyType> MarshalledList;
+			while (this->MarshalledLocalObjs.Dequeue(MarshalledList)) { }
+			LocalObjs_Internal = MoveTemp(MarshalledList);
+		}
+
+		// -----------------------------------------------------
+		//	Detect Corrections
+		// -----------------------------------------------------
+		PendingCorrections_Internal.Reset();
+		PendingCorrectionIdx_Internal = INDEX_NONE;
+
+		FutureInputs_Internal.Reset();
+		
+		int32 RewindToFrame = INDEX_NONE;
+		const int32 MinFrame = PersistentStorage_Internal.GetTailFrame();
+
+		TReconcileSnapshot<ModelDef> Snapshot;
+		while (SnapshotForReconcile_Internal.Dequeue(Snapshot))
+		{
+			// Cache the most up to date FrameOffset for PT access
+			// This should probably be built into SimCallbackObject or RewindData or something
+			AsyncCallback->LocalFrameOffset_Internal = Snapshot.LocalFrameOffset;
+			
+			const TSortedMap<ObjKeyType, ManagedObjType>* LocalMap = PersistentStorage_Internal.Get(Snapshot.Frame);
+			if (!LocalMap)
+			{
+				UE_LOG(LogNetworkPrediction, Warning, TEXT("Snapshot too old to reconcile. Snapshot Frame: %d. MinFrame: %d"), Snapshot.Frame, MinFrame);
+				continue;
+			}
+
+			TSnapshot<ModelDef> CorrectionSnapshot;
+			for (auto& ObjIt : Snapshot.Objects)
+			{
+				ObjKeyType& ObjKey = ObjIt.Key;
+				ManagedObjType& AuthObj = ObjIt.Value.template Get<ManagedObjType>();
+				const ManagedObjType* LocalState = LocalMap->Find(ObjKey);
+
+				bool bShouldReconcile = false;
+				if (LocalState)
+				{
+					const bool bIsLocallyControlled = LocalObjs_Internal.Contains(ObjKey);
+					if (LocalState->State.ShouldReconcile(AuthObj.State, bIsLocallyControlled))
+					{
+						bShouldReconcile = true;
+					}
+					else if (!bIsLocallyControlled && LocalState->Input.ShouldReconcile(AuthObj.Input))
+					{
+						bShouldReconcile = true;
+					}
+				}
+				else 
+				{
+					bShouldReconcile = true;
+				}
+
+				if (bShouldReconcile)
+				{
+					RewindToFrame = UE_NP::ClampRewindFrame(RewindToFrame, Snapshot.Frame);
+					CorrectionSnapshot.Frame = Snapshot.Frame;
+					CorrectionSnapshot.LocalFrameOffset = Snapshot.LocalFrameOffset;
+					CorrectionSnapshot.Objects.Add(ObjKey, MoveTemp(AuthObj)); // FIXME: sorting on insert is useless here
+				}
+
+				// Even if we don't cause a correction, steal FutureInputs so they can be applied if we do end up resimulating
+				auto& FutureInputs = ObjIt.Value.template Get<TArray<InputCmdType>>();
+				if (FutureInputs.Num() > 0)
+				{
+					FutureInputs_Internal.Emplace(MakeTuple(Snapshot.Frame, ObjKey, MoveTemp(FutureInputs)));
+				}
+			}
+
+			if (CorrectionSnapshot.Objects.Num() > 0)
+			{
+				PendingCorrections_Internal.Emplace(MoveTemp(CorrectionSnapshot));
+				PendingCorrectionIdx_Internal = 0;
+			}
+		}
+		
+		return RewindToFrame;
+	}
+
+	void ProcessInputs_Internal(int32 PhysicsStep) override
+	{
+		PersistentStorage_Internal.IncrementFrame(PhysicsStep, true);
+	}
+
+	void PreResimStep_Internal(int32 PhysicsStep, bool bFirst)
+	{
+		if (bFirst)
+		{
+			PersistentStorage_Internal.RollbackToFrame(PhysicsStep);
+
+			// When we start a resim, even if we didn't cause the correction
+			// we still want to apply "more accurate inputs" when possible
+			// this will give us more accurate repridiections 
+			//
+			// The transofmrations the data has to go through are very gross
+			// I don't see a better way atm that doesn't compromise the golden path (non resim)
+			PendingInputCorrections_Internal.Reset();
+			for (auto& FutureInput : FutureInputs_Internal)
+			{
+				const int32 Frame = FutureInput.template Get<int32>();
+				const ObjKeyType ObjKey = FutureInput.template Get<ObjKeyType>();
+				TArray<InputCmdType> FutureInputs = FutureInput.template Get<TArray<InputCmdType>>();
+
+				for (InputCmdType& InputCmd : FutureInputs)
+				{
+					TTuple<int32, TSortedMap<ObjKeyType, InputCmdType>>* InputCorrectionSnapshot = PendingInputCorrections_Internal.FindByPredicate([Frame](const TTuple<int32, TSortedMap<ObjKeyType, InputCmdType>>& T) { return T.template Get<int32>() == Frame; });
+					if (!InputCorrectionSnapshot)
+					{
+						InputCorrectionSnapshot = &PendingInputCorrections_Internal.AddDefaulted_GetRef();
+						InputCorrectionSnapshot->template Get<int32>() = Frame;
+					}
+
+					InputCorrectionSnapshot->template Get<1>().FindOrAdd(ObjKey) = InputCmd;	
+				}
+			}
+		}
+	}
+
+	void NetworkSyncPersistentStorage_Internal(int32 PhysicsStep)
+	{
+		// Called from the AsyncCallback's OnPreSimulate_Internal and is the point we synchronize the PersistentStorage state
+
+		// Copy PersistentStorage for GT (Server will send this snapshot to client)
+		TSortedMap<ObjKeyType, ManagedObjType>* ObjectMap = PersistentStorage_Internal.GetWritable();
+		SnapshotsForGT_External.Enqueue(TSnapshot<ModelDef>{PhysicsStep, AsyncCallback->LocalFrameOffset_Internal, *ObjectMap});
+
+		// Apply correctionn
+		if (PendingCorrections_Internal.IsValidIndex(PendingCorrectionIdx_Internal))
+		{
+			TSnapshot<ModelDef>& CorrectionSnapshot = PendingCorrections_Internal[PendingCorrectionIdx_Internal];
+			npEnsureMsgf(CorrectionSnapshot.Frame >= PhysicsStep, TEXT("Correction too late to apply? CorrectionFrame: %d PhysicsStep: %d"), CorrectionSnapshot.Frame, PhysicsStep);
+
+			if (CorrectionSnapshot.Frame == PhysicsStep)
+			{
+				PendingCorrectionIdx_Internal++;
+				TSortedMap<ObjKeyType, ManagedObjType>* InputMap = PersistentStorage_Internal.Get(PhysicsStep);
+				if (!npEnsureMsgf(InputMap, TEXT("Could not find storage for frame %d. Cannot apply correction! ValidRange: %d-%d"), PhysicsStep, PersistentStorage_Internal.GetTailFrame(), PersistentStorage_Internal.GetHeadFrame()))
+				{
+					return;
+				}
+
+				for (auto& CorrectionIt : CorrectionSnapshot.Objects)
+				{
+					const ObjKeyType& CorrectionKey = CorrectionIt.Key;
+					const ManagedObjType& CorrectionObj = CorrectionIt.Value;
+					
+					ManagedObjType& LocalObj = InputMap->FindOrAdd(CorrectionKey);
+
+					const bool bIsLocal = LocalObjs_Internal.Contains(CorrectionKey);
+					LocalObj.State.ApplyCorrection(CorrectionObj.State, bIsLocal);
+					
+					if (!bIsLocal)
+					{
+						LocalObj.Input.ApplyCorrection(CorrectionObj.Input);
+					}
+				}
+			}
+		}
+
+		// Input Corrections
+		for (auto& InputCorrection : PendingInputCorrections_Internal)
+		{
+			if (InputCorrection.template Get<int32>() == PhysicsStep)
+			{
+				TSortedMap<ObjKeyType, ManagedObjType>* InputMap = PersistentStorage_Internal.Get(PhysicsStep);
+				for (auto& MapIt : InputCorrection.template Get<1>())
+				{
+					if (ManagedObjType* Obj = InputMap->Find(MapIt.Key))
+					{
+						Obj->Input = MapIt.Value;
+					}
+				}
+				break;
+			}
+		}
+	}
+
+	// ---------------------------------------------
+
+	void RegisterNewInstance(ObjKeyType Key, const ObjStateType& InitialValue)
+	{
+		npCheckSlow(AsyncCallback);
+		TNetworkPredictionAsyncInput<ModelDef>* Input = AsyncCallback->GetProducerInputData_External();
+		npCheckSlow(Input);
+
+		Input->MarshalledMods.Emplace(Key, [InitialValue](ObjStateType& Obj){ Obj = InitialValue; });
+	}
+
+	void UnregisterInstance(ObjKeyType Key)
+	{
+		npCheckSlow(AsyncCallback);
+		TNetworkPredictionAsyncInput<ModelDef>* Input = AsyncCallback->GetProducerInputData_External();
+		npCheckSlow(Input);
+
+		Input->MarshalledDeletes.Emplace(Key);
+		
+		NetRecvInputCmds.Remove(Key);
+		ControllerInputMap.Remove(Key);
+	}
+
+	// ---------------------------------------------
+
+	void RegisterController(ObjKeyType Key, ControllerType Controller, InputCmdType* LocalPendingInputCmd = nullptr)
+	{		
+		bControllerMapDirty = true;
+		if (Controller)
+		{
+			ControllerInputMap.Add(Key) = MakeTuple(Controller, LocalPendingInputCmd, TStorageBuffer<InputCmdType, 16>());
+		}
+		else
+		{
+			ControllerInputMap.Remove(Key);
+		}
+	}
+
+private:
+
+	UWorld* World = nullptr;
+
+	TNetworkPredictionAsyncCallback<ModelDef>* AsyncCallback = nullptr;
+		
+	// -------------------------------------------------------
+
+	TSnapshot<ModelDef>	LatestSnapshot; // Latest snapshot marshalled from PT
+
+	TSortedMap<ObjKeyType, ObjStateType*> OutputMap; // Where we write data that has been marshalled PT->GT
+
+	TArray<TSnapshot<ModelDef>> PendingCorrections_Internal;
+	int32 PendingCorrectionIdx_Internal = INDEX_NONE;
+	TPersistentStorage<TSortedMap<ObjKeyType, ManagedObjType>> PersistentStorage_Internal;
+
+	TArray<TTuple<int32, ObjKeyType, TArray<InputCmdType>>> FutureInputs_Internal;
+	TArray<TTuple<int32, TSortedMap<ObjKeyType, InputCmdType>>> PendingInputCorrections_Internal;
+
+
+	
+	TSortedMap<ObjKeyType, TStorageBuffer<InputCmdType, 16>> NetRecvInputCmds;
+
+
+	// Set to hold ObjKeys that are locally controlled: this is just so we can know who is locally controlled on PT functions
+	// Consider replacing with some form of "ObjectFlags"
+	// (However ObjectFlags may need to really be per-object. Where as the set of locally controlled objects is small and we can minimize marshalling this way)
+	// Right now locally controlled on the physics thread is used to do InputDecay and determine if we need to reconcile/correct on the InputCmds or not.
+	TQueue<TArray<ObjKeyType>> MarshalledLocalObjs;
+	TArray<ObjKeyType> LocalObjs_Internal;
+
+	TQueue<TSnapshot<ModelDef>> SnapshotsForGT_External;
+
+	// ---------------------------------------------------------------------------------------
+
+	// Map for objects where a Controller exists.
+	//	ControllerType - the controller object (PlayerController)
+	//	InputCmdTpye* - the Pending Inputcmd. If locally controlled, where GameCode write latest input to
+	//  StorageBuffer - history of commands that have been submitted
+	TSortedMap<ObjKeyType, TTuple<ControllerType, InputCmdType*, TStorageBuffer<InputCmdType, 16>>> ControllerInputMap;
+	bool bControllerMapDirty = false;
+
+	// Replicated data is written directly here on GT. The array of InputCmdTypes are future inputs
+	TSortedMap<ObjKeyType, TTuple<int32, ManagedObjType, TArray<InputCmdType>>> NetRecvData;
+
+	// Marshalled data for PT to perform reconcile
+	TQueue<TReconcileSnapshot<ModelDef>> SnapshotForReconcile_Internal;
+};
+
+}; // namespace UE_NP
+
+// -------------------------------------------------------------------------------------------
+
+USTRUCT()
+struct FNetworkPredictionProxyAsync
+{
+	GENERATED_BODY()
+
+	template<typename ModelDef>
+	void InitProxy( UE_NP::TNetworkPredictionAsyncSystemManager<ModelDef>* SystemManager, typename ModelDef::ObjKeyType ObjKey)
+	{
+		NetSerializeFunc = [SystemManager, ObjKey](FArchive& Ar)
+		{
+			SystemManager->NetSerializeGameObject(Ar, ObjKey);
+		};
+	}
+
+	bool NetSerialize(FArchive& Ar, class UPackageMap* Map, bool& bOutSuccess);
+	
+	bool Identical(const FNetworkPredictionProxyAsync* Other, uint32 PortFlags) const
+	{
+		return false; // FIXME!!!!
+	}
+
+	TFunction<void(FArchive& Ar)> NetSerializeFunc;
+};
+
+template<>
+struct TStructOpsTypeTraits<FNetworkPredictionProxyAsync> : public TStructOpsTypeTraitsBase2<FNetworkPredictionProxyAsync>
+{
+	enum
+	{
+		WithNetSerializer = true,
+		WithNetSharedSerialization = true,
+		WithIdentical = true,
+	};
+};
+
+// -------------------------------------------------------------------------------------------
+
+USTRUCT(BlueprintType)
+struct FTempMockInputCmd
+{
+	GENERATED_BODY()
+
+	// Simple world vector force to be applied
+	UPROPERTY(BlueprintReadWrite,Category="Input")
+	FVector	Force;
+
+	UPROPERTY(BlueprintReadWrite,Category="Input")
+	float Turn=0.f;
+
+	UPROPERTY(BlueprintReadWrite,Category="Input")
+	bool bJumpedPressed = false;
+
+	UPROPERTY(BlueprintReadWrite,Category="Input")
+	bool bBrakesPressed = false;
+
+	void NetSerialize(FArchive& Ar)
+	{
+		Ar << Force;
+		Ar << Turn;
+		Ar << bJumpedPressed;
+		Ar << bBrakesPressed;
+	}
+
+	bool ShouldReconcile(const FTempMockInputCmd& AuthState) const
+	{
+		return
+			FVector::DistSquared(Force, AuthState.Force) > 0.1f || 
+			bJumpedPressed != AuthState.bJumpedPressed || 
+			bBrakesPressed != AuthState.bBrakesPressed || 
+			FMath::Abs<float>(Turn - AuthState.Turn) > 0.1f;
+	}
+
+	void ApplyCorrection(const FTempMockInputCmd& AuthState)
+	{
+		*this = AuthState;
+	}
+};
+
+USTRUCT(BlueprintType)
+struct FTempMockObject
+{
+	GENERATED_BODY()
+
+	UPROPERTY(BlueprintReadWrite,Category="Mock Object")
+	int32 JumpCooldownMS = 0;
+
+	// Number of frames jump has been pressed
+	UPROPERTY(BlueprintReadWrite,Category="Mock Object")
+	int32 JumpCount = 0;
+
+	UPROPERTY(BlueprintReadWrite,Category="Mock Object")
+	float ForceMultiplier = 125000.f;
+
+	// Arbitrary data that doesn't affect sim but could still trigger rollback
+	UPROPERTY(BlueprintReadWrite,Category="Mock Object")
+	int32 RandValue = 0;
+
+	void NetSerialize(FArchive& Ar)
+	{
+		Ar << ForceMultiplier;
+		Ar << RandValue;
+		Ar << JumpCooldownMS;
+		Ar << JumpCount;
+	}
+
+	bool ShouldReconcile(const FTempMockObject& AuthState, const bool bIsLocallyControlled) const
+	{
+		return ForceMultiplier != AuthState.ForceMultiplier || 
+			RandValue != AuthState.RandValue ||
+			JumpCooldownMS != AuthState.JumpCooldownMS || 
+			JumpCount != AuthState.JumpCount;
+	}
+
+	void ApplyCorrection(const FTempMockObject& AuthState, const bool bIsLocallyControlled)
+	{
+		ForceMultiplier = AuthState.ForceMultiplier;
+		RandValue = AuthState.RandValue;
+		JumpCooldownMS = AuthState.JumpCooldownMS;
+		JumpCount = AuthState.JumpCount;
+	}
+
+	static void SimulationTick(const FTempMockInputCmd* InputCmd, FTempMockObject* SimObject, const float DeltaSeconds, const int32 Frame, const int32 StorageFrame);
+
+	// -----------------------------------------------------
+	// Non-networked data
+	//	Should this be kept here or someplace else?
+	//	Having it here is convenient, and we can be sure that non-replicated data is preserved throughout the system. (E.g, we don't stomp out the entire struct anywhere)
+	//	But also means we can never expose "Set FTempMockObject" to the gameplay code. We need to expose individual "Set Property X" functions.
+	//	Or we could internally split this into "GameState/SystemState" structs and only expose the set GameState to higher level code.
+	// -----------------------------------------------------
+
+	FSingleParticlePhysicsProxy* Proxy = nullptr;
+};
+
+
+UCLASS(BlueprintType,meta=(BlueprintSpawnableComponent))
+class UTempMockComponent: public UActorComponent
+{
+public:
+	GENERATED_BODY()
+
+	UTempMockComponent();
+
+	virtual void InitializeComponent() override;
+	virtual void EndPlay(const EEndPlayReason::Type EndPlayReason) override;
+	virtual void TickComponent(float DeltaTime, enum ELevelTick TickType, FActorComponentTickFunction *ThisTickFunction) override;
+
+	// Game code should write to this when locally controlling this object
+	UPROPERTY(BlueprintReadWrite, Category="Input")
+	FTempMockInputCmd PendingInputCmd;
+
+	DECLARE_DYNAMIC_MULTICAST_DELEGATE(FOnGeneratedLocalInputCmd);
+	
+	UPROPERTY(BlueprintAssignable, Category = "Components|Activation")
+	FOnGeneratedLocalInputCmd OnGeneratedLocalInputCmd;
+
+	UFUNCTION(BlueprintPure, Category = "Mock Input")
+	float GetJumpCooldownMS() const { return 0.f; }
+
+private:
+
+	APlayerController* GetOwnerPC() const;
+
+	UPROPERTY(Replicated, transient)
+	FNetworkPhysicsState NetworkPhysicsState;
+
+	// Automatically marshalled complete copy of the latest object state
+	// This isn't mandatory: we can optionally register for the "full copy" or we could
+	// register for "specialized output struct" with only the data we want marshalled back to the GT
+	UPROPERTY()
+	FTempMockObject	SimOutput;
+
+	UPROPERTY(Replicated)
+	FNetworkPredictionProxyAsync ReplicatedObject;
+
+	UPROPERTY(transient)
+	APlayerController* CachedPC = nullptr;
+};
\ No newline at end of file
