speculative SA fix
#rnx

diff --git a/Engine/Plugins/Runtime/NetworkPrediction/Source/NetworkPrediction/Public/NetworkPredictionAsync.h b/Engine/Plugins/Runtime/NetworkPrediction/Source/NetworkPrediction/Public/NetworkPredictionAsync.h
--- a/Engine/Plugins/Runtime/NetworkPrediction/Source/NetworkPrediction/Public/NetworkPredictionAsync.h
+++ b/Engine/Plugins/Runtime/NetworkPrediction/Source/NetworkPrediction/Public/NetworkPredictionAsync.h
@@ -799,7 +799,7 @@ public:
 					if (!InputCorrectionSnapshot)
 					{
 						InputCorrectionSnapshot = &PendingInputCorrections_Internal.AddDefaulted_GetRef();
-						InputCorrectionSnapshot->template Get<int32>() = Frame;
+						InputCorrectionSnapshot->template Get<0>() = Frame;
 					}
 
 					InputCorrectionSnapshot->template Get<1>().FindOrAdd(ObjKey) = InputCmd;	
