Network Prediction
-Another SA fix

#rnx

diff --git a/Engine/Plugins/Runtime/NetworkPrediction/Source/NetworkPrediction/Public/NetworkPredictionAsync.h b/Engine/Plugins/Runtime/NetworkPrediction/Source/NetworkPrediction/Public/NetworkPredictionAsync.h
--- a/Engine/Plugins/Runtime/NetworkPrediction/Source/NetworkPrediction/Public/NetworkPredictionAsync.h
+++ b/Engine/Plugins/Runtime/NetworkPrediction/Source/NetworkPrediction/Public/NetworkPredictionAsync.h
@@ -795,7 +795,7 @@ public:
 
 				for (InputCmdType& InputCmd : FutureInputs)
 				{
-					TTuple<int32, TSortedMap<ObjKeyType, InputCmdType>>* InputCorrectionSnapshot = PendingInputCorrections_Internal.FindByPredicate([Frame](const TTuple<int32, TSortedMap<ObjKeyType, InputCmdType>>& T) { return T.template Get<int32>() == Frame; });
+					TTuple<int32, TSortedMap<ObjKeyType, InputCmdType>>* InputCorrectionSnapshot = PendingInputCorrections_Internal.FindByPredicate([Frame](const TTuple<int32, TSortedMap<ObjKeyType, InputCmdType>>& T) { return T.template Get<0>() == Frame; });
 					if (!InputCorrectionSnapshot)
 					{
 						InputCorrectionSnapshot = &PendingInputCorrections_Internal.AddDefaulted_GetRef();
