NetworkPrediction
-Final SA fixes, I think

#rnx

diff --git a/Engine/Plugins/Runtime/NetworkPrediction/Source/NetworkPrediction/Public/NetworkPredictionAsync.h b/Engine/Plugins/Runtime/NetworkPrediction/Source/NetworkPrediction/Public/NetworkPredictionAsync.h
--- a/Engine/Plugins/Runtime/NetworkPrediction/Source/NetworkPrediction/Public/NetworkPredictionAsync.h
+++ b/Engine/Plugins/Runtime/NetworkPrediction/Source/NetworkPrediction/Public/NetworkPredictionAsync.h
@@ -722,7 +722,7 @@ public:
 			for (auto& ObjIt : Snapshot.Objects)
 			{
 				ObjKeyType& ObjKey = ObjIt.Key;
-				ManagedObjType& AuthObj = ObjIt.Value.template Get<ManagedObjType>();
+				ManagedObjType& AuthObj = ObjIt.Value.template Get<0>();
 				const ManagedObjType* LocalState = LocalMap->Find(ObjKey);
 
 				bool bShouldReconcile = false;
@@ -752,7 +752,7 @@ public:
 				}
 
 				// Even if we don't cause a correction, steal FutureInputs so they can be applied if we do end up resimulating
-				auto& FutureInputs = ObjIt.Value.template Get<TArray<InputCmdType>>();
+				auto& FutureInputs = ObjIt.Value.template Get<1>();
 				if (FutureInputs.Num() > 0)
 				{
 					FutureInputs_Internal.Emplace(MakeTuple(Snapshot.Frame, ObjKey, MoveTemp(FutureInputs)));
@@ -852,7 +852,7 @@ public:
 		// Input Corrections
 		for (auto& InputCorrection : PendingInputCorrections_Internal)
 		{
-			if (InputCorrection.template Get<int32>() == PhysicsStep)
+			if (InputCorrection.template Get<0>() == PhysicsStep)
 			{
 				TSortedMap<ObjKeyType, ManagedObjType>* InputMap = PersistentStorage_Internal.Get(PhysicsStep);
 				for (auto& MapIt : InputCorrection.template Get<1>())
