Network Physics
-Fix Playercontroller input buffer under writing issue
-Added more debug to core NetwortkPhysics
-Mock Physics spawners return spawned actor

#rnx

diff --git a/Engine/Plugins/Runtime/NetworkPrediction/Source/NetworkPrediction/Private/NetworkPhysics.cpp b/Engine/Plugins/Runtime/NetworkPrediction/Source/NetworkPrediction/Private/NetworkPhysics.cpp
--- a/Engine/Plugins/Runtime/NetworkPrediction/Source/NetworkPrediction/Private/NetworkPhysics.cpp
+++ b/Engine/Plugins/Runtime/NetworkPrediction/Source/NetworkPrediction/Private/NetworkPhysics.cpp
@@ -19,6 +19,8 @@
 #include "GameFramework/Actor.h"
 #include "NetworkPredictionCheck.h"
 #include "PhysicsProxy/SingleParticlePhysicsProxy.h"
+#include "NetworkPredictionDebug.h"
+#include "GameFramework/PlayerState.h"
 
 DEFINE_LOG_CATEGORY(LogNetworkPhysics);
 
@@ -750,6 +752,15 @@ void UNetworkPhysicsManager::ProcessInputs_External(int32 PhysicsStep, const TAr
 	UWorld* World = GetWorld();
 	const ENetMode NetMode = World->GetNetMode();
 
+	int32 FrameDelta = this->LatestConfirmedFrame - PhysicsStep;
+
+	if (PhysicsStep % 500 == 0)
+	{
+		ClientServerMaxFrameDelta = 0;
+	}
+
+	ClientServerMaxFrameDelta = FMath::Max(ClientServerMaxFrameDelta, FrameDelta);
+
 	// ----------------------------------------------------------------------
 	//	Server: Do remote client input buffering logic
 	//	This should probably be moved into PlayerController itself, with more options.
@@ -758,7 +769,7 @@ void UNetworkPhysicsManager::ProcessInputs_External(int32 PhysicsStep, const TAr
 	
 	if (NetMode != NM_Client)
 	{
-		constexpr int32 MaxBufferedCmds = 64;
+		constexpr int32 MaxBufferedCmds = 16;
 		for (FConstPlayerControllerIterator Iterator = World->GetPlayerControllerIterator(); Iterator; ++Iterator)
 		{
 			if (APlayerController* PC = Iterator->Get())
@@ -807,7 +818,7 @@ void UNetworkPhysicsManager::ProcessInputs_External(int32 PhysicsStep, const TAr
 
 
 				// ----------------------------------------------------------------------------------------------
-				// REMOVE Once Mock/Hardcoded exmaple is gone
+				// REMOVE Once Mock/Hardcoded example is gone
 				// ----------------------------------------------------------------------------------------------
 				if (FrameInfo.LastProcessedInputFrame != INDEX_NONE)
 				{
@@ -1039,6 +1050,22 @@ void UNetworkPhysicsManager::TickDrawDebug()
 				const int32 LatestSimulatedFrame = Solver->GetCurrentFrame() - Manager->LocalOffset;
 				YPos += Canvas->DrawText(GEngine->GetMediumFont(), FString::Printf(TEXT("Num Predicted Frames: %d"), (LatestSimulatedFrame - Manager->LatestConfirmedFrame)), XPos, YPos);
 
+				APlayerController* LocalPC = GEngine->GetFirstLocalPlayerController(Manager->GetWorld());
+
+				if (APlayerState* PS = LocalPC->GetPlayerState<APlayerState>())
+				{
+					YPos += Canvas->DrawText(GEngine->GetMediumFont(), FString::Printf(TEXT("Ping: %d"), PS->GetPing()), XPos, YPos);
+				}
+
+				
+				YPos += Canvas->DrawText(GEngine->GetMediumFont(), FString::Printf(TEXT("ClientServerMaxFrameDelta: %d"), Manager->ClientServerMaxFrameDelta), XPos, YPos);
+				
+				if (APlayerController* ServerPC = Cast<APlayerController>(NetworkPredictionDebug::FindReplicatedObjectOnPIEServer(LocalPC)))
+				{
+					int32 NumBufferedCmd = ServerPC->GetInputBuffer().HeadFrame() - ServerPC->GetServerFrameInfo().LastProcessedInputFrame;
+					YPos += Canvas->DrawText(GEngine->GetMediumFont(), FString::Printf(TEXT("Server-Side Buffered InputCmds: %d"), NumBufferedCmd), XPos, YPos);
+				}
+
 				YPos += 20.f;
 
 				const float W = 2.f;
diff --git a/Engine/Plugins/Runtime/NetworkPrediction/Source/NetworkPrediction/Public/NetworkPhysics.h b/Engine/Plugins/Runtime/NetworkPrediction/Source/NetworkPrediction/Public/NetworkPhysics.h
--- a/Engine/Plugins/Runtime/NetworkPrediction/Source/NetworkPrediction/Public/NetworkPhysics.h
+++ b/Engine/Plugins/Runtime/NetworkPrediction/Source/NetworkPrediction/Public/NetworkPhysics.h
@@ -163,6 +163,8 @@ private:
 	int32 LatestSimulatedFrame = INDEX_NONE;	// Latest frame we have sent off to PT for simulation. Doesn't mean its been completed and marshalled back.
 	int32 LocalOffset = 0; // Calculated client/server frame offset. ClientFrame = ServerFrame + LocalOffset
 
+	int32 ClientServerMaxFrameDelta = 0;
+
 	TSortedMap<int32, int32> ManagedHandleToIndexMap;
 	TSparseArray<FNetworkPhysicsState*> ManagedPhysicsStates;
 	int32 LastFreeIndex = 0;
diff --git a/Engine/Plugins/Runtime/NetworkPredictionExtras/Source/NetworkPredictionExtras/Private/MockNetworkPhysics.cpp b/Engine/Plugins/Runtime/NetworkPredictionExtras/Source/NetworkPredictionExtras/Private/MockNetworkPhysics.cpp
--- a/Engine/Plugins/Runtime/NetworkPredictionExtras/Source/NetworkPredictionExtras/Private/MockNetworkPhysics.cpp
+++ b/Engine/Plugins/Runtime/NetworkPredictionExtras/Source/NetworkPredictionExtras/Private/MockNetworkPhysics.cpp
@@ -1066,7 +1066,7 @@ void UNetworkPhysicsComponent::ProcessInputs_External(FMockManagedState& State,
 
 // ============================================================================================================
 
-void ANetworkPredictionSpawner::Spawn(FName StreamName)
+AActor* ANetworkPredictionSpawner::Spawn(FName StreamName)
 {
 	ANetworkPredictionSpawner* SourceSpawner = nullptr;
 	for (TObjectIterator<UWorld> WorldIt; WorldIt; ++WorldIt)
@@ -1086,13 +1086,14 @@ void ANetworkPredictionSpawner::Spawn(FName StreamName)
 		}
 	}
 
+	AActor* SpawnedActor = nullptr;
 	ensure(SourceSpawner);
 
 	StreamName = (StreamName == NAME_None ? SourceSpawner->RecordedInputs[FMath::Rand() % SourceSpawner->RecordedInputs.Num()].Name : StreamName);
 
 	if (FMockRecordedInputs* PlaybackRecordedInputs = SourceSpawner->RecordedInputs.FindByKey(StreamName))
 	{
-		AActor* SpawnedActor = this->GetWorld()->SpawnActor<AActor>(this->SpawnClass.Get(), this->GetActorTransform());
+		SpawnedActor = this->GetWorld()->SpawnActor<AActor>(this->SpawnClass.Get(), this->GetActorTransform());
 		if (UNetworkPhysicsComponent* Comp = SpawnedActor->FindComponentByClass<UNetworkPhysicsComponent>())
 		{
 			Comp->StartPlayback(&PlaybackRecordedInputs->Inputs);
@@ -1102,11 +1103,13 @@ void ANetworkPredictionSpawner::Spawn(FName StreamName)
 	{
 		UE_LOG(LogNetworkPhysics, Warning, TEXT("Could not find Inputs named %s on %s"), *StreamName.ToString(), *GetName());
 	}
+
+	return SpawnedActor;
 }
 
-void ANetworkPredictionSpawner::SpawnRandom()
+AActor* ANetworkPredictionSpawner::SpawnRandom()
 {
-	Spawn(NAME_None);
+	return Spawn(NAME_None);
 }
 
 void ANetworkPredictionSpawner::StartRecording(UNetworkPhysicsComponent* Target, FName StreamName)
diff --git a/Engine/Plugins/Runtime/NetworkPredictionExtras/Source/NetworkPredictionExtras/Public/MockNetworkPhysics.h b/Engine/Plugins/Runtime/NetworkPredictionExtras/Source/NetworkPredictionExtras/Public/MockNetworkPhysics.h
--- a/Engine/Plugins/Runtime/NetworkPredictionExtras/Source/NetworkPredictionExtras/Public/MockNetworkPhysics.h
+++ b/Engine/Plugins/Runtime/NetworkPredictionExtras/Source/NetworkPredictionExtras/Public/MockNetworkPhysics.h
@@ -397,11 +397,11 @@ public:
 
 	// Spawns actor and plays named recorded inputs
 	UFUNCTION(BlueprintCallable, Category="Spawner")
-	void Spawn(FName NamedInputs);
+	AActor* Spawn(FName NamedInputs);
 
 	// Spawns actor and plays random stream of recorded input
 	UFUNCTION(BlueprintCallable, Category="Spawner")
-	void SpawnRandom();
+	AActor* SpawnRandom();
 	
 	void StartRecording(UNetworkPhysicsComponent* Target, FName NamedInputs);
 
diff --git a/Engine/Source/Runtime/Engine/Classes/GameFramework/PlayerController.h b/Engine/Source/Runtime/Engine/Classes/GameFramework/PlayerController.h
--- a/Engine/Source/Runtime/Engine/Classes/GameFramework/PlayerController.h
+++ b/Engine/Source/Runtime/Engine/Classes/GameFramework/PlayerController.h
@@ -2015,6 +2015,7 @@ public:
 	struct FInputCmdBuffer
 	{
 		int32 HeadFrame() const { return LastWritten; }
+		int32 TailFrame() const { return FMath::Max(0, LastWritten - Buffer.Num() + 1); }
 		TArray<uint8>& Write(int32 Frame) { LastWritten = Frame; return Buffer[Frame % Buffer.Num()]; }
 		const TArray<uint8>& Get(int32 Frame) const { return Buffer[Frame % Buffer.Num()]; }
 
diff --git a/Engine/Source/Runtime/Engine/Private/PlayerController.cpp b/Engine/Source/Runtime/Engine/Private/PlayerController.cpp
--- a/Engine/Source/Runtime/Engine/Private/PlayerController.cpp
+++ b/Engine/Source/Runtime/Engine/Private/PlayerController.cpp
@@ -1688,10 +1688,20 @@ void APlayerController::ServerRecvClientInputFrame_Implementation(int32 InRecvCl
 
 	for (int32 DroppedFrame = InputBuffer.HeadFrame()+1; DroppedFrame < InRecvClientInputFrame && DroppedFrame > 0; ++DroppedFrame)
 	{
+		UE_LOG(LogPlayerController, Warning, TEXT( "ClientInput Gap in frames (Dropped). %s [%s]. DroppedFrame: %d. RecvFrame: %d. LastProcessInputFrame: %d" ), *GetName(), PlayerState ? *PlayerState->GetPlayerName() : TEXT("???"), DroppedFrame, InRecvClientInputFrame, ServerFrameInfo.LastProcessedInputFrame);
+		ServerFrameInfo.LastProcessedInputFrame++; // Ehh lets try this for now
 		InputBuffer.Write(DroppedFrame) = InputBuffer.Get(DroppedFrame-1);
 	}
 
 	InputBuffer.Write(InRecvClientInputFrame) = MoveTemp(const_cast<TArray<uint8>&>(Data));
+
+	if (ServerFrameInfo.LastProcessedInputFrame < InputBuffer.TailFrame())
+	{
+		// At this point, things are pretty bad and we are going to drop commands. 
+		// We still need guards client side to not send too many commands
+		UE_LOG(LogPlayerController, Warning, TEXT( "ClientInput buffer overflow. %s [%s]. RecvFrame: %d. LastProcessInputFrame: %d." ), *GetName(), PlayerState ? *PlayerState->GetPlayerName() : TEXT("???"), InRecvClientInputFrame, ServerFrameInfo.LastProcessedInputFrame);
+		ServerFrameInfo.LastProcessedInputFrame = (InputBuffer.TailFrame() + InputBuffer.HeadFrame()) / 2;
+	}
 }
 
 void APlayerController::ClientRecvServerAckFrame_Implementation(int32 InClientInputFrame, int32 InServerFrameNumber)
