Network Physics
-Fix recorded Inputcmds not working in cooked games
-Fix issues in ChaosResultsManager causing false positives in bDiverged flag
-cvar np2.Mock.JumpHack: disables the trace in the jump code which is a source of non-determinism since trace is still done against GT data
-cvar np2.Debug.Tolerance: only draw corrections that were significant
-cvar np2.ForceResimWithoutCorrection: resim without applying server side data

#rnx

diff --git a/Engine/Plugins/Runtime/NetworkPrediction/Source/NetworkPrediction/Private/NetworkPhysics.cpp b/Engine/Plugins/Runtime/NetworkPrediction/Source/NetworkPrediction/Private/NetworkPhysics.cpp
--- a/Engine/Plugins/Runtime/NetworkPrediction/Source/NetworkPrediction/Private/NetworkPhysics.cpp
+++ b/Engine/Plugins/Runtime/NetworkPrediction/Source/NetworkPrediction/Private/NetworkPhysics.cpp
@@ -31,12 +31,24 @@ namespace UE_NETWORK_PHYSICS
 	float V = 1.0;	FAutoConsoleVariableRef CVarV(TEXT("np2.Tolerance.V"), V, TEXT("Velocity Tolerance"));
 	float W = 1.0;	FAutoConsoleVariableRef CVarW(TEXT("np2.Tolerance.W"), W, TEXT("Rotational Velocity Tolerance"));
 
+	float DebugDrawTolerance_X = 5.0;
+	FAutoConsoleVariableRef CVarDebugX(TEXT("np2.Debug.Tolerance.X"), DebugDrawTolerance_X, TEXT("Location Debug Drawing Toleration"));
+
+	float DebugDrawTolerance_R = 2.0;
+	FAutoConsoleVariableRef CVarDebugR(TEXT("np2.Debug.Tolerance.R"), DebugDrawTolerance_R, TEXT("Rotation Debug Drawing Tolerance"));
+
 	int32 Debug = 0;
 	FAutoConsoleVariableRef CVarDebug(TEXT("np2.Debug"), Debug, TEXT("Debug mode for in world drawing"));
 
+	int32 DebugTolerance = 0;
+	FAutoConsoleVariableRef CVarDebugTolerance(TEXT("np2.Debug.Tolerance"), DebugTolerance, TEXT("If enabled, only draw large corrections in world."));
+
 	bool bForceResim=false;
 	FAutoConsoleVariableRef CVarResim(TEXT("np2.ForceResim"), bForceResim, TEXT("Forces near constant resimming"));
 
+	bool ForceResimWithoutCorrection=false;
+	FAutoConsoleVariableRef CVarForceResimWithoutCorrection(TEXT("np2.ForceResimWithoutCorrection"), ForceResimWithoutCorrection, TEXT("Forces near constant resimming WITHOUT applying server data"));
+
 	bool bEnable=false;
 	FAutoConsoleVariableRef CVarEnable(TEXT("np2.bEnable"), bEnable, TEXT("Enabled rollback physics. Must be set before starting game"));
 
@@ -79,24 +91,27 @@ struct FNetworkPhysicsRewindCallback : public Chaos::IRewindCallback
 
 	int32 TriggerRewindIfNeeded_Internal(int32 LastCompletedStep) override
 	{
+		if (bIsServer)
+		{
+			return INDEX_NONE;
+		}
+
 		if (ResimEndFrame != INDEX_NONE)
 		{
 			// We are already in a resim and shouldn't trigger a new one
 			return INDEX_NONE;
 		}
 
-		/*
-		if (UE_NETWORK_PHYSICS::bForceResim)
+		PendingCorrections.Reset();
+
+		if (UE_NETWORK_PHYSICS::ForceResimWithoutCorrection)
 		{
-			DataFromNetwork.Empty();
-			if (LastCompletedStep > 100 && LastCompletedStep > LastResim+4)
-			{
-				LastResim = LastCompletedStep;
-				return LastCompletedStep - 4;
-			}
-			return INDEX_NONE;
+			//const int32 ForceRewindFrame = RewindData->GetEarliestFrame_Internal()+1; // This was causing us to back > 64 frames?
+			const int32 ForceRewindFrame = LastCompletedStep - 9;
+			ResimEndFrame = LastCompletedStep;
+			//UE_LOG(LogTemp, Warning, TEXT("Forcing rewind to %d. LastCompletedStep: %d"), ForceRewindFrame, LastCompletedStep);
+			return ForceRewindFrame;
 		}
-		*/
 
 		FStats Stats;
 		Stats.LatestSimFrame = RewindData->CurrentFrame();
@@ -112,9 +127,7 @@ struct FNetworkPhysicsRewindCallback : public Chaos::IRewindCallback
 		};
 		
 		checkSlow(RewindData);
-		const int32 MinFrame = RewindData->CurrentFrame() - RewindData->GetFramesSaved();
-		
-		PendingCorrections.Reset();
+		const int32 MinFrame = RewindData->CurrentFrame() - RewindData->GetFramesSaved();				
 
 		FSnapshot Snapshot;
 		int32 RewindToFrame = INDEX_NONE;
@@ -225,7 +238,6 @@ struct FNetworkPhysicsRewindCallback : public Chaos::IRewindCallback
 			return INDEX_NONE;
 		}
 		
-		ResimEndFrame = LastCompletedStep;
 		return RewindToFrame;
 	}
 	
@@ -244,7 +256,7 @@ struct FNetworkPhysicsRewindCallback : public Chaos::IRewindCallback
 
 				if (auto* PT = CorrectionState.Proxy->GetPhysicsThreadAPI())
 				{
-					UE_LOG(LogNetworkPhysics, Log, TEXT("Applying Correction from frame %d (actual step: %d). Location: %s"), CorrectionState.Frame, PhysicsStep, *FVector(CorrectionState.Physics.Location).ToString());
+					UE_CLOG(UE_NETWORK_PHYSICS::bLogCorrections, LogNetworkPhysics, Log, TEXT("Applying Correction from frame %d (actual step: %d). Location: %s"), CorrectionState.Frame, PhysicsStep, *FVector(CorrectionState.Physics.Location).ToString());
 				
 					PT->SetX(CorrectionState.Physics.Location, false);
 					PT->SetV(CorrectionState.Physics.LinearVelocity, false);
@@ -254,7 +266,7 @@ struct FNetworkPhysicsRewindCallback : public Chaos::IRewindCallback
 					//if (PT->ObjectState() != CorrectionState.Physics.ObjectState)
 					{
 						ensure(CorrectionState.Physics.ObjectState != Chaos::EObjectStateType::Uninitialized);
-						UE_LOG(LogNetworkPhysics, Log, TEXT("Applying Correction State %d"), CorrectionState.Physics.ObjectState);
+						UE_CLOG(UE_NETWORK_PHYSICS::bLogCorrections, LogNetworkPhysics, Log, TEXT("Applying Correction State %d"), CorrectionState.Physics.ObjectState);
 						PT->SetObjectState(CorrectionState.Physics.ObjectState);
 					}
 				}
@@ -386,6 +398,7 @@ struct FNetworkPhysicsRewindCallback : public Chaos::IRewindCallback
 	TArray<Chaos::ISimCallbackObject*> SimObjectCallbacks;
 	UNetworkPhysicsManager* NetworkPhysicsManager = nullptr;
 
+	bool bIsServer = false;
 
 	// Debugging
 	struct FStats
@@ -491,6 +504,7 @@ void UNetworkPhysicsManager::PostNetRecv()
 	checkSlow(World);
 
 	const bool bIsServer = World->GetNetMode() != NM_Client;
+	RewindCallback->bIsServer = bIsServer;
 
 	// Iterate through all objects we are managing and coalease them into a snapshot that is marshalled to PT for rollback consideration
 	if (bIsServer)
@@ -991,6 +1005,19 @@ void UNetworkPhysicsManager::TickDrawDebug()
 
 			for (auto& Correction : Stats.Corrections)
 			{
+				if (UE_NETWORK_PHYSICS::DebugTolerance > 0)
+				{
+					if (FVector::Distance(Correction.LocalState.Location, Correction.AuthorityState.Location) < UE_NETWORK_PHYSICS::DebugDrawTolerance_X)
+					{
+						continue;
+					}
+
+					if (FQuat::ErrorAutoNormalize(Correction.LocalState.Rotation, Correction.AuthorityState.Rotation) < UE_NETWORK_PHYSICS::DebugDrawTolerance_R)
+					{
+						continue;
+					}
+				}
+
 				if (auto Func = DrawDebugMap.Find(Correction.Proxy))
 				{
 					P.Color = FColor::Red;
diff --git a/Engine/Plugins/Runtime/NetworkPredictionExtras/Source/NetworkPredictionExtras/Private/MockNetworkPhysics.cpp b/Engine/Plugins/Runtime/NetworkPredictionExtras/Source/NetworkPredictionExtras/Private/MockNetworkPhysics.cpp
--- a/Engine/Plugins/Runtime/NetworkPredictionExtras/Source/NetworkPredictionExtras/Private/MockNetworkPhysics.cpp
+++ b/Engine/Plugins/Runtime/NetworkPredictionExtras/Source/NetworkPredictionExtras/Private/MockNetworkPhysics.cpp
@@ -144,6 +144,9 @@ namespace UE_NETWORK_PHYSICS
 
 	int32 JumpFudgeFrames=10;
 	FAutoConsoleVariableRef CVarJumpFudgeFrames(TEXT("np2.Mock.JumpFudgeFrames"), JumpFudgeFrames, TEXT("How many frames after being in air do we still allow a jump to begin"));
+
+	bool JumpHack=false;
+	FAutoConsoleVariableRef CVarJumpHack(TEXT("np2.Mock.JumpHack"), JumpHack, TEXT("Make jump not rely on trace which currently causes non determinism"));
 }
 
 void FMockManagedState::AsyncTick(UWorld* World, Chaos::FPhysicsSolver* Solver, const float DeltaSeconds, const int32 SimulationFrame, const int32 LocalStorageFrame)
@@ -163,8 +166,7 @@ void FMockManagedState::AsyncTick(UWorld* World, Chaos::FPhysicsSolver* Solver,
 			FCollisionObjectQueryParams ObjectParams(ECollisionChannel::ECC_PhysicsBody);
 
 			FHitResult OutHit;
-			const bool bInAir = !World->LineTraceSingleByChannel(OutHit, TracePosition, EndPosition, ECollisionChannel::ECC_WorldStatic, QueryParams, ResponseParams);
-
+			const bool bInAir = !UE_NETWORK_PHYSICS::JumpHack && !World->LineTraceSingleByChannel(OutHit, TracePosition, EndPosition, ECollisionChannel::ECC_WorldStatic, QueryParams, ResponseParams);
 			const float UpDot = FVector::DotProduct(PT->R().GetUpVector(), FVector::UpVector);
 
 			// ---------------------------------------------------------------------------------------------
@@ -471,8 +473,8 @@ public:
 					{
 						//UE_CLOG(UE_NETWORK_PHYSICS::bLogCorrections, LogNetworkPhysics, Log, TEXT("Rewind Needed for MockPersistState. Obj.Frame: %d. LastCompletedStep: %d."), Obj.Frame, LastCompletedStep);
 						UE_CLOG(UE_NETWORK_PHYSICS::MockDebug, LogNetworkPhysics, Error, TEXT("[%d] Rewind Needed for MockPersistState. Obj.Frame: %d. LastCompletedStep: %d. JumpCnt: %d (auth) vs %d (pred)"), Obj.Frame - this->LocalFrameOffset, Obj.Frame, LastCompletedStep, Obj.PT_State.JumpCount, LocalState.PT_State.JumpCount);
- 						UE_CLOG(UE_NETWORK_PHYSICS::MockDebug, LogNetworkPhysics, Error, TEXT("     Server: JumpCnt: %d. JumpCooldownMS: %d"), Obj.PT_State.JumpCount, Obj.PT_State.JumpCooldownMS);
-						UE_CLOG(UE_NETWORK_PHYSICS::MockDebug, LogNetworkPhysics, Error, TEXT("     Local:  JumpCnt: %d. JumpCooldownMS: %d"), LocalState.PT_State.JumpCount, LocalState.PT_State.JumpCooldownMS);
+   						UE_CLOG(UE_NETWORK_PHYSICS::MockDebug, LogNetworkPhysics, Error, TEXT("     Server: JumpCnt: %d. JumpCooldownMS: %d. Airframe: %d. JumpFrame: %d."), Obj.PT_State.JumpCount, Obj.PT_State.JumpCooldownMS, Obj.PT_State.InAirFrame, Obj.PT_State.JumpStartFrame);
+						UE_CLOG(UE_NETWORK_PHYSICS::MockDebug, LogNetworkPhysics, Error, TEXT("     Local:  JumpCnt: %d. JumpCooldownMS: %d. Airframe: %d. JumpFrame: %d."), LocalState.PT_State.JumpCount, LocalState.PT_State.JumpCooldownMS, LocalState.PT_State.InAirFrame, LocalState.PT_State.JumpStartFrame);
 
 
 						RewindToFrame = RewindToFrame == INDEX_NONE ? Obj.Frame : FMath::Min(RewindToFrame, Obj.Frame);
@@ -1073,7 +1075,7 @@ AActor* ANetworkPredictionSpawner::Spawn(FName StreamName)
 	for (TObjectIterator<UWorld> WorldIt; WorldIt; ++WorldIt)
 	{
 		// Saved inputs are saved on the map version, so we need to copy them over to the server version
-		if (WorldIt->WorldType == EWorldType::Editor)
+		if (WorldIt->WorldType == EWorldType::Editor || WorldIt->WorldType == EWorldType::Game)
 		{
 			for (TActorIterator<ANetworkPredictionSpawner> It(*WorldIt); It; ++It)
 			{
diff --git a/Engine/Source/Runtime/Experimental/Chaos/Private/Chaos/Framework/ChaosResultsManager.cpp b/Engine/Source/Runtime/Experimental/Chaos/Private/Chaos/Framework/ChaosResultsManager.cpp
--- a/Engine/Source/Runtime/Experimental/Chaos/Private/Chaos/Framework/ChaosResultsManager.cpp
+++ b/Engine/Source/Runtime/Experimental/Chaos/Private/Chaos/Framework/ChaosResultsManager.cpp
@@ -379,7 +379,7 @@ namespace Chaos
 					FResimParticleInfo* ResimInfo = ParticleToResimInfo.Find(OriginalProxy);
 					if(ensure(ResimInfo))
 					{
-						ResimInfo->bDiverged = (ResimInfo->LeashStartTime == OriginalData->ExternalEndTime) ? StateDiverged(ResimInfo->Next, OriginalDirty) : true;
+						ResimInfo->bDiverged = (ResimInfo->EntryTime == OriginalData->ExternalEndTime) ? StateDiverged(ResimInfo->Next, OriginalDirty) : true;
 						if(ResimInfo->bDiverged)
 						{
 							//We still use resim's latest target (i.e. maybe it stopped moving a few frames ago)
@@ -392,6 +392,17 @@ namespace Chaos
 					}
 				}
 			}
+
+			for (const FDirtyRigidParticleData& ResimDirty : Results.Next->DirtyRigids)
+			{
+				if (FSingleParticlePhysicsProxy* ResimProxy = ResimDirty.GetProxy())
+				{
+					if (ParticleToResimInfo.FindChecked(ResimProxy).bDiverged)
+					{
+						UE_LOG(LogTemp, Warning, TEXT("Diverged"));
+					}
+				}
+			}
 		}
 	}
 
