NetworkPhysics
-np2.dump command to dump local physics state history for debugging

#rnx

diff --git a/Engine/Plugins/Runtime/NetworkPrediction/Source/NetworkPrediction/Private/NetworkPhysics.cpp b/Engine/Plugins/Runtime/NetworkPrediction/Source/NetworkPrediction/Private/NetworkPhysics.cpp
--- a/Engine/Plugins/Runtime/NetworkPrediction/Source/NetworkPrediction/Private/NetworkPhysics.cpp
+++ b/Engine/Plugins/Runtime/NetworkPrediction/Source/NetworkPrediction/Private/NetworkPhysics.cpp
@@ -21,6 +21,7 @@
 #include "PhysicsProxy/SingleParticlePhysicsProxy.h"
 #include "NetworkPredictionDebug.h"
 #include "GameFramework/PlayerState.h"
+#include "Misc/FileHelper.h"
 
 DEFINE_LOG_CATEGORY(LogNetworkPhysics);
 
@@ -273,11 +274,13 @@ struct FNetworkPhysicsRewindCallback : public Chaos::IRewindCallback
 			}
 		}
 
-		// Marhsall data back to GT based on what was requested for networking (this should be server only)
+		// Marhsall data back to GT based on what was requested for networking
 		FRequest Request;
 		while (DataRequested.Dequeue(Request));
 
 		FSnapshot Snapshot;
+		Snapshot.SimulationFrame = PhysicsStep - this->LastLocalOffset;
+		Snapshot.LocalFrameOffset = this->LastLocalOffset;
 		for (FSingleParticlePhysicsProxy* Proxy : Request.Proxies)
 		{
 			if (auto* PT = Proxy->GetPhysicsThreadAPI())
@@ -366,6 +369,7 @@ struct FNetworkPhysicsRewindCallback : public Chaos::IRewindCallback
 
 	struct FSnapshot
 	{
+		int32 SimulationFrame = 0; // Only set in DataFromPhysics
 		int32 LocalFrameOffset = 0; // only needed for debugging (translating client to server frames). Could be removed or maybe done via Insights tracing
 		TArray<FNetworkPhysicsState> Objects;
 	};
@@ -390,7 +394,6 @@ struct FNetworkPhysicsRewindCallback : public Chaos::IRewindCallback
 	Chaos::FPhysicsSolver* Solver = nullptr; //todo: shouldn't have to cache this, should be part of base class
 	UWorld* World = nullptr;
 
-	int32 LastResim = INDEX_NONE; // only used for force resim cvar
 	int32 LastLocalOffset = 0;
 
 	// ----------------------
@@ -431,7 +434,9 @@ struct FNetworkPhysicsRewindCallback : public Chaos::IRewindCallback
 
 UNetworkPhysicsManager::UNetworkPhysicsManager()
 {
-	
+#if !(UE_BUILD_SHIPPING || UE_BUILD_TEST)
+	bRecordDebugSnapshots = true;
+#endif
 }
 
 void UNetworkPhysicsManager::Initialize(FSubsystemCollectionBase& Collection)
@@ -507,7 +512,7 @@ void UNetworkPhysicsManager::PostNetRecv()
 	RewindCallback->bIsServer = bIsServer;
 
 	// Iterate through all objects we are managing and coalease them into a snapshot that is marshalled to PT for rollback consideration
-	if (bIsServer)
+	//if (bIsServer)
 	{
 		// Server has to tell PT what objects he wants to hear about.
 		// This feels kind of not right but I'm not sure how else to do it from the PT alone? Can I get a list of active proxies on the PT?
@@ -528,7 +533,8 @@ void UNetworkPhysicsManager::PostNetRecv()
 			RewindCallback->DataRequested.Enqueue(MoveTemp(Request));
 		}
 	}
-	else
+
+	if (!bIsServer)
 	{
 		// Client: Marshal received data from Network to PT so it can be reconciled
 		FNetworkPhysicsRewindCallback::FSnapshot Snapshot;
@@ -702,10 +708,6 @@ void UNetworkPhysicsManager::PreNetSend(float DeltaSeconds)
 
 
 	const bool bIsServer = World->GetNetMode() != NM_Client;
-	if (!bIsServer)
-	{
-		return;
-	}
 
 	FNetworkPhysicsRewindCallback::FSnapshot Snapshot;
 	bool bFoundData = false;
@@ -716,44 +718,68 @@ void UNetworkPhysicsManager::PreNetSend(float DeltaSeconds)
 
 	if (bFoundData)
 	{
-		for (FNetworkPhysicsState& Obj : Snapshot.Objects)
+		if (bRecordDebugSnapshots)
 		{
-			check(Obj.Proxy);
+			DebugSnapshotHead = Snapshot.SimulationFrame;
+			FDebugSnapshot& DebugSnapshot = DebugSnapshots[Snapshot.SimulationFrame % DebugSnapshots.Num()];
+			DebugSnapshot.LocalFrameOffset = Snapshot.LocalFrameOffset;
+			DebugSnapshot.Objects.Reset(Snapshot.Objects.Num());
+			for (FNetworkPhysicsState& Obj : Snapshot.Objects)
+			{
+				FDebugSnapshot::FDebugObject& DebugObj = DebugSnapshot.Objects.AddDefaulted_GetRef();
+				DebugObj.State = Obj;
 
-			// Copy data that was marshalled from PT to the managed FNetworkPhysicsState* that will be replicated.
-			// The tricky thing is that the objects we are hearing back from the PT about maybe deleted on the GT already.
+				if (FBodyInstance* BodyInstance = FPhysicsUserData::Get<FBodyInstance>(Obj.Proxy->GetGameThreadAPI().UserData()))
+				{
+					if (UPrimitiveComponent* PrimComponent = BodyInstance->OwnerComponent.Get())
+					{
+						DebugObj.WeakOwningActor = PrimComponent->GetOwner();
+					}
+				}
+			}
+		}
+
+		if (bIsServer)
+		{
+			for (FNetworkPhysicsState& Obj : Snapshot.Objects)
+			{
+				check(Obj.Proxy);
+
+				// Copy data that was marshalled from PT to the managed FNetworkPhysicsState* that will be replicated.
+				// The tricky thing is that the objects we are hearing back from the PT about maybe deleted on the GT already.
 			
-			// If we could use the physics proxy to safely get to the FNetworkPhysicsState, that would be ideal.
-			// The interface approach sort of allows this... if the interface is implemented on the actor. So not actually workable.
-			// If we just merged FNetworkPhysicsState into FRepMovement, we could just cast to actor and be done...
+				// If we could use the physics proxy to safely get to the FNetworkPhysicsState, that would be ideal.
+				// The interface approach sort of allows this... if the interface is implemented on the actor. So not actually workable.
+				// If we just merged FNetworkPhysicsState into FRepMovement, we could just cast to actor and be done...
 
-			// INetworkPhysicsObject* NetworkObject = Cast<INetworkPhysicsObject>(PhysicsState.Proxy->GetOwner());
-			// FNetworkPhysicsState* DestState = NetworkObject->GetNetworkPhysicsState();
+				// INetworkPhysicsObject* NetworkObject = Cast<INetworkPhysicsObject>(PhysicsState.Proxy->GetOwner());
+				// FNetworkPhysicsState* DestState = NetworkObject->GetNetworkPhysicsState();
 
-			// So we have to do the map lookup for now...
+				// So we have to do the map lookup for now...
 
-			// Actually this wont work either because we aren't round tripping the LocalManagedHandle to the PT
+				// Actually this wont work either because we aren't round tripping the LocalManagedHandle to the PT
 
-			/*			
-			if (int32* IdxPtr = ManagedHandleToIndexMap.Find(PhysicsState.LocalManagedHandle))			
-			{
-				const int32 Idx = *IdxPtr;
-				if (ensure(ManagedPhysicsStates.IsValidIndex(Idx)))
+				/*			
+				if (int32* IdxPtr = ManagedHandleToIndexMap.Find(PhysicsState.LocalManagedHandle))			
 				{
-					*ManagedPhysicsStates[Idx] = PhysicsState;
+					const int32 Idx = *IdxPtr;
+					if (ensure(ManagedPhysicsStates.IsValidIndex(Idx)))
+					{
+						*ManagedPhysicsStates[Idx] = PhysicsState;
+					}
 				}
-			}
-			*/
+				*/
 
-			// This is how it has to work for now, linear search to match the proxy pointer.
-			// We might be able to leverage consistent ordering (each linear search through both Snapshot.Objects and ManagedPhysicsStates incrementally in step)
-			//	(but be careful in the case where GT is ahead of PT and later requests have newer objects in lower positions in the MangedPhysicsState list...)
-			for (FNetworkPhysicsState* ManagedState : ManagedPhysicsStates)
-			{
-				if (ManagedState->Proxy == Obj.Proxy)
+				// This is how it has to work for now, linear search to match the proxy pointer.
+				// We might be able to leverage consistent ordering (each linear search through both Snapshot.Objects and ManagedPhysicsStates incrementally in step)
+				//	(but be careful in the case where GT is ahead of PT and later requests have newer objects in lower positions in the MangedPhysicsState list...)
+				for (FNetworkPhysicsState* ManagedState : ManagedPhysicsStates)
 				{
-					ManagedState->Physics = Obj.Physics;					
-					ManagedState->Frame = Obj.Frame;
+					if (ManagedState->Proxy == Obj.Proxy)
+					{
+						ManagedState->Physics = Obj.Physics;					
+						ManagedState->Frame = Obj.Frame;
+					}
 				}
 			}
 		}
@@ -1156,4 +1182,64 @@ void UNetworkPhysicsManager::TickDrawDebug()
 			}
 		});
 	}
-}
\ No newline at end of file
+}
+
+void UNetworkPhysicsManager::DumpDebugHistory()
+{
+	FStringOutputDevice Out;
+
+	TMap<TWeakObjectPtr<AActor>, FNetworkGUID> GUIDMap;
+
+	for (int32 i= FMath::Max(0, DebugSnapshotHead - DebugSnapshots.Num() + 1); i <= DebugSnapshotHead; ++i)
+	{
+		FDebugSnapshot& Snapshot = DebugSnapshots[i % DebugSnapshots.Num()];
+
+		for (FDebugSnapshot::FDebugObject& Obj : Snapshot.Objects)
+		{
+			if (Obj.WeakOwningActor.IsValid() == false)
+			{
+				continue;
+			}
+
+			FNetworkGUID GUID;
+			if (FNetworkGUID* FoundGUID = GUIDMap.Find(Obj.WeakOwningActor))
+			{
+				GUID = *FoundGUID;
+			}
+			else
+			{
+				GUID = NetworkPredictionDebug::FindObjectNetGUID(Obj.WeakOwningActor.Get());
+				GUIDMap.Add(Obj.WeakOwningActor, GUID);
+			}
+
+			Out.Logf(TEXT("[%d][n:%-4d][%d] X: %-45sR: %-64sV: %-40sW: %-40s\n"), i, GUID.Value, Snapshot.LocalFrameOffset, *Obj.State.Physics.Location.ToString(), *Obj.State.Physics.Rotation.ToString(), *Obj.State.Physics.LinearVelocity.ToString(), *Obj.State.Physics.AngularVelocity.ToString());
+		}
+	}
+
+	Out.Logf(TEXT("\n"));
+	for (auto& MapIt : GUIDMap)
+	{
+		Out.Logf(TEXT("[n:%-4d] -> %s\n"), MapIt.Value.Value, *GetNameSafe(MapIt.Key.Get()));
+	}
+
+	FString Path = FPaths::ProfilingDir() + FString::Printf(TEXT("/NetworkPrediction/NpDump_%s_%d.txt"), GetWorld()->GetNetMode() == NM_Client ? TEXT("Client") : TEXT("Server"), DebugSnapshotHead);
+	FFileHelper::SaveStringToFile(Out, *Path);
+
+}
+
+FAutoConsoleCommandWithWorldAndArgs NpDumpCmd(TEXT("np2.Dump"), TEXT(""),
+	FConsoleCommandWithWorldAndArgsDelegate::CreateLambda([](const TArray< FString >& Args, UWorld* InWorld) 
+{
+	if (UNetworkPhysicsManager* Manager = InWorld->GetSubsystem<UNetworkPhysicsManager>())
+	{
+		Manager->DumpDebugHistory();
+	}
+
+	if (InWorld->GetNetMode() == NM_Client)
+	{
+		if (APlayerController* PC = GEngine->GetFirstLocalPlayerController(InWorld))
+		{
+			PC->ServerExec(TEXT("np2.dump"));
+		}
+	}
+}));
\ No newline at end of file
diff --git a/Engine/Plugins/Runtime/NetworkPrediction/Source/NetworkPrediction/Private/NetworkPredictionDebug.cpp b/Engine/Plugins/Runtime/NetworkPrediction/Source/NetworkPrediction/Private/NetworkPredictionDebug.cpp
--- a/Engine/Plugins/Runtime/NetworkPrediction/Source/NetworkPrediction/Private/NetworkPredictionDebug.cpp
+++ b/Engine/Plugins/Runtime/NetworkPrediction/Source/NetworkPrediction/Private/NetworkPredictionDebug.cpp
@@ -6,7 +6,6 @@
 #include "GameFramework/Actor.h"
 #include "DrawDebugHelpers.h"
 #include "NetworkPredictionWorldManager.h"
-#include "Misc/NetworkGuid.h"
 #include "Engine/NetConnection.h"
 #include "Engine/NetDriver.h"
 #include "EngineUtils.h"
@@ -37,6 +36,26 @@ namespace NetworkPredictionDebug
 		DrawDebugString(World, Transform.GetLocation(), Str, nullptr, DrawColor, Lifetime, false);
 	}
 
+	FNetworkGUID FindObjectNetGUID(UObject* Obj)
+	{
+		FNetworkGUID NetGUID;
+		if (UWorld* World = Obj->GetWorld())
+		{
+			if (UNetDriver* NetDriver = World->GetNetDriver())
+			{
+				if (UNetConnection* NetConnection = NetDriver->ServerConnection)
+				{
+					NetGUID = NetConnection->PackageMap->GetNetGUIDFromObject(Obj);
+				}
+				else if (NetDriver->ClientConnections.Num() > 0)
+				{
+					NetGUID = NetDriver->ClientConnections[0]->PackageMap->GetNetGUIDFromObject(Obj);
+				}
+			}
+		}
+		return NetGUID;
+	}
+
 	UObject* FindReplicatedObjectOnPIEServer(UObject* ClientObject)
 	{
 		if (ClientObject == nullptr)
@@ -47,28 +66,21 @@ namespace NetworkPredictionDebug
 		UObject* ServerObject = nullptr;
 
 #if WITH_EDITOR
-		if (UWorld* World = ClientObject->GetWorld())
+		FNetworkGUID NetGUID = FindObjectNetGUID(ClientObject);
+		if (NetGUID.IsValid())
 		{
-			if (UNetDriver* NetDriver = World->GetNetDriver())
+			// Find the PIE server world
+			for (TObjectIterator<UWorld> It; It; ++It)
 			{
-				if (UNetConnection* NetConnection = NetDriver->ServerConnection)
+				if (It->WorldType == EWorldType::PIE && It->GetNetMode() != NM_Client)
 				{
-					FNetworkGUID NetGUID = NetConnection->PackageMap->GetNetGUIDFromObject(ClientObject);
-
-					// Find the PIE server world
-					for (TObjectIterator<UWorld> It; It; ++It)
+					if (UNetDriver* ServerNetDriver = It->GetNetDriver())
 					{
-						if (It->WorldType == EWorldType::PIE && It->GetNetMode() != NM_Client)
+						if (ServerNetDriver->ClientConnections.Num() > 0)
 						{
-							if (UNetDriver* ServerNetDriver = It->GetNetDriver())
-							{
-								if (ServerNetDriver->ClientConnections.Num() > 0)
-								{
-									UPackageMap* ServerPackageMap = ServerNetDriver->ClientConnections[0]->PackageMap;
-									ServerObject = ServerPackageMap->GetObjectFromNetGUID(NetGUID, true);
-									break;
-								}
-							}
+							UPackageMap* ServerPackageMap = ServerNetDriver->ClientConnections[0]->PackageMap;
+							ServerObject = ServerPackageMap->GetObjectFromNetGUID(NetGUID, true);
+							break;
 						}
 					}
 				}
@@ -79,6 +91,7 @@ namespace NetworkPredictionDebug
 		return ServerObject;
 	}
 
+	
 };
 
 
diff --git a/Engine/Plugins/Runtime/NetworkPrediction/Source/NetworkPrediction/Public/NetworkPhysics.h b/Engine/Plugins/Runtime/NetworkPrediction/Source/NetworkPrediction/Public/NetworkPhysics.h
--- a/Engine/Plugins/Runtime/NetworkPrediction/Source/NetworkPrediction/Public/NetworkPhysics.h
+++ b/Engine/Plugins/Runtime/NetworkPrediction/Source/NetworkPrediction/Public/NetworkPhysics.h
@@ -7,6 +7,7 @@
 #include "Chaos/Particles.h"
 #include "Components/ActorComponent.h"
 #include "UObject/ObjectKey.h"
+#include "Misc/NetworkGUID.h"
 
 #include "NetworkPhysics.generated.h"
 
@@ -150,6 +151,8 @@ public:
 
 	void RegisterPhysicsProxyDebugDraw(FNetworkPhysicsState* State, TUniqueFunction<void(const FDrawDebugParams&)>&& Func);
 
+	void DumpDebugHistory();
+
 private:
 
 	void TickDrawDebug();
@@ -176,4 +179,23 @@ private:
 
 	friend struct FNetworkPhysicsRewindCallback;
 	void ProcessInputs_External(int32 PhysicsStep, const TArray<Chaos::FSimCallbackInputAndObject>& SimCallbackInputs);
+
+	// Historic GT recording for GT only (compiled out. probably temp until Insights integration is back)
+	struct FDebugSnapshot
+	{
+		int32 LocalFrameOffset = 0;
+		struct FDebugObject
+		{
+			FNetworkGUID NetGUID;
+			TWeakObjectPtr<AActor> WeakOwningActor;
+			FNetworkPhysicsState State;
+		};
+
+		TArray<FDebugObject> Objects;
+	};
+
+	TStaticArray<FDebugSnapshot, 1024> DebugSnapshots;
+	int32 DebugSnapshotHead = 0;
+
+	bool bRecordDebugSnapshots = false;
 };
\ No newline at end of file
diff --git a/Engine/Plugins/Runtime/NetworkPrediction/Source/NetworkPrediction/Public/NetworkPredictionDebug.h b/Engine/Plugins/Runtime/NetworkPrediction/Source/NetworkPrediction/Public/NetworkPredictionDebug.h
--- a/Engine/Plugins/Runtime/NetworkPrediction/Source/NetworkPrediction/Public/NetworkPredictionDebug.h
+++ b/Engine/Plugins/Runtime/NetworkPrediction/Source/NetworkPrediction/Public/NetworkPredictionDebug.h
@@ -4,6 +4,7 @@
 #include "Math/Box.h"
 #include "Math/Color.h"
 #include "HAL/Platform.h"
+#include "Misc/NetworkGuid.h"
 
 // non ModelDef specific debug helpers
 namespace NetworkPredictionDebug
@@ -11,4 +12,5 @@ namespace NetworkPredictionDebug
 	NETWORKPREDICTION_API void DrawDebugOutline(FTransform Transform, FBox BoundingBox, FColor Color, float Lifetime);
 	NETWORKPREDICTION_API void DrawDebugText3D(const TCHAR* Str, FTransform Transform, FColor, float Lifetime);
 	NETWORKPREDICTION_API UObject* FindReplicatedObjectOnPIEServer(UObject* ClientObject);
+	NETWORKPREDICTION_API FNetworkGUID FindObjectNetGUID(UObject* Obj);
 };
\ No newline at end of file
