Network Physics
-Debug helpers for catching specific frame numbers

#rnx

diff --git a/Engine/Plugins/Runtime/NetworkPrediction/Source/NetworkPrediction/Private/NetworkPhysics.cpp b/Engine/Plugins/Runtime/NetworkPrediction/Source/NetworkPrediction/Private/NetworkPhysics.cpp
--- a/Engine/Plugins/Runtime/NetworkPrediction/Source/NetworkPrediction/Private/NetworkPhysics.cpp
+++ b/Engine/Plugins/Runtime/NetworkPrediction/Source/NetworkPrediction/Private/NetworkPhysics.cpp
@@ -67,6 +67,40 @@ namespace UE_NETWORK_PHYSICS
 
 	float LODDistance = 2400.f;
 	FAutoConsoleVariableRef CVarLODDistance(TEXT("np2.LODDistance"), LODDistance, TEXT("Simple distance based LOD"));
+
+	// ----------------------------------------------
+	// Debugger helpers:
+	//	See current frame anywhere with: 
+	//		{,,UE4Editor-NetworkPrediction.dll}UE_NETWORK_PHYSICS::GFrame
+	//	Set frame you want to break at via: 
+	//		{,,UE4Editor-NetworkPrediction.dll}UE_NETWORK_PHYSICS::GBreakAtFrame
+	//		or via console with "np2.BreakAtFrame"
+	//
+	//	Do something like this in your code:
+	//	if (UE_NETWORK_PHYSICS::ConditionalFrameBreakpoint()) { UE_LOG(...); }
+	//		or
+	//	UE_NETWORK_PHYSICS::ConditionalFrameEnsure();
+	// ----------------------------------------------
+
+	// Set this in debugger or via console to cause ConditionalFrameBreakpoint() to return true when current simulation frame is equal to this
+	int32 GBreakAtFrame = 0;
+	FAutoConsoleVariableRef CVarConditionalBreakAtFrame(TEXT("np2.BreakAtFrame"), GBreakAtFrame, TEXT(""));
+
+	// Set by the system prior to NP subsystems ticking
+	int32 GFrame = 0;
+	bool GServer = false;
+
+	// returns true/false so you can log or do whatever at GBreakAtFrame
+	bool ConditionalFrameBreakpoint()
+	{
+		return (GFrame == GBreakAtFrame);
+	}
+
+	// forces ensure failure to invoke debugger
+	void ConditionalFrameEnsure()
+	{
+		ensureAlways(GFrame != GBreakAtFrame);
+	}
 }
 
 struct FNetworkPhysicsRewindCallback : public Chaos::IRewindCallback
@@ -313,6 +347,9 @@ struct FNetworkPhysicsRewindCallback : public Chaos::IRewindCallback
 			DataFromPhysics.Enqueue(MoveTemp(Snapshot));
 		}
 
+		UE_NETWORK_PHYSICS::GFrame = PhysicsStep - this->LastLocalOffset;
+		UE_NETWORK_PHYSICS::GServer = this->bIsServer;
+
 		for (auto& SubSys : NetworkPhysicsManager->SubSystems)
 		{
 			SubSys.Value->ProcessInputs_Internal(PhysicsStep);
diff --git a/Engine/Plugins/Runtime/NetworkPrediction/Source/NetworkPrediction/Public/NetworkPhysics.h b/Engine/Plugins/Runtime/NetworkPrediction/Source/NetworkPrediction/Public/NetworkPhysics.h
--- a/Engine/Plugins/Runtime/NetworkPrediction/Source/NetworkPrediction/Public/NetworkPhysics.h
+++ b/Engine/Plugins/Runtime/NetworkPrediction/Source/NetworkPrediction/Public/NetworkPhysics.h
@@ -198,4 +198,12 @@ private:
 	int32 DebugSnapshotHead = 0;
 
 	bool bRecordDebugSnapshots = false;
+};
+
+namespace UE_NETWORK_PHYSICS
+{
+	extern NETWORKPREDICTION_API bool ConditionalFrameBreakpoint();
+	extern NETWORKPREDICTION_API bool GServer;
+	extern NETWORKPREDICTION_API int32 GFrame;
+	extern NETWORKPREDICTION_API int32 GBreakAtFrame;
 };
\ No newline at end of file
