NetworkPhysics
-Missed forward decl of ConditionalFrameEnsure

#rnx

diff --git a/Engine/Plugins/Runtime/NetworkPrediction/Source/NetworkPrediction/Private/NetworkPhysics.cpp b/Engine/Plugins/Runtime/NetworkPrediction/Source/NetworkPrediction/Private/NetworkPhysics.cpp
--- a/Engine/Plugins/Runtime/NetworkPrediction/Source/NetworkPrediction/Private/NetworkPhysics.cpp
+++ b/Engine/Plugins/Runtime/NetworkPrediction/Source/NetworkPrediction/Private/NetworkPhysics.cpp
@@ -93,13 +93,13 @@ namespace UE_NETWORK_PHYSICS
 	// returns true/false so you can log or do whatever at GBreakAtFrame
 	bool ConditionalFrameBreakpoint()
 	{
-		return (GFrame == GBreakAtFrame);
+		return (GBreakAtFrame != 0 && GFrame == GBreakAtFrame);
 	}
 
 	// forces ensure failure to invoke debugger
 	void ConditionalFrameEnsure()
 	{
-		ensureAlways(GFrame != GBreakAtFrame);
+		ensureAlways(!ConditionalFrameBreakpoint());
 	}
 }
 
diff --git a/Engine/Plugins/Runtime/NetworkPrediction/Source/NetworkPrediction/Public/NetworkPhysics.h b/Engine/Plugins/Runtime/NetworkPrediction/Source/NetworkPrediction/Public/NetworkPhysics.h
--- a/Engine/Plugins/Runtime/NetworkPrediction/Source/NetworkPrediction/Public/NetworkPhysics.h
+++ b/Engine/Plugins/Runtime/NetworkPrediction/Source/NetworkPrediction/Public/NetworkPhysics.h
@@ -203,6 +203,7 @@ private:
 namespace UE_NETWORK_PHYSICS
 {
 	extern NETWORKPREDICTION_API bool ConditionalFrameBreakpoint();
+	extern NETWORKPREDICTION_API void ConditionalFrameEnsure();
 	extern NETWORKPREDICTION_API bool GServer;
 	extern NETWORKPREDICTION_API int32 GFrame;
 	extern NETWORKPREDICTION_API int32 GBreakAtFrame;
