PR #8933: Fix NetworkPredictionCVars cmd/description mismatch so that the commands actually match up with their proper functions. (Contributed by [REDACTED_PERSON])

diff --git a/Engine/Plugins/Runtime/NetworkPrediction/Source/NetworkPrediction/Public/Services/NetworkPredictionService_Interpolate.inl b/Engine/Plugins/Runtime/NetworkPrediction/Source/NetworkPrediction/Public/Services/NetworkPredictionService_Interpolate.inl
--- a/Engine/Plugins/Runtime/NetworkPrediction/Source/NetworkPrediction/Public/Services/NetworkPredictionService_Interpolate.inl
+++ b/Engine/Plugins/Runtime/NetworkPrediction/Source/NetworkPrediction/Public/Services/NetworkPredictionService_Interpolate.inl
@@ -24,8 +24,8 @@
 namespace NetworkPredictionCVars
 {
 	// These are pretty bare bones. Will move to Insights tracing in the future for interpolation debugging
-	NETSIM_DEVCVAR_SHIPCONST_INT(DrawInterpolation, 0, "np.Interpolation.PrintSync", "Print Sync State buffer during interpolation");
-	NETSIM_DEVCVAR_SHIPCONST_INT(PrintSyncInterpolation, 0, "np.Interpolation.Draw", "Draw interpolation debug state in world");
+	NETSIM_DEVCVAR_SHIPCONST_INT(PrintSyncInterpolation, 0, "np.Interpolation.PrintSync", "Print Sync State buffer during interpolation");
+	NETSIM_DEVCVAR_SHIPCONST_INT(DrawInterpolation, 0, "np.Interpolation.Draw", "Draw interpolation debug state in world");
 
 	NETSIM_DEVCVAR_SHIPCONST_INT(DisableInterpolation, 0, "np.Interpolation.Disable", "Disables smooth interpolation and just Finalizes the last received frame");
 }
