Push model and constifying

diff --git a/Engine/Plugins/Runtime/NetworkPrediction/Source/NetworkPrediction/Private/NetworkPredictionReplicatedManager.cpp b/Engine/Plugins/Runtime/NetworkPrediction/Source/NetworkPrediction/Private/NetworkPredictionReplicatedManager.cpp
--- a/Engine/Plugins/Runtime/NetworkPrediction/Source/NetworkPrediction/Private/NetworkPredictionReplicatedManager.cpp
+++ b/Engine/Plugins/Runtime/NetworkPrediction/Source/NetworkPrediction/Private/NetworkPredictionReplicatedManager.cpp
@@ -36,7 +36,10 @@ void ANetworkPredictionReplicatedManager::GetLifetimeReplicatedProps(TArray<FLif
 {
 	Super::GetLifetimeReplicatedProps(OutLifetimeProps);
 
-	DOREPLIFETIME(ANetworkPredictionReplicatedManager, SharedPackageMap);
+	FDoRepLifetimeParams SharedParams;
+	SharedParams.bIsPushBased = true;
+
+	DOREPLIFETIME_WITH_PARAMS_FAST(ANetworkPredictionReplicatedManager, SharedPackageMap, SharedParams);
 }
 
 FDelegateHandle ANetworkPredictionReplicatedManager::OnAuthoritySpawn(const TFunction<void(ANetworkPredictionReplicatedManager*)>& Func)
@@ -50,12 +53,12 @@ FDelegateHandle ANetworkPredictionReplicatedManager::OnAuthoritySpawn(const TFun
 	return OnAuthoritySpawnDelegate.AddLambda(Func);
 }
 
-uint8 ANetworkPredictionReplicatedManager::GetIDForObject(UObject* Obj)
+uint8 ANetworkPredictionReplicatedManager::GetIDForObject(UObject* Obj) const
 {
 	// Naive lookup
-	for (auto It = SharedPackageMap.Items.CreateIterator(); It; ++It)
+	for (auto It = SharedPackageMap.Items.CreateConstIterator(); It; ++It)
 	{
-		FSharedPackageMapItem& Item = *It;
+		const FSharedPackageMapItem& Item = *It;
 		if (Item.SoftPtr.Get() == Obj)
 		{
 			npCheckSlow(It.GetIndex() < TNumericLimits<uint8>::Max());
@@ -67,7 +70,7 @@ uint8 ANetworkPredictionReplicatedManager::GetIDForObject(UObject* Obj)
 	return 0;
 }
 
-TSoftObjectPtr<UObject> ANetworkPredictionReplicatedManager::GetObjectForID(uint8 ID)
+TSoftObjectPtr<UObject> ANetworkPredictionReplicatedManager::GetObjectForID(uint8 ID) const
 {
 	if (SharedPackageMap.Items.IsValidIndex(ID))
 	{
@@ -91,5 +94,7 @@ uint8 ANetworkPredictionReplicatedManager::AddObjectToSharedPackageMap(TSoftObje
 	}
 
 	SharedPackageMap.Items.Add(FSharedPackageMapItem{SoftPtr});
+	MARK_PROPERTY_DIRTY_FROM_NAME(ANetworkPredictionReplicatedManager, SharedPackageMap, this);
+
 	return SharedPackageMap.Items.Num()-1;
 }
\ No newline at end of file
diff --git a/Engine/Plugins/Runtime/NetworkPrediction/Source/NetworkPrediction/Public/NetworkPredictionReplicatedManager.h b/Engine/Plugins/Runtime/NetworkPrediction/Source/NetworkPrediction/Public/NetworkPredictionReplicatedManager.h
--- a/Engine/Plugins/Runtime/NetworkPrediction/Source/NetworkPrediction/Public/NetworkPredictionReplicatedManager.h
+++ b/Engine/Plugins/Runtime/NetworkPrediction/Source/NetworkPrediction/Public/NetworkPredictionReplicatedManager.h
@@ -45,9 +45,9 @@ public:
 
 	uint8 AddObjectToSharedPackageMap(TSoftObjectPtr<UObject> SoftPtr);
 
-	uint8 GetIDForObject(UObject* Obj);
+	uint8 GetIDForObject(UObject* Obj) const;
 
-	TSoftObjectPtr<UObject> GetObjectForID(uint8 ID);
+	TSoftObjectPtr<UObject> GetObjectForID(uint8 ID) const;
 
 private:
 
