Network Prediction plugin: improved support for debugging with controlled de-syncs/reconciliation/rollbacks.
Fixes issue where ForceReconcile did not reset after use, and enhances it to roll back extra frames if desired.
Added similar method to force a client to send default inputs to the server instead of the locally-produced one, which can induce rollbacks.

diff --git a/Engine/Plugins/Runtime/NetworkPrediction/Source/NetworkPrediction/Public/NetworkPredictionSerialization.h b/Engine/Plugins/Runtime/NetworkPrediction/Source/NetworkPrediction/Public/NetworkPredictionSerialization.h
--- a/Engine/Plugins/Runtime/NetworkPrediction/Source/NetworkPrediction/Public/NetworkPredictionSerialization.h
+++ b/Engine/Plugins/Runtime/NetworkPrediction/Source/NetworkPrediction/Public/NetworkPredictionSerialization.h
@@ -17,6 +17,11 @@
 #define NETSIM_NETCONSTANT_NUM_BITS_FRAME 8	// Allows you to override this setting via UBT, but access via cleaner FActorMotionNetworkingConstants::NUM_BITS_FRAME
 #endif
 
+namespace NetworkPredictionCVars
+{
+	NETSIM_DEVCVAR_SHIPCONST_INT(ForceSendDefaultInputCommands, 0, "np.ForceSendDefaultInputCommands", "While enabled on a client, it will send default input cmds to the server, rather than the locally-produced input. This is a way to introduce de-syncs/rollbacks for debugging.");
+}
+
 struct FNetworkPredictionSerialization
 {
 	// How many bits we use to NetSerialize Frame numbers. This is only relevant AP Client <--> Server communication.
@@ -166,7 +171,16 @@ public:
 			}
 			else
 			{
-				FNetworkPredictionDriver<ModelDef>::NetSerialize(Frames->Buffer[Frame].InputCmd, P); // 2. InputCmd
+				if (NetworkPredictionCVars::ForceSendDefaultInputCommands())
+				{
+					// for debugging, send blank default input instead of what we've produced locally
+					static TConditionalState<InputType> DefaultInputCmd;
+					FNetworkPredictionDriver<ModelDef>::NetSerialize(DefaultInputCmd, P); 
+				}
+				else
+				{
+					FNetworkPredictionDriver<ModelDef>::NetSerialize(Frames->Buffer[Frame].InputCmd, P); // 2. InputCmd
+				}
 			}
 		}
 	}
@@ -276,7 +290,17 @@ public:
 			}
 			else
 			{
-				FNetworkPredictionDriver<ModelDef>::NetSerialize(Frames->Buffer[Frame].InputCmd, P); // 2. InputCmd
+				if (NetworkPredictionCVars::ForceSendDefaultInputCommands())
+				{
+					// for debugging, send blank default input instead of what we've produced locally
+					static TConditionalState<InputType> DefaultInputCmd;
+					FNetworkPredictionDriver<ModelDef>::NetSerialize(DefaultInputCmd, P);
+				}
+				else
+				{
+					FNetworkPredictionDriver<ModelDef>::NetSerialize(Frames->Buffer[Frame].InputCmd, P); // 2. InputCmd
+				}
+
 				FNetworkPredictionSerialization::SerializeDeltaMS(P.Ar, TickState->Frames[Frame].DeltaMS); // 3. Delta InputCmd
 			}
 		}
diff --git a/Engine/Plugins/Runtime/NetworkPrediction/Source/NetworkPrediction/Public/Services/NetworkPredictionService_Rollback.inl b/Engine/Plugins/Runtime/NetworkPrediction/Source/NetworkPrediction/Public/Services/NetworkPredictionService_Rollback.inl
--- a/Engine/Plugins/Runtime/NetworkPrediction/Source/NetworkPrediction/Public/Services/NetworkPredictionService_Rollback.inl
+++ b/Engine/Plugins/Runtime/NetworkPrediction/Source/NetworkPrediction/Public/Services/NetworkPredictionService_Rollback.inl
@@ -5,9 +5,10 @@
 
 namespace NetworkPredictionCVars
 {
-	NETSIM_DEVCVAR_SHIPCONST_INT(ForceReconcile, 0, "np.ForceReconcile", "Force a single reconcile");
-	NETSIM_DEVCVAR_SHIPCONST_INT(SkipReconcile, 0, "np.SkipReconcile", "Skip all reconciles");
-	NETSIM_DEVCVAR_SHIPCONST_INT(PrintReconciles, 0, "np.PrintReconciles", "Print reconciles to log");
+	NETSIM_DEVCVAR_SHIPCONST_INT(ForceReconcile,			0, "np.ForceReconcile",				"Force a single reconcile back to the last server-acknoledged frame. When used with np.ForceReconcileExtraFrames, additional frames can be rolled back. No effect on server. Resets after use.");
+	NETSIM_DEVCVAR_SHIPCONST_INT(ForceReconcileExtraFrames, 0, "np.ForceReconcileExtraFrames",	"Roll back this extra number of frames during the next targeted reconcile. Must be positive and reasonable given the buffer sizes.");
+	NETSIM_DEVCVAR_SHIPCONST_INT(SkipReconcile,				0, "np.SkipReconcile",				"Skip all reconciles");
+	NETSIM_DEVCVAR_SHIPCONST_INT(PrintReconciles,			0, "np.PrintReconciles",			"Print reconciles to log");
 }
 
 class IFixedRollbackService
@@ -79,7 +80,21 @@ public:
 			typename TInstanceFrameState<ModelDef>::FFrame& LocalFrameData = Frames.Buffer[LocalFrame];
 
 			bool bDoRollback = false;
-			if (FNetworkPredictionDriver<ModelDef>::ShouldReconcile( SyncAuxType(LocalFrameData.SyncState, LocalFrameData.AuxState), SyncAuxType(ClientRecvData.SyncState, ClientRecvData.AuxState) ))
+
+			if (NetworkPredictionCVars::ForceReconcile() > 0)
+			{
+				UE_NP_TRACE_SHOULD_RECONCILE(ClientRecvData.TraceID);
+				bDoRollback = true;
+				RollbackFrame = LocalFrame - FMath::Max(0, NetworkPredictionCVars::ForceReconcileExtraFrames());
+
+				if (NetworkPredictionCVars::PrintReconciles())
+				{				
+					UE_LOG(LogNetworkPrediction, Warning, TEXT("Reconcile activated due to ForceReconcile (to RollbackFrame=%i, including %i extra rollback frames)"), RollbackFrame, -(RollbackFrame-LocalFrame));
+				}
+
+				NetworkPredictionCVars::SetForceReconcile(0); // reset
+			}
+			else if (FNetworkPredictionDriver<ModelDef>::ShouldReconcile( SyncAuxType(LocalFrameData.SyncState, LocalFrameData.AuxState), SyncAuxType(ClientRecvData.SyncState, ClientRecvData.AuxState) ))
 			{
 				UE_NP_TRACE_SHOULD_RECONCILE(ClientRecvData.TraceID);
 				bDoRollback = true;
@@ -121,7 +136,7 @@ public:
 				}
 			}
 
-			if ((bDoRollback || NetworkPredictionCVars::ForceReconcile() > 0) && !NetworkPredictionCVars::SkipReconcile())
+			if (bDoRollback && !NetworkPredictionCVars::SkipReconcile())
 			{
 				RollbackFrame = (RollbackFrame == INDEX_NONE) ? LocalFrame : FMath::Min(RollbackFrame, LocalFrame);
 			}
