Network Prediction plugin: fixes input consumption lag on server after client recovers from framerate drop in fixed timestep simulations. Part 1 is to skip over missing frames when recovering from input starvation. Part 2 is to cap the number of inputs we must receive before starting to consume them again, which is configurable. Also made number of inputs to send configurable and improved logging.

diff --git a/Engine/Plugins/Runtime/NetworkPrediction/Source/NetworkPrediction/Public/NetworkPredictionSerialization.h b/Engine/Plugins/Runtime/NetworkPrediction/Source/NetworkPrediction/Public/NetworkPredictionSerialization.h
--- a/Engine/Plugins/Runtime/NetworkPrediction/Source/NetworkPrediction/Public/NetworkPredictionSerialization.h
+++ b/Engine/Plugins/Runtime/NetworkPrediction/Source/NetworkPrediction/Public/NetworkPredictionSerialization.h
@@ -88,7 +88,7 @@ struct FNetworkPredictionSerialization
 //	AP Client -> Server replication
 //
 //	The Fixed/Independent ticking implementations are a more than trivially different so they are split into separate implementations.
-//	Both currently send the last 'NumSendPerUpdate' per serialization. This could be improved with something more configurable/dynamic.
+//	Both currently send the last 'NumInputsPerSend' per serialization, determined from config settings. This could be improved with something more dynamic.
 //
 // ---------------------------------------------------------------------------------------------------------------------------
 
@@ -101,7 +101,10 @@ public:
 	using StateTypes = typename ModelDef::StateTypes;
 	using InputType = typename StateTypes::InputType;
 
-	static const int32 NumSendPerUpdate = 6;
+	static void SetNumInputsPerSend(int32 NumToSend) 
+	{ 
+		NumInputsPerSend = NumToSend; 
+	}
 
 	// ------------------------------------------------------------------------------------------------------------
 	// TFixedTickReplicator_Server::NetRecv Server Receiving from AP client
@@ -112,25 +115,34 @@ public:
 		NETSIM_CHECKSUM(Ar);
 
 		const int32 EndFrame = FNetworkPredictionSerialization::ReadCompressedFrame(Ar, ServerRecvData.LastRecvFrame); // 1. StartFrame
-		const int32 StartFrame = EndFrame - NumSendPerUpdate;
+		const int32 StartFrame = EndFrame - NumInputsPerSend;
+		const bool bIsStarvedForInput = ServerRecvData.LastRecvFrame <= ServerRecvData.LastConsumedFrame;
 
 		for (int32 Frame=StartFrame; Frame < EndFrame; ++Frame)
 		{
-			if (Frame <= ServerRecvData.LastRecvFrame)
+			if (Frame <= ServerRecvData.LastRecvFrame || Frame <= ServerRecvData.LastConsumedFrame)
 			{
 				EatCmd(P);
 			}
 			else
 			{
-				for (int32 DroppedFrame = ServerRecvData.LastRecvFrame+1; DroppedFrame < Frame; ++DroppedFrame)
+				if (!bIsStarvedForInput)
 				{
-					UE_NP_TRACE_SYSTEM_FAULT("Gap in input stream detected on server. LastRecvFrame: %d. New Frame: %d", ServerRecvData.LastRecvFrame, DroppedFrame);
-					if (DroppedFrame > 0)
+					for (int32 DroppedFrame = ServerRecvData.LastRecvFrame+1; DroppedFrame < Frame; ++DroppedFrame)
 					{
-						// FixedTick can't skip frames like independent, so copy previous input
-						ServerRecvData.InputBuffer[DroppedFrame] = ServerRecvData.InputBuffer[DroppedFrame-1];
+						UE_NP_TRACE_SYSTEM_FAULT("Gap in input stream detected on server. Client frames involved: LastConsumedFrame: %d LastRecvFrame: %d. DroppedFrame: %d", ServerRecvData.LastConsumedFrame, ServerRecvData.LastRecvFrame, DroppedFrame);
+						if (DroppedFrame > 0)
+						{
+							// FixedTick can't skip frames like independent, so copy previous input
+							ServerRecvData.InputBuffer[DroppedFrame] = ServerRecvData.InputBuffer[DroppedFrame-1];
+						}
 					}
 				}
+				else
+				{
+					UE_NP_TRACE_SYSTEM_FAULT("Recovering from input stream starvation on server. Advancing over dropped client frames [%d-%d]", ServerRecvData.LastConsumedFrame+1, StartFrame-1);
+					ServerRecvData.LastConsumedFrame = StartFrame-1;
+				}
 
 				npEnsure(Frame >= 0);
 
@@ -161,7 +173,7 @@ public:
 		FNetworkPredictionSerialization::WriteCompressedFrame(Ar, TickState->PendingFrame); // 1. Client's PendingFrame number
 
 		const int32 EndFrame = TickState->PendingFrame; // PendingFrame doesn't have an input written until right before it ticks, so don't send it's contents
-		const int32 StartFrame = EndFrame - NumSendPerUpdate;
+		const int32 StartFrame = EndFrame - NumInputsPerSend;
 		
 		for (int32 Frame = StartFrame; Frame < EndFrame; ++Frame)
 		{
@@ -190,6 +202,9 @@ public:
 		TConditionalState<InputType> Empty;
 		FNetworkPredictionDriver<ModelDef>::NetSerialize(Empty, P);  // 2. InputCmd	
 	}
+
+private:
+	inline static int32 NumInputsPerSend = 6;
 };
 
 template<typename InModelDef>
@@ -202,7 +217,10 @@ public:
 	using StateTypes = typename ModelDef::StateTypes;
 	using InputType= typename StateTypes::InputType;
 
-	static const int32 NumSendPerUpdate = 6;
+	static void SetNumInputsPerSend(int32 NumToSend) 
+	{ 
+		NumInputsPerSend=NumToSend; 
+	}
 
 	// ------------------------------------------------------------------------------------------------------------
 	// TIndependentTickReplicator_Server::NetRecv AP Client sending to server
@@ -213,7 +231,7 @@ public:
 		NETSIM_CHECKSUM(Ar);
 
 		const int32 EndFrame = FNetworkPredictionSerialization::ReadCompressedFrame(Ar, ServerRecvData.LastRecvFrame); // 1. StartFrame
-		const int32 StartFrame = EndFrame - NumSendPerUpdate;
+		const int32 StartFrame = EndFrame - NumInputsPerSend;
 
 		// Reset consumed frame if we detect a gap.
 		// Note this could discard unprocessed commands we previously received (but didn't process) but handling this case doesn't seem necessary or practical
@@ -280,7 +298,7 @@ public:
 		FNetworkPredictionSerialization::WriteCompressedFrame(Ar, TickState->PendingFrame); // 1. Client's PendingFrame number
 
 		const int32 EndFrame = TickState->PendingFrame; // PendingFrame doesn't have an input written until right before it ticks, so don't send it's contents
-		const int32 StartFrame = EndFrame - NumSendPerUpdate;
+		const int32 StartFrame = EndFrame - NumInputsPerSend;
 
 		for (int32 Frame = StartFrame; Frame < EndFrame; ++Frame)
 		{
@@ -316,6 +334,8 @@ private:
 		int32 TimeMS = 0;
 		FNetworkPredictionSerialization::SerializeDeltaMS(P.Ar, TimeMS); // 3. Delta InputCmd
 	}
+
+	inline static int32 NumInputsPerSend = 6;
 };
 
 // ---------------------------------------------------------------------------------------------------------------------------
diff --git a/Engine/Plugins/Runtime/NetworkPrediction/Source/NetworkPrediction/Public/NetworkPredictionSettings.h b/Engine/Plugins/Runtime/NetworkPrediction/Source/NetworkPrediction/Public/NetworkPredictionSettings.h
--- a/Engine/Plugins/Runtime/NetworkPrediction/Source/NetworkPrediction/Public/NetworkPredictionSettings.h
+++ b/Engine/Plugins/Runtime/NetworkPrediction/Source/NetworkPrediction/Public/NetworkPredictionSettings.h
@@ -51,6 +51,20 @@ struct FNetworkPredictionSettings
 	// Max buffered time to keep for fixed independent interpolated sims (client only).
 	UPROPERTY(config, EditAnywhere, Category = Interpolation)
 	int32 IndependentTickInterpolationMaxBufferedMS = 250;
+
+	// ------------------------------------------------------------------------------------------
+
+	// Clients should send this number of most recent input commands together per update, for a Fixed Tick simulation
+	UPROPERTY(config, EditAnywhere, Category = Input)
+	int32 FixedTickInputSendCount = 6;
+
+	// Clients should send this number of most recent input commands together per update, for an Independent Tick simulation
+	UPROPERTY(config, EditAnywhere, Category = Input)
+	int32 IndependentTickInputSendCount = 6;
+
+	// Cap the number of remote input commands required to be buffered before resuming input consumption after a fault
+	UPROPERTY(config, EditAnywhere, Category = Input)
+	int32 MaximumRemoteInputFaultLimit = 6;
 };
 
 USTRUCT(meta=(ShowOnlyInnerProperties))
diff --git a/Engine/Plugins/Runtime/NetworkPrediction/Source/NetworkPrediction/Public/NetworkPredictionWorldManager.h b/Engine/Plugins/Runtime/NetworkPrediction/Source/NetworkPrediction/Public/NetworkPredictionWorldManager.h
--- a/Engine/Plugins/Runtime/NetworkPrediction/Source/NetworkPrediction/Public/NetworkPredictionWorldManager.h
+++ b/Engine/Plugins/Runtime/NetworkPrediction/Source/NetworkPrediction/Public/NetworkPredictionWorldManager.h
@@ -197,6 +197,8 @@ void UNetworkPredictionWorldManager::ConfigureInstance(FNetworkPredictionID ID,
 		SetUsingPhysics();
 	}
 
+	TRemoteInputService<ModelDef>::SetMaxFaultLimit(Settings.MaximumRemoteInputFaultLimit);
+
 	ENetworkPredictionService ServiceMask = ENetworkPredictionService::None;
 	
 	if (Archetype.TickingMode == ENetworkPredictionTickingPolicy::Independent)
@@ -448,6 +450,8 @@ void UNetworkPredictionWorldManager::BindServerNetRecv_Fixed(FNetworkPredictionI
 
 	const int32 ServerRecvIdx = DataStore->ServerRecv.GetIndex(ID);
 
+	TFixedTickReplicator_Server<ModelDef>::SetNumInputsPerSend(Settings.FixedTickInputSendCount);
+
 	TServerRecvData_Fixed<ModelDef>& ServerRecvData = DataStore->ServerRecv.GetByIndexChecked(ServerRecvIdx);
 	ServerRecvData.TraceID = ID.GetTraceID();
 
@@ -470,6 +474,8 @@ void UNetworkPredictionWorldManager::BindServerNetRecv_Independent(FNetworkPredi
 
 	const int32 ServerRecvIdx = DataStore->ServerRecv_IndependentTick.GetIndex(ID);
 
+	TIndependentTickReplicator_Server<ModelDef>::SetNumInputsPerSend(Settings.IndependentTickInputSendCount);
+
 	TServerRecvData_Independent<ModelDef>& ServerRecvData = DataStore->ServerRecv_IndependentTick.GetByIndexChecked(ServerRecvIdx);
 	ServerRecvData.TraceID = ID.GetTraceID();
 	ServerRecvData.InstanceIdx = DataStore->Instances.GetIndex(ID);
diff --git a/Engine/Plugins/Runtime/NetworkPrediction/Source/NetworkPrediction/Public/Services/NetworkPredictionService_Input.inl b/Engine/Plugins/Runtime/NetworkPrediction/Source/NetworkPrediction/Public/Services/NetworkPredictionService_Input.inl
--- a/Engine/Plugins/Runtime/NetworkPrediction/Source/NetworkPrediction/Public/Services/NetworkPredictionService_Input.inl
+++ b/Engine/Plugins/Runtime/NetworkPrediction/Source/NetworkPrediction/Public/Services/NetworkPredictionService_Input.inl
@@ -87,6 +87,11 @@ public:
 	TRemoteInputService(TModelDataStore<ModelDef>* InDataStore)
 		: DataStore(InDataStore) { }
 
+	static void SetMaxFaultLimit(int32 InMaxFaultLimit) 
+	{ 
+		MaxFaultLimit = InMaxFaultLimit; 
+	}
+
 	void RegisterInstance(FNetworkPredictionID ID)
 	{
 		TInstanceData<ModelDef>* InstanceData = DataStore->Instances.Find(ID);
@@ -140,6 +145,8 @@ public:
 					UE_NP_TRACE_BUFFERED_INPUT(NumBufferedInputCmds, true);
 					continue;
 				}
+
+				UE_NP_TRACE_SYSTEM_FAULT("[Remote.Input] Recovered from Fault. FaultLimit now: %d", Remote.FaultLimit);
 				Remote.bFault = false;
 			}
 			else if (NumBufferedInputCmds == 0)
@@ -147,8 +154,8 @@ public:
 				// No Cmds to process, enter fault state. Increment FaultLimit each time this happens.
 				// TODO: We should have something to bring this back down (which means skipping frames) we don't want temporary poor conditions to cause permanent high input buffering
 				Remote.bFault = true;
-				Remote.FaultLimit = FMath::Min(Remote.FaultLimit+1, ServerRecvData.InputBuffer.Capacity());
-
+				Remote.FaultLimit = FMath::Min(Remote.FaultLimit + 1, FMath::Min(MaxFaultLimit, ServerRecvData.InputBuffer.Capacity()));
+				
 				UE_NP_TRACE_SYSTEM_FAULT("[Remote.Input] New Fault. No cmds to process. New FaultLimit=%d", Remote.FaultLimit);
 				UE_NP_TRACE_BUFFERED_INPUT(0, true);
 				continue;
@@ -176,5 +183,7 @@ private:
 
 	TSortedMap<int32, FInstance> InstanceMap;
 	TModelDataStore<ModelDef>* DataStore;
+	
+	inline static int32 MaxFaultLimit = 6;
 };
 
