Added markups for headers that can't be compiled standalone.. this is needed for both Header units and IWYU

diff --git a/Engine/Plugins/Animation/LiveLink/Source/LiveLink/Public/LiveLinkBlueprintStructs.h b/Engine/Plugins/Animation/LiveLink/Source/LiveLink/Public/LiveLinkBlueprintStructs.h
--- a/Engine/Plugins/Animation/LiveLink/Source/LiveLink/Public/LiveLinkBlueprintStructs.h
+++ b/Engine/Plugins/Animation/LiveLink/Source/LiveLink/Public/LiveLinkBlueprintStructs.h
@@ -2,6 +2,8 @@
 
 #pragma once
 
+// HEADER_UNIT_SKIP - Deprecated
+
 #include "Roles/LiveLinkAnimationBlueprintStructs.h"
 
 UE_DEPRECATED(4.23, "Including LiveLinkBlueprintStructs.h has been deprecated. Please include Roles/LiveLinkAnimationBlueprintStructs.h instead.")
diff --git a/Engine/Plugins/Animation/RigLogic/Source/RigLogicLib/Private/WindowsPlatformUE.h b/Engine/Plugins/Animation/RigLogic/Source/RigLogicLib/Private/WindowsPlatformUE.h
--- a/Engine/Plugins/Animation/RigLogic/Source/RigLogicLib/Private/WindowsPlatformUE.h
+++ b/Engine/Plugins/Animation/RigLogic/Source/RigLogicLib/Private/WindowsPlatformUE.h
@@ -2,4 +2,6 @@
 
 #pragma once
 
+// HEADER_UNIT_SKIP - Not included directly
+
 #include "Windows/WindowsHWrapper.h"
diff --git a/Engine/Plugins/Animation/RigLogic/Source/RigLogicLib/Private/trio/utils/PlatformWindows.h b/Engine/Plugins/Animation/RigLogic/Source/RigLogicLib/Private/trio/utils/PlatformWindows.h
--- a/Engine/Plugins/Animation/RigLogic/Source/RigLogicLib/Private/trio/utils/PlatformWindows.h
+++ b/Engine/Plugins/Animation/RigLogic/Source/RigLogicLib/Private/trio/utils/PlatformWindows.h
@@ -2,6 +2,8 @@
 
 #pragma once
 
+// HEADER_UNIT_SKIP - Not included directly
+
 #ifdef _MSC_VER
     #pragma warning(push)
     #pragma warning(disable : 4005 4365 4987)
diff --git a/Engine/Plugins/AudioGameplayVolume/Source/AudioGameplayVolume/Public/AudioGameplayVolumeProxyMutator.h b/Engine/Plugins/AudioGameplayVolume/Source/AudioGameplayVolume/Public/AudioGameplayVolumeProxyMutator.h
--- a/Engine/Plugins/AudioGameplayVolume/Source/AudioGameplayVolume/Public/AudioGameplayVolumeProxyMutator.h
+++ b/Engine/Plugins/AudioGameplayVolume/Source/AudioGameplayVolume/Public/AudioGameplayVolumeProxyMutator.h
@@ -11,6 +11,8 @@
 
 #pragma once
 
+// HEADER_UNIT_SKIP - Deprecated
+
 #ifdef _MSC_VER
 #pragma message(__FILE__"(15): warning: use AudioGameplayVolumeSubsystem.h and AudioGameplayVolumeMutator.h instead of AudioGameplayVolumeProxyMutator.h")
 #else
diff --git a/Engine/Plugins/Cameras/GameplayCameras/Source/GameplayCameras/Public/Camera/CameraShake.h b/Engine/Plugins/Cameras/GameplayCameras/Source/GameplayCameras/Public/Camera/CameraShake.h
--- a/Engine/Plugins/Cameras/GameplayCameras/Source/GameplayCameras/Public/Camera/CameraShake.h
+++ b/Engine/Plugins/Cameras/GameplayCameras/Source/GameplayCameras/Public/Camera/CameraShake.h
@@ -2,6 +2,8 @@
 
 #pragma once
 
+// HEADER_UNIT_SKIP - Deprecated
+
 #include "CoreTypes.h"
 
 COMPILE_WARNING("UCameraShake has been renamed to ULegacyCameraShake, change to #include \"LegacyCameraShake.h\" and rename references")
diff --git a/Engine/Plugins/Cameras/GameplayCameras/Source/GameplayCameras/Public/MatineeCameraShake.h b/Engine/Plugins/Cameras/GameplayCameras/Source/GameplayCameras/Public/MatineeCameraShake.h
--- a/Engine/Plugins/Cameras/GameplayCameras/Source/GameplayCameras/Public/MatineeCameraShake.h
+++ b/Engine/Plugins/Cameras/GameplayCameras/Source/GameplayCameras/Public/MatineeCameraShake.h
@@ -2,6 +2,8 @@
 
 #pragma once
 
+// HEADER_UNIT_SKIP - Deprecated
+
 #include "CoreTypes.h"
 
 COMPILE_WARNING("UMatineeCameraShake has been renamed to ULegacyCameraShake, change to #include \"LegacyCameraShake.h\" and rename references")
diff --git a/Engine/Plugins/Enterprise/MDLImporter/Source/MDLImporter/Private/mdl/MdlSdkDefines.h b/Engine/Plugins/Enterprise/MDLImporter/Source/MDLImporter/Private/mdl/MdlSdkDefines.h
--- a/Engine/Plugins/Enterprise/MDLImporter/Source/MDLImporter/Private/mdl/MdlSdkDefines.h
+++ b/Engine/Plugins/Enterprise/MDLImporter/Source/MDLImporter/Private/mdl/MdlSdkDefines.h
@@ -2,6 +2,8 @@
 
 #pragma once
 
+// HEADER_UNIT_SKIP - Special header
+
 #if PLATFORM_WINDOWS
 
 #ifndef MDLSDK_INCLUDES_START
diff --git a/Engine/Plugins/Experimental/ChaosEditor/Source/FractureEditor/Private/SGeometryCollectionTreeView.h b/Engine/Plugins/Experimental/ChaosEditor/Source/FractureEditor/Private/SGeometryCollectionTreeView.h
--- a/Engine/Plugins/Experimental/ChaosEditor/Source/FractureEditor/Private/SGeometryCollectionTreeView.h
+++ b/Engine/Plugins/Experimental/ChaosEditor/Source/FractureEditor/Private/SGeometryCollectionTreeView.h
@@ -2,6 +2,8 @@
 
 #pragma once
 
+// HEADER_UNIT_SKIP - Bad include
+
 #include "Widgets/Views/STreeView.h"
 #include "SGeometryCollectionTreeItem.h"
 
diff --git a/Engine/Plugins/Experimental/CommonUI/Source/CommonUI/Public/Input/CommonInputMode.h b/Engine/Plugins/Experimental/CommonUI/Source/CommonUI/Public/Input/CommonInputMode.h
--- a/Engine/Plugins/Experimental/CommonUI/Source/CommonUI/Public/Input/CommonInputMode.h
+++ b/Engine/Plugins/Experimental/CommonUI/Source/CommonUI/Public/Input/CommonInputMode.h
@@ -2,6 +2,8 @@
 
 #pragma once
 
+// HEADER_UNIT_SKIP - Deprecated
+
 #ifdef _MSC_VER
 #pragma message(__FILE__"(8): warning: use CommonInputModeTypes.h from CommonInput module instead of Input/CommonInputMode.h from CommonUI module")
 #else
diff --git a/Engine/Plugins/Experimental/ControlFlows/Source/ControlFlows/Public/ControlFlowContainer.h b/Engine/Plugins/Experimental/ControlFlows/Source/ControlFlows/Public/ControlFlowContainer.h
--- a/Engine/Plugins/Experimental/ControlFlows/Source/ControlFlows/Public/ControlFlowContainer.h
+++ b/Engine/Plugins/Experimental/ControlFlows/Source/ControlFlows/Public/ControlFlowContainer.h
@@ -2,6 +2,8 @@
 
 #pragma once
 
+// HEADER_UNIT_SKIP - Not included directly
+
 #include "Templates/SharedPointer.h"
 
 class FControlFlowContainerBase : public TSharedFromThis<FControlFlowContainerBase>
diff --git a/Engine/Plugins/Experimental/Mutable/Source/MutableRuntime/Public/MuR/RefCounted.h b/Engine/Plugins/Experimental/Mutable/Source/MutableRuntime/Public/MuR/RefCounted.h
--- a/Engine/Plugins/Experimental/Mutable/Source/MutableRuntime/Public/MuR/RefCounted.h
+++ b/Engine/Plugins/Experimental/Mutable/Source/MutableRuntime/Public/MuR/RefCounted.h
@@ -2,6 +2,8 @@
 
 #pragma once
 
+// HEADER_UNIT_SKIP - 
+
 #include "MuR/MutableMemory.h"
 #include "MuR/Ptr.h"
 
diff --git a/Engine/Plugins/Experimental/VirtualCamera/Source/PixelStreamingVCam/Private/PixelStreamingVCamUtil.h b/Engine/Plugins/Experimental/VirtualCamera/Source/PixelStreamingVCam/Private/PixelStreamingVCamUtil.h
--- a/Engine/Plugins/Experimental/VirtualCamera/Source/PixelStreamingVCam/Private/PixelStreamingVCamUtil.h
+++ b/Engine/Plugins/Experimental/VirtualCamera/Source/PixelStreamingVCam/Private/PixelStreamingVCamUtil.h
@@ -2,6 +2,8 @@
 
 #pragma once
 
+// HEADER_UNIT_SKIP - Bad header
+
 #include "Misc/Paths.h"
 #include "Interfaces/IPluginManager.h"
 #include "PixelStreamingVCamLogging.h"
diff --git a/Engine/Plugins/Importers/USDImporter/Source/UnrealUSDWrapper/Public/USDIncludesEnd.h b/Engine/Plugins/Importers/USDImporter/Source/UnrealUSDWrapper/Public/USDIncludesEnd.h
--- a/Engine/Plugins/Importers/USDImporter/Source/UnrealUSDWrapper/Public/USDIncludesEnd.h
+++ b/Engine/Plugins/Importers/USDImporter/Source/UnrealUSDWrapper/Public/USDIncludesEnd.h
@@ -1,5 +1,7 @@
 // Copyright Epic Games, Inc. All Rights Reserved.
 
+// HEADER_UNIT_SKIP - Special include
+
 #if USE_USD_SDK
 
 #pragma warning(pop)
diff --git a/Engine/Plugins/Importers/USDImporter/Source/UnrealUSDWrapper/Public/USDIncludesStart.h b/Engine/Plugins/Importers/USDImporter/Source/UnrealUSDWrapper/Public/USDIncludesStart.h
--- a/Engine/Plugins/Importers/USDImporter/Source/UnrealUSDWrapper/Public/USDIncludesStart.h
+++ b/Engine/Plugins/Importers/USDImporter/Source/UnrealUSDWrapper/Public/USDIncludesStart.h
@@ -1,5 +1,7 @@
 // Copyright Epic Games, Inc. All Rights Reserved.
 
+// HEADER_UNIT_SKIP - Special include
+
 #if USE_USD_SDK
 
 #include "CoreMinimal.h"
diff --git a/Engine/Plugins/Online/OnlineServices/Source/OnlineServicesInterface/Public/Online/OnlineUtils.h b/Engine/Plugins/Online/OnlineServices/Source/OnlineServicesInterface/Public/Online/OnlineUtils.h
--- a/Engine/Plugins/Online/OnlineServices/Source/OnlineServicesInterface/Public/Online/OnlineUtils.h
+++ b/Engine/Plugins/Online/OnlineServices/Source/OnlineServicesInterface/Public/Online/OnlineUtils.h
@@ -2,6 +2,8 @@
 
 #pragma once
 
+// HEADER_UNIT_SKIP - Bad includes
+
 #include "Online/OnlineAsyncOp.h"
 #include "Online/OnlineMeta.h"
 #include "Containers/UnrealString.h"
diff --git a/Engine/Plugins/Online/OnlineSubsystem/Source/Public/OnlineErrorMacros.inl b/Engine/Plugins/Online/OnlineSubsystem/Source/Public/OnlineErrorMacros.inl
--- a/Engine/Plugins/Online/OnlineSubsystem/Source/Public/OnlineErrorMacros.inl
+++ b/Engine/Plugins/Online/OnlineSubsystem/Source/Public/OnlineErrorMacros.inl
@@ -1,5 +1,8 @@
 // Copyright Epic Games, Inc. All Rights Reserved.
 
+// HEADER_UNIT_SKIP - Is included in ways that does not support header units
+// IWYU pragma: private
+
 #ifndef ONLINE_ERROR_NAMESPACE
 #error "ONLINE_ERROR_NAMESPACE not defined"
 #endif
diff --git a/Engine/Plugins/Runtime/Metasound/Source/MetasoundGraphCore/Public/MetasoundNodeRegister.h b/Engine/Plugins/Runtime/Metasound/Source/MetasoundGraphCore/Public/MetasoundNodeRegister.h
--- a/Engine/Plugins/Runtime/Metasound/Source/MetasoundGraphCore/Public/MetasoundNodeRegister.h
+++ b/Engine/Plugins/Runtime/Metasound/Source/MetasoundGraphCore/Public/MetasoundNodeRegister.h
@@ -2,6 +2,8 @@
 
 #pragma once
 
+// HEADER_UNIT_SKIP - Bad header
+
 #include "CoreMinimal.h"
 
 namespace Metasound
diff --git a/Engine/Plugins/Runtime/NetworkPrediction/Source/NetworkPrediction/Public/Services/NetworkPredictionService_Rollback.inl b/Engine/Plugins/Runtime/NetworkPrediction/Source/NetworkPrediction/Public/Services/NetworkPredictionService_Rollback.inl
--- a/Engine/Plugins/Runtime/NetworkPrediction/Source/NetworkPrediction/Public/Services/NetworkPredictionService_Rollback.inl
+++ b/Engine/Plugins/Runtime/NetworkPrediction/Source/NetworkPrediction/Public/Services/NetworkPredictionService_Rollback.inl
@@ -1,6 +1,10 @@
 // Copyright Epic Games, Inc. All Rights Reserved.
 
 #pragma once
+
+// HEADER_UNIT_SKIP - Not included directly
+
+#include "NetworkPredictionTickState.h"
 #include "Stats/Stats.h"
 
 namespace NetworkPredictionCVars
diff --git a/Engine/Plugins/Runtime/NetworkPrediction/Source/NetworkPrediction/Public/Services/NetworkPredictionService_Ticking.inl b/Engine/Plugins/Runtime/NetworkPrediction/Source/NetworkPrediction/Public/Services/NetworkPredictionService_Ticking.inl
--- a/Engine/Plugins/Runtime/NetworkPrediction/Source/NetworkPrediction/Public/Services/NetworkPredictionService_Ticking.inl
+++ b/Engine/Plugins/Runtime/NetworkPrediction/Source/NetworkPrediction/Public/Services/NetworkPredictionService_Ticking.inl
@@ -1,8 +1,12 @@
 // Copyright Epic Games, Inc. All Rights Reserved.
 
 #pragma once
+
+// HEADER_UNIT_SKIP - Not included directly
+
 #include "NetworkPredictionSimulation.h"
 #include "NetworkPredictionInstanceData.h"
+#include "NetworkPredictionTickState.h"
 
 // Common util used by the ticking services. Might make sense to move to FNetworkPredictionDriverBase if needed elsewhere
 template<typename ModelDef>
diff --git a/Engine/Plugins/Runtime/NetworkPredictionInsights/Source/NetworkPredictionInsights/Private/NetworkPredictionInsightsStyle.h b/Engine/Plugins/Runtime/NetworkPredictionInsights/Source/NetworkPredictionInsights/Private/NetworkPredictionInsightsStyle.h
--- a/Engine/Plugins/Runtime/NetworkPredictionInsights/Source/NetworkPredictionInsights/Private/NetworkPredictionInsightsStyle.h
+++ b/Engine/Plugins/Runtime/NetworkPredictionInsights/Source/NetworkPredictionInsights/Private/NetworkPredictionInsightsStyle.h
@@ -2,6 +2,8 @@
 
 #pragma once
 
+// HEADER_UNIT_SKIP - TODO
+
 #include "Styling/SlateStyle.h"
 #include "Styling/SlateStyleRegistry.h"
 
diff --git a/Engine/Plugins/Runtime/OpenCV/Source/OpenCVHelper/Public/PostOpenCVHeaders.h b/Engine/Plugins/Runtime/OpenCV/Source/OpenCVHelper/Public/PostOpenCVHeaders.h
--- a/Engine/Plugins/Runtime/OpenCV/Source/OpenCVHelper/Public/PostOpenCVHeaders.h
+++ b/Engine/Plugins/Runtime/OpenCV/Source/OpenCVHelper/Public/PostOpenCVHeaders.h
@@ -6,6 +6,8 @@
 #error Mismatched PreOpenCVHeadersTypes.h detected.
 #endif
 
+// HEADER_UNIT_SKIP - Special include
+
 #include "OpenCVUtils.h"
 
 #if PLATFORM_WINDOWS
diff --git a/Engine/Plugins/Runtime/OpenXR/Source/OpenXRHMD/Public/IOpenXRExtensionPluginDelegates.h b/Engine/Plugins/Runtime/OpenXR/Source/OpenXRHMD/Public/IOpenXRExtensionPluginDelegates.h
--- a/Engine/Plugins/Runtime/OpenXR/Source/OpenXRHMD/Public/IOpenXRExtensionPluginDelegates.h
+++ b/Engine/Plugins/Runtime/OpenXR/Source/OpenXRHMD/Public/IOpenXRExtensionPluginDelegates.h
@@ -2,6 +2,8 @@
 
 #pragma once
 
+// HEADER_UNIT_SKIP - Skip this for now
+
 #include "CoreMinimal.h"
 
 struct XrSpaceLocation;
diff --git a/Engine/Source/Developer/Android/AndroidTargetPlatform/Private/AndroidTargetDevice.h b/Engine/Source/Developer/Android/AndroidTargetPlatform/Private/AndroidTargetDevice.h
--- a/Engine/Source/Developer/Android/AndroidTargetPlatform/Private/AndroidTargetDevice.h
+++ b/Engine/Source/Developer/Android/AndroidTargetPlatform/Private/AndroidTargetDevice.h
@@ -254,4 +254,4 @@ protected:
 };
 
 
-#include "AndroidTargetDevice.inl"
+#include "AndroidTargetDevice.inl" // IWYU pragma: export
diff --git a/Engine/Source/Developer/Android/AndroidTargetPlatform/Private/AndroidTargetDevice.inl b/Engine/Source/Developer/Android/AndroidTargetPlatform/Private/AndroidTargetDevice.inl
--- a/Engine/Source/Developer/Android/AndroidTargetPlatform/Private/AndroidTargetDevice.inl
+++ b/Engine/Source/Developer/Android/AndroidTargetPlatform/Private/AndroidTargetDevice.inl
@@ -1,5 +1,7 @@
 // Copyright Epic Games, Inc. All Rights Reserved.
 
+// HEADER_UNIT_SKIP - Not included directly
+
 /* ITargetDevice interface
  *****************************************************************************/
 #include "CoreTypes.h"
diff --git a/Engine/Source/Developer/Android/AndroidTargetPlatform/Private/AndroidTargetDeviceOutput.h b/Engine/Source/Developer/Android/AndroidTargetPlatform/Private/AndroidTargetDeviceOutput.h
--- a/Engine/Source/Developer/Android/AndroidTargetPlatform/Private/AndroidTargetDeviceOutput.h
+++ b/Engine/Source/Developer/Android/AndroidTargetPlatform/Private/AndroidTargetDeviceOutput.h
@@ -2,6 +2,8 @@
 
 #pragma once
 
+// HEADER_UNIT_SKIP - Not included directly
+
 #include "Containers/UnrealString.h"
 #include "HAL/Runnable.h"
 #include "HAL/RunnableThread.h"
@@ -53,4 +55,4 @@ private:
 	FString											DeviceName;
 };
 
-#include "AndroidTargetDeviceOutput.inl"
+#include "AndroidTargetDeviceOutput.inl" // IWYU pragma: export
diff --git a/Engine/Source/Developer/Android/AndroidTargetPlatform/Private/AndroidTargetDeviceOutput.inl b/Engine/Source/Developer/Android/AndroidTargetPlatform/Private/AndroidTargetDeviceOutput.inl
--- a/Engine/Source/Developer/Android/AndroidTargetPlatform/Private/AndroidTargetDeviceOutput.inl
+++ b/Engine/Source/Developer/Android/AndroidTargetPlatform/Private/AndroidTargetDeviceOutput.inl
@@ -1,5 +1,7 @@
 // Copyright Epic Games, Inc. All Rights Reserved.
 
+// HEADER_UNIT_SKIP - Not included directly
+
 #include "CoreTypes.h"
 #include "CoreFwd.h"
 
diff --git a/Engine/Source/Developer/AnimationWidgets/Public/SAdvancedTransformInputBox.h b/Engine/Source/Developer/AnimationWidgets/Public/SAdvancedTransformInputBox.h
--- a/Engine/Source/Developer/AnimationWidgets/Public/SAdvancedTransformInputBox.h
+++ b/Engine/Source/Developer/AnimationWidgets/Public/SAdvancedTransformInputBox.h
@@ -2,6 +2,8 @@
 
 #pragma once
 
+// HEADER_UNIT_SKIP - Bad include
+
 #include "SAdvancedRotationInputBox.h"
 #include "Widgets/Input/SComboButton.h"
 #include "Widgets/Input/SButton.h"
diff --git a/Engine/Source/Developer/DesktopPlatform/Public/SteamDeck/SteamDeckDevice.h b/Engine/Source/Developer/DesktopPlatform/Public/SteamDeck/SteamDeckDevice.h
--- a/Engine/Source/Developer/DesktopPlatform/Public/SteamDeck/SteamDeckDevice.h
+++ b/Engine/Source/Developer/DesktopPlatform/Public/SteamDeck/SteamDeckDevice.h
@@ -2,6 +2,8 @@
 
 #pragma once
 
+// HEADER_UNIT_UNSUPPORTED - Common/TargetPlatformBase not part of additional include paths
+
 #include "CoreMinimal.h"
 #include "Common/TargetPlatformBase.h"
 #include "Misc/ConfigCacheIni.h"
diff --git a/Engine/Source/Developer/MeshSimplifier/Private/Cache.h b/Engine/Source/Developer/MeshSimplifier/Private/Cache.h
--- a/Engine/Source/Developer/MeshSimplifier/Private/Cache.h
+++ b/Engine/Source/Developer/MeshSimplifier/Private/Cache.h
@@ -3,6 +3,8 @@
 
 #pragma once
 
+// HEADER_UNIT_SKIP - Bad header.. code is not compiling
+
 #include "CoreTypes.h"
 #include "HAL/UnrealMemory.h"
 
diff --git a/Engine/Source/Developer/ShaderCompilerCommon/Public/ShaderSymbolExport.h b/Engine/Source/Developer/ShaderCompilerCommon/Public/ShaderSymbolExport.h
--- a/Engine/Source/Developer/ShaderCompilerCommon/Public/ShaderSymbolExport.h
+++ b/Engine/Source/Developer/ShaderCompilerCommon/Public/ShaderSymbolExport.h
@@ -2,6 +2,8 @@
 
 #pragma once
 
+// HEADER_UNIT_SKIP - Bad includes (ZipArchiveWriter.h)
+
 #if WITH_ENGINE
 
 #include "CoreTypes.h"
diff --git a/Engine/Source/Developer/TextureFormat/Public/ChildTextureFormat.h b/Engine/Source/Developer/TextureFormat/Public/ChildTextureFormat.h
--- a/Engine/Source/Developer/TextureFormat/Public/ChildTextureFormat.h
+++ b/Engine/Source/Developer/TextureFormat/Public/ChildTextureFormat.h
@@ -2,6 +2,8 @@
 
 #pragma once
 
+// HEADER_UNIT_SKIP - Bad include TextureCompressorModule.h
+
 #include "Interfaces/ITextureFormat.h"
 #include "Interfaces/ITextureFormatModule.h"
 #include "Interfaces/ITextureFormatManagerModule.h"
diff --git a/Engine/Source/Runtime/AppFramework/Private/Widgets/Colors/SColorValueSlider.h b/Engine/Source/Runtime/AppFramework/Private/Widgets/Colors/SColorValueSlider.h
--- a/Engine/Source/Runtime/AppFramework/Private/Widgets/Colors/SColorValueSlider.h
+++ b/Engine/Source/Runtime/AppFramework/Private/Widgets/Colors/SColorValueSlider.h
@@ -2,6 +2,8 @@
 
 #pragma once
 
+// HEADER_UNIT_SKIP - Bad header, does not compile
+
 #include "CoreMinimal.h"
 #include "Misc/Attribute.h"
 #include "InputCoreTypes.h"
diff --git a/Engine/Source/Runtime/ApplicationCore/Public/GenericPlatform/IForceFeedbackSystem.h b/Engine/Source/Runtime/ApplicationCore/Public/GenericPlatform/IForceFeedbackSystem.h
--- a/Engine/Source/Runtime/ApplicationCore/Public/GenericPlatform/IForceFeedbackSystem.h
+++ b/Engine/Source/Runtime/ApplicationCore/Public/GenericPlatform/IForceFeedbackSystem.h
@@ -6,7 +6,9 @@
  * DEPRECATED GenericPlatform/IForceFeedbackSystem.h. Use GenericPlatform/IInputInterface.h instead
  */
 #pragma once
-  
+
+// HEADER_UNIT_SKIP - Deprecated
+
 #ifdef _MSC_VER
 	#pragma message(__FILE__"(9): warning: use GenericPlatform/IInputInterface.h instead of GenericPlatform/IForceFeedbackSystem.h")
 #else
diff --git a/Engine/Source/Runtime/AssetRegistry/Public/ARFilter.h b/Engine/Source/Runtime/AssetRegistry/Public/ARFilter.h
--- a/Engine/Source/Runtime/AssetRegistry/Public/ARFilter.h
+++ b/Engine/Source/Runtime/AssetRegistry/Public/ARFilter.h
@@ -2,6 +2,8 @@
 
 #pragma once
 
+// HEADER_UNIT_SKIP - Deprecated
+
 // Deprecated in UE 5.1.
 #pragma message (__FILE__"(6): warning: #include AssetRegistry/ARFilter.h instead of ARFilter.h")
 
diff --git a/Engine/Source/Runtime/AssetRegistry/Public/AssetBundleData.h b/Engine/Source/Runtime/AssetRegistry/Public/AssetBundleData.h
--- a/Engine/Source/Runtime/AssetRegistry/Public/AssetBundleData.h
+++ b/Engine/Source/Runtime/AssetRegistry/Public/AssetBundleData.h
@@ -2,6 +2,8 @@
 
 #pragma once
 
+// HEADER_UNIT_SKIP - Deprecated
+
 // Deprecated in UE 5.1.
 #pragma message (__FILE__"(6): warning: #include AssetRegistry/AssetBundleData.h instead of AssetBundleData.h")
 
diff --git a/Engine/Source/Runtime/AssetRegistry/Public/AssetDataTagMap.h b/Engine/Source/Runtime/AssetRegistry/Public/AssetDataTagMap.h
--- a/Engine/Source/Runtime/AssetRegistry/Public/AssetDataTagMap.h
+++ b/Engine/Source/Runtime/AssetRegistry/Public/AssetDataTagMap.h
@@ -2,6 +2,8 @@
 
 #pragma once
 
+// HEADER_UNIT_SKIP - Should not include this
+
 // Deprecated in UE 5.1.
 #pragma message (__FILE__"(6): warning: #include AssetRegistry/AssetDataTagMap.h instead of AssetDataTagMap.h")
 
diff --git a/Engine/Source/Runtime/AssetRegistry/Public/IAssetRegistry.h b/Engine/Source/Runtime/AssetRegistry/Public/IAssetRegistry.h
--- a/Engine/Source/Runtime/AssetRegistry/Public/IAssetRegistry.h
+++ b/Engine/Source/Runtime/AssetRegistry/Public/IAssetRegistry.h
@@ -2,6 +2,8 @@
 
 #pragma once
 
+// HEADER_UNIT_SKIP - Deprecated
+
 // Deprecated in UE 5.1.
 #pragma message (__FILE__"(6): warning: #include AssetRegistry/IAssetRegistry.h instead of IAssetRegistry.h")
 
diff --git a/Engine/Source/Runtime/Core/Public/Containers/BasicArray.h b/Engine/Source/Runtime/Core/Public/Containers/BasicArray.h
--- a/Engine/Source/Runtime/Core/Public/Containers/BasicArray.h
+++ b/Engine/Source/Runtime/Core/Public/Containers/BasicArray.h
@@ -2,6 +2,8 @@
 
 #pragma once
 
+// HEADER_UNIT_SKIP - Not used?
+
 #include "Templates/MemoryOps.h"
 
 #if _MSC_VER
diff --git a/Engine/Source/Runtime/Core/Public/Containers/ClosableMpscQueue.h b/Engine/Source/Runtime/Core/Public/Containers/ClosableMpscQueue.h
--- a/Engine/Source/Runtime/Core/Public/Containers/ClosableMpscQueue.h
+++ b/Engine/Source/Runtime/Core/Public/Containers/ClosableMpscQueue.h
@@ -2,6 +2,8 @@
 
 #pragma once
 
+// HEADER_UNIT_SKIP - Not included directly
+
 #include "CoreTypes.h"
 #include "Templates/TypeCompatibleBytes.h"
 #include <atomic>
diff --git a/Engine/Source/Runtime/Core/Public/Containers/DepletableMpscQueue.h b/Engine/Source/Runtime/Core/Public/Containers/DepletableMpscQueue.h
--- a/Engine/Source/Runtime/Core/Public/Containers/DepletableMpscQueue.h
+++ b/Engine/Source/Runtime/Core/Public/Containers/DepletableMpscQueue.h
@@ -2,6 +2,7 @@
 
 #pragma once
 
+// HEADER_UNIT_SKIP - Deprecated
 #ifdef _MSC_VER
 #pragma message(__FILE__"(6): warning : #include Containers/DepletableMpmcQueue.h instead of Containers/DepletableMpscQueue.h (/s/Mpsc/Mpmc)")
 #else
diff --git a/Engine/Source/Runtime/Core/Public/Delegates/DelegateSignatureImpl.inl b/Engine/Source/Runtime/Core/Public/Delegates/DelegateSignatureImpl.inl
--- a/Engine/Source/Runtime/Core/Public/Delegates/DelegateSignatureImpl.inl
+++ b/Engine/Source/Runtime/Core/Public/Delegates/DelegateSignatureImpl.inl
@@ -6,6 +6,9 @@
 #endif
 
 #pragma once
+
+// HEADER_UNIT_SKIP - Not included directly
+
 #include "CoreTypes.h"
 #include "Misc/AssertionMacros.h"
 #include "Misc/Crc.h"
diff --git a/Engine/Source/Runtime/Core/Public/GenericPlatform/MicrosoftPlatformString.h b/Engine/Source/Runtime/Core/Public/GenericPlatform/MicrosoftPlatformString.h
--- a/Engine/Source/Runtime/Core/Public/GenericPlatform/MicrosoftPlatformString.h
+++ b/Engine/Source/Runtime/Core/Public/GenericPlatform/MicrosoftPlatformString.h
@@ -2,6 +2,8 @@
 
 #pragma once
 
+// HEADER_UNIT_SKIP - Included through other header
+
 #include "Misc/Char.h"
 #if PLATFORM_USE_GENERIC_STRING_IMPLEMENTATION
 	#include "GenericPlatform/GenericWidePlatformString.h"
diff --git a/Engine/Source/Runtime/Core/Public/Linux/LinuxPlatform.h b/Engine/Source/Runtime/Core/Public/Linux/LinuxPlatform.h
--- a/Engine/Source/Runtime/Core/Public/Linux/LinuxPlatform.h
+++ b/Engine/Source/Runtime/Core/Public/Linux/LinuxPlatform.h
@@ -1,6 +1,8 @@
 // Copyright Epic Games, Inc. All Rights Reserved.
 #pragma once
 
+// HEADER_UNIT_SKIP - Not included directly
+
 // Base defines, must define these for the platform, there are no defaults
 #ifndef PLATFORM_DESKTOP
 	#define PLATFORM_DESKTOP					1
diff --git a/Engine/Source/Runtime/Core/Public/MSVC/MSVCPlatform.h b/Engine/Source/Runtime/Core/Public/MSVC/MSVCPlatform.h
--- a/Engine/Source/Runtime/Core/Public/MSVC/MSVCPlatform.h
+++ b/Engine/Source/Runtime/Core/Public/MSVC/MSVCPlatform.h
@@ -6,6 +6,8 @@
 
 #pragma once
 
+// HEADER_UNIT_SKIP - Not included directly
+
 #if _MSC_VER < 1920
 	#error "Compiler is expected to support if constexpr"
 #endif
diff --git a/Engine/Source/Runtime/Core/Public/Math/UnrealPlatformMathSSE.h b/Engine/Source/Runtime/Core/Public/Math/UnrealPlatformMathSSE.h
--- a/Engine/Source/Runtime/Core/Public/Math/UnrealPlatformMathSSE.h
+++ b/Engine/Source/Runtime/Core/Public/Math/UnrealPlatformMathSSE.h
@@ -2,6 +2,8 @@
 
 #pragma once
 
+// HEADER_UNIT_SKIP - Included through other header
+
 #include "CoreTypes.h"
 #include "HAL/PlatformMisc.h"
 // Code including this header is responsible for including the correct platform-specific header for SSE intrinsics.
diff --git a/Engine/Source/Runtime/Core/Public/Math/UnrealPlatformMathSSE4.h b/Engine/Source/Runtime/Core/Public/Math/UnrealPlatformMathSSE4.h
--- a/Engine/Source/Runtime/Core/Public/Math/UnrealPlatformMathSSE4.h
+++ b/Engine/Source/Runtime/Core/Public/Math/UnrealPlatformMathSSE4.h
@@ -2,6 +2,8 @@
 
 #pragma once
 
+// HEADER_UNIT_SKIP - Included through other header
+
 #include "HAL/Platform.h"
 #include "Math/UnrealPlatformMathSSE.h"
 // Code including this header is responsible for including the correct platform-specific header for SSE intrinsics.
diff --git a/Engine/Source/Runtime/Core/Public/Templates/SharedPointerInternals.h b/Engine/Source/Runtime/Core/Public/Templates/SharedPointerInternals.h
--- a/Engine/Source/Runtime/Core/Public/Templates/SharedPointerInternals.h
+++ b/Engine/Source/Runtime/Core/Public/Templates/SharedPointerInternals.h
@@ -2,6 +2,8 @@
 
 #pragma once
 
+// HEADER_UNIT_SKIP - Included through SharedPointer.h
+
 #include "CoreTypes.h"
 #include "HAL/UnrealMemory.h"
 #include "Misc/AssertionMacros.h"
diff --git a/Engine/Source/Runtime/Core/Public/Unix/UnixPlatform.h b/Engine/Source/Runtime/Core/Public/Unix/UnixPlatform.h
--- a/Engine/Source/Runtime/Core/Public/Unix/UnixPlatform.h
+++ b/Engine/Source/Runtime/Core/Public/Unix/UnixPlatform.h
@@ -6,6 +6,8 @@
 
 #pragma once
 
+// HEADER_UNIT_SKIP - Not included directly
+
 #include <linux/version.h>
 #include <signal.h>
 #include "Clang/ClangPlatform.h"
diff --git a/Engine/Source/Runtime/CoreUObject/Public/Serialization/ArchiveReplaceObjectAndStructPropertyRef.h b/Engine/Source/Runtime/CoreUObject/Public/Serialization/ArchiveReplaceObjectAndStructPropertyRef.h
--- a/Engine/Source/Runtime/CoreUObject/Public/Serialization/ArchiveReplaceObjectAndStructPropertyRef.h
+++ b/Engine/Source/Runtime/CoreUObject/Public/Serialization/ArchiveReplaceObjectAndStructPropertyRef.h
@@ -2,6 +2,8 @@
 
 #pragma once
 
+// HEADER_UNIT_SKIP - Internal
+
 #include "CoreMinimal.h"
 #include "Serialization/ArchiveReplaceObjectRef.h"
 
diff --git a/Engine/Source/Runtime/CoreUObject/Public/UObject/CoreOnlineFwd.h b/Engine/Source/Runtime/CoreUObject/Public/UObject/CoreOnlineFwd.h
--- a/Engine/Source/Runtime/CoreUObject/Public/UObject/CoreOnlineFwd.h
+++ b/Engine/Source/Runtime/CoreUObject/Public/UObject/CoreOnlineFwd.h
@@ -2,4 +2,6 @@
 
 #pragma once
 
+// HEADER_UNIT_UNSUPPORTED - Should not include this
+
 static_assert(false, "UObject/CoreOnlineFwd.h has moved to Online/CoreOnlineFwd.h in the CoreOnline module. Please update #includes and module dependencies.");
\ No newline at end of file
diff --git a/Engine/Source/Runtime/Engine/Classes/Materials/MaterialInstanceUpdateParameterSet.h b/Engine/Source/Runtime/Engine/Classes/Materials/MaterialInstanceUpdateParameterSet.h
--- a/Engine/Source/Runtime/Engine/Classes/Materials/MaterialInstanceUpdateParameterSet.h
+++ b/Engine/Source/Runtime/Engine/Classes/Materials/MaterialInstanceUpdateParameterSet.h
@@ -2,6 +2,8 @@
 
 #pragma once
 
+// HEADER_UNIT_SKIP - Not included directly
+
 #include "Materials/Material.h"
 #include "Materials/MaterialInstance.h"
 
diff --git a/Engine/Source/Runtime/Engine/Public/OverrideVoidReturnInvoker.h b/Engine/Source/Runtime/Engine/Public/OverrideVoidReturnInvoker.h
--- a/Engine/Source/Runtime/Engine/Public/OverrideVoidReturnInvoker.h
+++ b/Engine/Source/Runtime/Engine/Public/OverrideVoidReturnInvoker.h
@@ -2,6 +2,8 @@
 
 #pragma once
 
+// HEADER_UNIT_SKIP - Not included directly
+
 #include "CoreTypes.h"
 #include <type_traits>
 
diff --git a/Engine/Source/Runtime/Experimental/ChaosCore/Public/Chaos/ArrayFaceND.h b/Engine/Source/Runtime/Experimental/ChaosCore/Public/Chaos/ArrayFaceND.h
--- a/Engine/Source/Runtime/Experimental/ChaosCore/Public/Chaos/ArrayFaceND.h
+++ b/Engine/Source/Runtime/Experimental/ChaosCore/Public/Chaos/ArrayFaceND.h
@@ -1,6 +1,8 @@
 // Copyright Epic Games, Inc. All Rights Reserved.
 #pragma once
 
+// HEADER_UNIT_SKIP - Bad include. Some headers are in Chaos while this is in ChaosCore
+
 #include "Chaos/ArrayND.h"
 #include "Chaos/UniformGrid.h"
 #include "Chaos/Vector.h"
diff --git a/Engine/Source/Runtime/Experimental/ChaosCore/Public/Chaos/ArrayND.h b/Engine/Source/Runtime/Experimental/ChaosCore/Public/Chaos/ArrayND.h
--- a/Engine/Source/Runtime/Experimental/ChaosCore/Public/Chaos/ArrayND.h
+++ b/Engine/Source/Runtime/Experimental/ChaosCore/Public/Chaos/ArrayND.h
@@ -1,6 +1,8 @@
 // Copyright Epic Games, Inc. All Rights Reserved.
 #pragma once
 
+// HEADER_UNIT_SKIP - Bad include. Some headers are in Chaos while this is in ChaosCore
+
 #include "Chaos/Core.h"
 #include "Chaos/Array.h"
 #include "Chaos/UniformGrid.h"
diff --git a/Engine/Source/Runtime/GameplayDebugger/Public/GameplayDebuggerCompat.h b/Engine/Source/Runtime/GameplayDebugger/Public/GameplayDebuggerCompat.h
--- a/Engine/Source/Runtime/GameplayDebugger/Public/GameplayDebuggerCompat.h
+++ b/Engine/Source/Runtime/GameplayDebugger/Public/GameplayDebuggerCompat.h
@@ -2,6 +2,8 @@
 
 #pragma once
 
+// HEADER_UNIT_SKIP - Bad?
+
 #if WITH_EDITOR
 #include "Editor/EditorEngine.h"
 #include "LevelEditor.h"
diff --git a/Engine/Source/Runtime/GeometryCore/Public/Image/SpatialPhotoSet.h b/Engine/Source/Runtime/GeometryCore/Public/Image/SpatialPhotoSet.h
--- a/Engine/Source/Runtime/GeometryCore/Public/Image/SpatialPhotoSet.h
+++ b/Engine/Source/Runtime/GeometryCore/Public/Image/SpatialPhotoSet.h
@@ -2,6 +2,8 @@
 
 #pragma once
 
+// HEADER_UNIT_SKIP - Bad include. Some headers are in Engine
+
 #include "VectorTypes.h"
 #include "FrameTypes.h"
 #include "SceneView.h"
diff --git a/Engine/Source/Runtime/Landscape/Public/LandscapeRenderMobile.h b/Engine/Source/Runtime/Landscape/Public/LandscapeRenderMobile.h
--- a/Engine/Source/Runtime/Landscape/Public/LandscapeRenderMobile.h
+++ b/Engine/Source/Runtime/Landscape/Public/LandscapeRenderMobile.h
@@ -6,6 +6,8 @@ DEPRECATED LandscapeRenderMobile.h: This file will be removed
 #pragma once
 #pragma message("LandscapeRenderMobile.h is deprecated. Please remove all usage of LandscapeRenderMobile.h")
 
+// HEADER_UNIT_SKIP - Deprecated
+
 #include "CoreMinimal.h"
 #include "RenderResource.h"
 
diff --git a/Engine/Source/Runtime/MeshConversion/Public/StaticMeshLODResourcesAdapter.h b/Engine/Source/Runtime/MeshConversion/Public/StaticMeshLODResourcesAdapter.h
--- a/Engine/Source/Runtime/MeshConversion/Public/StaticMeshLODResourcesAdapter.h
+++ b/Engine/Source/Runtime/MeshConversion/Public/StaticMeshLODResourcesAdapter.h
@@ -2,6 +2,8 @@
 
 #pragma once
 
+// HEADER_UNIT_SKIP - Bad includes
+
 #include "CoreMinimal.h"
 #include "StaticMeshAttributes.h"
 #include "StaticMeshResources.h"
diff --git a/Engine/Source/Runtime/Online/BuildPatchServices/Private/Generation/BuildDataVerifier.h b/Engine/Source/Runtime/Online/BuildPatchServices/Private/Generation/BuildDataVerifier.h
--- a/Engine/Source/Runtime/Online/BuildPatchServices/Private/Generation/BuildDataVerifier.h
+++ b/Engine/Source/Runtime/Online/BuildPatchServices/Private/Generation/BuildDataVerifier.h
@@ -1,6 +1,8 @@
 // Copyright Epic Games, Inc. All Rights Reserved.
 #pragma once
 
+// HEADER_UNIT_SKIP - Internal
+
 #include "Core/BlockRange.h"
 #include "Core/BlockStructure.h"
 #include "Common/FileSystem.h"
diff --git a/Engine/Source/Runtime/OpenGLDrv/Public/OpenGLResources.h b/Engine/Source/Runtime/OpenGLDrv/Public/OpenGLResources.h
--- a/Engine/Source/Runtime/OpenGLDrv/Public/OpenGLResources.h
+++ b/Engine/Source/Runtime/OpenGLDrv/Public/OpenGLResources.h
@@ -6,6 +6,8 @@
 
 #pragma once
 
+// HEADER_UNIT_SKIP - Not included directly
+
 #include "CoreTypes.h"
 #include "Misc/AssertionMacros.h"
 #include "HAL/UnrealMemory.h"
diff --git a/Engine/Source/Runtime/OpenGLDrv/Public/OpenGLShaderResources.h b/Engine/Source/Runtime/OpenGLDrv/Public/OpenGLShaderResources.h
--- a/Engine/Source/Runtime/OpenGLDrv/Public/OpenGLShaderResources.h
+++ b/Engine/Source/Runtime/OpenGLDrv/Public/OpenGLShaderResources.h
@@ -6,6 +6,8 @@
 
 #pragma once
 
+// HEADER_UNIT_SKIP - Not included directly
+
 #include "HAL/UnrealMemory.h"
 #include "Templates/UnrealTemplate.h"
 #include "Containers/Array.h"
diff --git a/Engine/Source/Runtime/RHI/Public/DynamicRHI.h b/Engine/Source/Runtime/RHI/Public/DynamicRHI.h
--- a/Engine/Source/Runtime/RHI/Public/DynamicRHI.h
+++ b/Engine/Source/Runtime/RHI/Public/DynamicRHI.h
@@ -6,6 +6,8 @@ DynamicRHI.h: Dynamically bound Render Hardware Interface definitions.
 
 #pragma once
 
+// HEADER_UNIT_SKIP - Should always be included through RHI.h (this code needs cleanup.. this file is included in the middle of RHI.h)
+
 #include "CoreTypes.h"
 #include "RHIContext.h"
 #include "MultiGPU.h"
diff --git a/Engine/Source/Runtime/RHI/Public/RHIContext.h b/Engine/Source/Runtime/RHI/Public/RHIContext.h
--- a/Engine/Source/Runtime/RHI/Public/RHIContext.h
+++ b/Engine/Source/Runtime/RHI/Public/RHIContext.h
@@ -6,6 +6,8 @@
 
 #pragma once
 
+// HEADER_UNIT_SKIP - Should always be included after RHI.h (this code needs cleanup.. this file is included in the middle of RHI.h)
+
 #include "Misc/AssertionMacros.h"
 #include "Math/Color.h"
 #include "Math/IntPoint.h"
diff --git a/Engine/Source/Runtime/RHI/Public/RHIResources.h b/Engine/Source/Runtime/RHI/Public/RHIResources.h
--- a/Engine/Source/Runtime/RHI/Public/RHIResources.h
+++ b/Engine/Source/Runtime/RHI/Public/RHIResources.h
@@ -2,6 +2,9 @@
 
 #pragma once
 
+// HEADER_UNIT_SKIP - Should always be included through RHI.h (this code needs cleanup.. this file is included in the middle of RHI.h)
+
+
 #include <atomic>
 
 #include "CoreTypes.h"
diff --git a/Engine/Source/Runtime/RHI/Public/RHIStaticStates.h b/Engine/Source/Runtime/RHI/Public/RHIStaticStates.h
--- a/Engine/Source/Runtime/RHI/Public/RHIStaticStates.h
+++ b/Engine/Source/Runtime/RHI/Public/RHIStaticStates.h
@@ -6,6 +6,8 @@
 
 #pragma once
 
+// HEADER_UNIT_SKIP - This can't be built by itself since it includes RenderResource.h which is not part of this module's additional include paths
+
 #include "CoreMinimal.h"
 #include "Stats/Stats.h"
 #include "Async/TaskGraphInterfaces.h"
diff --git a/Engine/Source/Runtime/RHI/Public/RHIValidationCommon.h b/Engine/Source/Runtime/RHI/Public/RHIValidationCommon.h
--- a/Engine/Source/Runtime/RHI/Public/RHIValidationCommon.h
+++ b/Engine/Source/Runtime/RHI/Public/RHIValidationCommon.h
@@ -6,6 +6,8 @@
 
 #pragma once
 
+// HEADER_UNIT_SKIP - Should always be included through RHI.h (this code needs cleanup.. this file is included in the middle of RHI.h)
+
 #include "Experimental/ConcurrentLinearAllocator.h"
 
 #if ENABLE_RHI_VALIDATION
diff --git a/Engine/Source/Runtime/RHI/Public/RHIValidationContext.h b/Engine/Source/Runtime/RHI/Public/RHIValidationContext.h
--- a/Engine/Source/Runtime/RHI/Public/RHIValidationContext.h
+++ b/Engine/Source/Runtime/RHI/Public/RHIValidationContext.h
@@ -6,6 +6,8 @@
 
 #pragma once
 
+// HEADER_UNIT_SKIP - Should always be included through RHIValidation.h
+
 #include "RHIValidationCommon.h"
 #include "RHIValidationUtils.h"
 
diff --git a/Engine/Source/Runtime/RHI/Public/RHIValidationUtils.h b/Engine/Source/Runtime/RHI/Public/RHIValidationUtils.h
--- a/Engine/Source/Runtime/RHI/Public/RHIValidationUtils.h
+++ b/Engine/Source/Runtime/RHI/Public/RHIValidationUtils.h
@@ -2,6 +2,8 @@
 
 #pragma once
 
+// HEADER_UNIT_SKIP - Included through others
+
 #include "Misc/Build.h"
 
 #if ENABLE_RHI_VALIDATION
diff --git a/Engine/Source/Runtime/Renderer/Private/GlobalDistanceFieldHeightfields.h b/Engine/Source/Runtime/Renderer/Private/GlobalDistanceFieldHeightfields.h
--- a/Engine/Source/Runtime/Renderer/Private/GlobalDistanceFieldHeightfields.h
+++ b/Engine/Source/Runtime/Renderer/Private/GlobalDistanceFieldHeightfields.h
@@ -6,6 +6,8 @@
 
 #pragma once
 
+// HEADER_UNIT_SKIP - Internal
+
 class FMarkHeightfieldPagesCS : public FGlobalShader
 {
 	DECLARE_GLOBAL_SHADER(FMarkHeightfieldPagesCS);
diff --git a/Engine/Source/Runtime/Renderer/Private/ShadingEnergyConservationData.h b/Engine/Source/Runtime/Renderer/Private/ShadingEnergyConservationData.h
--- a/Engine/Source/Runtime/Renderer/Private/ShadingEnergyConservationData.h
+++ b/Engine/Source/Runtime/Renderer/Private/ShadingEnergyConservationData.h
@@ -2,6 +2,8 @@
 
 #pragma once
 
+// HEADER_UNIT_SKIP - Not included directly
+
 // These values were generated using FBuildShadingEnergyConservationTableCS
 
 // How to update those tables?
diff --git a/Engine/Source/Runtime/SignalProcessing/Public/DSP/PassiveFilter.h b/Engine/Source/Runtime/SignalProcessing/Public/DSP/PassiveFilter.h
--- a/Engine/Source/Runtime/SignalProcessing/Public/DSP/PassiveFilter.h
+++ b/Engine/Source/Runtime/SignalProcessing/Public/DSP/PassiveFilter.h
@@ -1,7 +1,9 @@
 // Copyright Epic Games, Inc. All Rights Reserved.
 
 #pragma once
- 
+
+// HEADER_UNIT_SKIP - Bad includes! Audio is not available in this module
+
 #include "CoreMinimal.h"
 #include "DSP/AudioFFT.h"
 #include "DSP/FFTAlgorithm.h"
diff --git a/Engine/Source/Runtime/Slate/Public/Framework/Text/TextRange.h b/Engine/Source/Runtime/Slate/Public/Framework/Text/TextRange.h
--- a/Engine/Source/Runtime/Slate/Public/Framework/Text/TextRange.h
+++ b/Engine/Source/Runtime/Slate/Public/Framework/Text/TextRange.h
@@ -1,6 +1,8 @@
 // Copyright Epic Games, Inc. All Rights Reserved.
 #pragma once
 
+// HEADER_UNIT_SKIP - Deprecated
+
 #include "Containers/UnrealString.h"
 
 UE_DEPRECATED(4.23, "Including TextRange.h has been deprecated. Please include UnrealString.h in the Core module instead.")
diff --git a/Engine/Source/Runtime/Slate/Public/SlateMaterialBrush.h b/Engine/Source/Runtime/Slate/Public/SlateMaterialBrush.h
--- a/Engine/Source/Runtime/Slate/Public/SlateMaterialBrush.h
+++ b/Engine/Source/Runtime/Slate/Public/SlateMaterialBrush.h
@@ -2,6 +2,8 @@
 
 #pragma once
 
+// HEADER_UNIT_SKIP - Not included directly
+
 #include "CoreMinimal.h"
 #include "Layout/Margin.h"
 #include "Styling/SlateBrush.h"
diff --git a/Engine/Source/Runtime/TraceLog/Private/Trace/WriteBufferRedirect.h b/Engine/Source/Runtime/TraceLog/Private/Trace/WriteBufferRedirect.h
--- a/Engine/Source/Runtime/TraceLog/Private/Trace/WriteBufferRedirect.h
+++ b/Engine/Source/Runtime/TraceLog/Private/Trace/WriteBufferRedirect.h
@@ -2,6 +2,8 @@
 
 #pragma once
 
+// HEADER_UNIT_SKIP - Internal
+
 #include "Trace/Detail/Writer.inl"
 
 namespace UE {
diff --git a/Engine/Source/Runtime/TraceLog/Public/Trace/Detail/Protocols/Protocol0.h b/Engine/Source/Runtime/TraceLog/Public/Trace/Detail/Protocols/Protocol0.h
--- a/Engine/Source/Runtime/TraceLog/Public/Trace/Detail/Protocols/Protocol0.h
+++ b/Engine/Source/Runtime/TraceLog/Public/Trace/Detail/Protocols/Protocol0.h
@@ -2,6 +2,8 @@
 
 #pragma once
 
+// HEADER_UNIT_SKIP - Not included directly
+
 namespace UE {
 namespace Trace {
 
diff --git a/Engine/Source/Runtime/TraceLog/Public/Trace/Detail/Protocols/Protocol1.h b/Engine/Source/Runtime/TraceLog/Public/Trace/Detail/Protocols/Protocol1.h
--- a/Engine/Source/Runtime/TraceLog/Public/Trace/Detail/Protocols/Protocol1.h
+++ b/Engine/Source/Runtime/TraceLog/Public/Trace/Detail/Protocols/Protocol1.h
@@ -2,6 +2,8 @@
 
 #pragma once
 
+// HEADER_UNIT_SKIP - Not included directly
+
 namespace UE {
 namespace Trace {
 
diff --git a/Engine/Source/Runtime/TraceLog/Public/Trace/Detail/Protocols/Protocol2.h b/Engine/Source/Runtime/TraceLog/Public/Trace/Detail/Protocols/Protocol2.h
--- a/Engine/Source/Runtime/TraceLog/Public/Trace/Detail/Protocols/Protocol2.h
+++ b/Engine/Source/Runtime/TraceLog/Public/Trace/Detail/Protocols/Protocol2.h
@@ -2,6 +2,8 @@
 
 #pragma once
 
+// HEADER_UNIT_SKIP - Not included directly
+
 namespace UE {
 namespace Trace {
 
diff --git a/Engine/Source/Runtime/TraceLog/Public/Trace/Detail/Protocols/Protocol3.h b/Engine/Source/Runtime/TraceLog/Public/Trace/Detail/Protocols/Protocol3.h
--- a/Engine/Source/Runtime/TraceLog/Public/Trace/Detail/Protocols/Protocol3.h
+++ b/Engine/Source/Runtime/TraceLog/Public/Trace/Detail/Protocols/Protocol3.h
@@ -2,6 +2,8 @@
 
 #pragma once
 
+// HEADER_UNIT_SKIP - Not included directly
+
 namespace UE {
 namespace Trace {
 
diff --git a/Engine/Source/Runtime/TraceLog/Public/Trace/Detail/Protocols/Protocol4.h b/Engine/Source/Runtime/TraceLog/Public/Trace/Detail/Protocols/Protocol4.h
--- a/Engine/Source/Runtime/TraceLog/Public/Trace/Detail/Protocols/Protocol4.h
+++ b/Engine/Source/Runtime/TraceLog/Public/Trace/Detail/Protocols/Protocol4.h
@@ -2,6 +2,8 @@
 
 #pragma once
 
+// HEADER_UNIT_SKIP - Not included directly
+
 namespace UE {
 namespace Trace {
 
diff --git a/Engine/Source/Runtime/TraceLog/Public/Trace/Detail/Protocols/Protocol5.h b/Engine/Source/Runtime/TraceLog/Public/Trace/Detail/Protocols/Protocol5.h
--- a/Engine/Source/Runtime/TraceLog/Public/Trace/Detail/Protocols/Protocol5.h
+++ b/Engine/Source/Runtime/TraceLog/Public/Trace/Detail/Protocols/Protocol5.h
@@ -2,6 +2,8 @@
 
 #pragma once
 
+// HEADER_UNIT_SKIP - Not included directly
+
 namespace UE {
 namespace Trace {
 
diff --git a/Engine/Source/Runtime/TraceLog/Public/Trace/Detail/Protocols/Protocol6.h b/Engine/Source/Runtime/TraceLog/Public/Trace/Detail/Protocols/Protocol6.h
--- a/Engine/Source/Runtime/TraceLog/Public/Trace/Detail/Protocols/Protocol6.h
+++ b/Engine/Source/Runtime/TraceLog/Public/Trace/Detail/Protocols/Protocol6.h
@@ -1,6 +1,8 @@
-// Copyright Epic Games, Inc. All Rights Reserved.
+﻿// Copyright Epic Games, Inc. All Rights Reserved.
 #pragma once
 
+// HEADER_UNIT_SKIP - Not included directly
+
 #include "Protocol4.h"
 #include "Protocol5.h"
 
diff --git a/Engine/Source/Runtime/TraceLog/standalone_epilogue.h b/Engine/Source/Runtime/TraceLog/standalone_epilogue.h
--- a/Engine/Source/Runtime/TraceLog/standalone_epilogue.h
+++ b/Engine/Source/Runtime/TraceLog/standalone_epilogue.h
@@ -1,5 +1,7 @@
 // Copyright Epic Games, Inc. All Rights Reserved.
 
+// HEADER_UNIT_SKIP - Not included directly
+
 // {{{1 amalgamation-tail //////////////////////////////////////////////////////
 
 #if PLATFORM_WINDOWS
diff --git a/Engine/Source/Runtime/TraceLog/standalone_prologue.h b/Engine/Source/Runtime/TraceLog/standalone_prologue.h
--- a/Engine/Source/Runtime/TraceLog/standalone_prologue.h
+++ b/Engine/Source/Runtime/TraceLog/standalone_prologue.h
@@ -2,6 +2,8 @@
 
 #pragma once
 
+// HEADER_UNIT_SKIP - Not included directly
+
 #if !defined(TRACE_UE_COMPAT_LAYER)
 #	if defined(__UNREAL__) && __UNREAL__
 #		define TRACE_UE_COMPAT_LAYER	0
diff --git a/Engine/Source/Runtime/VectorVM/Public/VectorVMDataObject.h b/Engine/Source/Runtime/VectorVM/Public/VectorVMDataObject.h
--- a/Engine/Source/Runtime/VectorVM/Public/VectorVMDataObject.h
+++ b/Engine/Source/Runtime/VectorVM/Public/VectorVMDataObject.h
@@ -1,6 +1,9 @@
 // Copyright Epic Games, Inc. All Rights Reserved.
 
 #pragma once	
+
+// HEADER_UNIT_SKIP - Bad header
+
 #include "VectorVMDataObject.generated.h"
 
 //TODO Data Objects are to be totally re-worked
\ No newline at end of file
diff --git a/Engine/Source/Runtime/VectorVM/Public/VectorVMExperimental.h b/Engine/Source/Runtime/VectorVM/Public/VectorVMExperimental.h
--- a/Engine/Source/Runtime/VectorVM/Public/VectorVMExperimental.h
+++ b/Engine/Source/Runtime/VectorVM/Public/VectorVMExperimental.h
@@ -1,6 +1,8 @@
 // Copyright Epic Games, Inc. All Rights Reserved.
 #pragma once
 
+// HEADER_UNIT_SKIP - Not included directly
+
 #define VECTORVM_SUPPORTS_AVX 0
 #define VECTORVM_SUPPORTS_COMPUTED_GOTO 0
 
