Deprecated TIsVoidType.

diff --git a/Engine/Plugins/Runtime/NetworkPrediction/Source/NetworkPrediction/Public/NetworkPredictionDriver.h b/Engine/Plugins/Runtime/NetworkPrediction/Source/NetworkPrediction/Public/NetworkPredictionDriver.h
--- a/Engine/Plugins/Runtime/NetworkPrediction/Source/NetworkPrediction/Public/NetworkPredictionDriver.h
+++ b/Engine/Plugins/Runtime/NetworkPrediction/Source/NetworkPrediction/Public/NetworkPredictionDriver.h
@@ -70,11 +70,11 @@ struct FNetworkPredictionDriverBase
 	using AuxType = typename StateTypes::AuxType;
 	using PhysicsState = typename ModelDef::PhysicsState;
 
-	static constexpr bool HasNpState() { return !TIsVoidType<InputType>::Value || !TIsVoidType<SyncType>::Value || !TIsVoidType<AuxType>::Value; }
-	static constexpr bool HasDriver() { return !TIsVoidType<DriverType>::Value; }
-	static constexpr bool HasSimulation() { return !TIsVoidType<Simulation>::Value; }
-	static constexpr bool HasInput() { return !TIsVoidType<InputType>::Value; }
-	static constexpr bool HasPhysics() { return !TIsVoidType<PhysicsState>::Value; }
+	static constexpr bool HasNpState() { return !std::is_void_v<InputType> || !std::is_void_v<SyncType> || !std::is_void_v<AuxType>; }
+	static constexpr bool HasDriver() { return !std::is_void_v<DriverType>; }
+	static constexpr bool HasSimulation() { return !std::is_void_v<Simulation>; }
+	static constexpr bool HasInput() { return !std::is_void_v<InputType>; }
+	static constexpr bool HasPhysics() { return !std::is_void_v<PhysicsState>; }
 
 	// Defines what the ModelDef can do. This is a compile time thing only.
 	static constexpr FNetworkPredictionModelDefCapabilities GetCapabilities()
diff --git a/Engine/Plugins/Runtime/NetworkPrediction/Source/NetworkPrediction/Public/NetworkPredictionTrace.h b/Engine/Plugins/Runtime/NetworkPrediction/Source/NetworkPrediction/Public/NetworkPredictionTrace.h
--- a/Engine/Plugins/Runtime/NetworkPrediction/Source/NetworkPrediction/Public/NetworkPredictionTrace.h
+++ b/Engine/Plugins/Runtime/NetworkPrediction/Source/NetworkPrediction/Public/NetworkPredictionTrace.h
@@ -180,7 +180,7 @@ public:
 	static void TraceUserState(const StateType* State, ETraceUserState StateTypeEnum)
 	{
 #if UE_NP_TRACE_USER_STATES_ENABLED
-		if (TIsVoidType<StateType>::Value)
+		if (std::is_void_v<StateType>)
 		{
 			return;
 		}
diff --git a/Engine/Plugins/Runtime/NetworkPrediction/Source/NetworkPrediction/Public/Services/NetworkPredictionService_Input.inl b/Engine/Plugins/Runtime/NetworkPrediction/Source/NetworkPrediction/Public/Services/NetworkPredictionService_Input.inl
--- a/Engine/Plugins/Runtime/NetworkPrediction/Source/NetworkPrediction/Public/Services/NetworkPredictionService_Input.inl
+++ b/Engine/Plugins/Runtime/NetworkPrediction/Source/NetworkPrediction/Public/Services/NetworkPredictionService_Input.inl
@@ -30,7 +30,7 @@ public:
 
 	void RegisterInstance(FNetworkPredictionID ID)
 	{
-		npCheckfSlow(!TIsVoidType<InputType>::Value, TEXT("ModelDef %s with null InputCmd type was registered for local input service."), ModelDef::GetName());
+		npCheckfSlow(!std::is_void_v<InputType>, TEXT("ModelDef %s with null InputCmd type was registered for local input service."), ModelDef::GetName());
 
 		TInstanceData<ModelDef>* InstanceData = DataStore->Instances.Find(ID);
 		npCheckSlow(InstanceData);
diff --git a/Engine/Source/Runtime/Core/Public/Containers/BitArray.h b/Engine/Source/Runtime/Core/Public/Containers/BitArray.h
--- a/Engine/Source/Runtime/Core/Public/Containers/BitArray.h
+++ b/Engine/Source/Runtime/Core/Public/Containers/BitArray.h
@@ -17,7 +17,6 @@
 #include "Templates/Invoke.h"
 #include "Templates/UnrealTemplate.h"
 #include "Templates/UnrealTypeTraits.h"
-#include "Traits/IsVoidType.h"
 
 class FPointerTableBase;
 class FSHA1;
@@ -2002,7 +2001,7 @@ using TConstDualEitherSetBitIterator = TConstDualSetBitIterator<Allocator, Other
 template <typename Allocator, typename InDerivedType>
 class TScriptBitArray
 {
-	using DerivedType = typename TChooseClass<TIsVoidType<InDerivedType>::Value, TScriptBitArray, InDerivedType>::Result;
+	using DerivedType = typename TChooseClass<std::is_void_v<InDerivedType>, TScriptBitArray, InDerivedType>::Result;
 
 public:
 	/**
diff --git a/Engine/Source/Runtime/Core/Public/Containers/Map.h b/Engine/Source/Runtime/Core/Public/Containers/Map.h
--- a/Engine/Source/Runtime/Core/Public/Containers/Map.h
+++ b/Engine/Source/Runtime/Core/Public/Containers/Map.h
@@ -1627,7 +1627,7 @@ struct FScriptMapLayout
 template <typename AllocatorType, typename InDerivedType>
 class TScriptMap
 {
-	using DerivedType = typename TChooseClass<TIsVoidType<InDerivedType>::Value, TScriptMap, InDerivedType>::Result;
+	using DerivedType = typename TChooseClass<std::is_void_v<InDerivedType>, TScriptMap, InDerivedType>::Result;
 
 public:
 	static FScriptMapLayout GetScriptLayout(int32 KeySize, int32 KeyAlignment, int32 ValueSize, int32 ValueAlignment)
diff --git a/Engine/Source/Runtime/Core/Public/Containers/Set.h b/Engine/Source/Runtime/Core/Public/Containers/Set.h
--- a/Engine/Source/Runtime/Core/Public/Containers/Set.h
+++ b/Engine/Source/Runtime/Core/Public/Containers/Set.h
@@ -1782,7 +1782,7 @@ struct FScriptSetLayout
 template <typename Allocator, typename InDerivedType>
 class TScriptSet
 {
-	using DerivedType = typename TChooseClass<TIsVoidType<InDerivedType>::Value, TScriptSet, InDerivedType>::Result;
+	using DerivedType = typename TChooseClass<std::is_void_v<InDerivedType>, TScriptSet, InDerivedType>::Result;
 
 public:
 	static FScriptSetLayout GetScriptLayout(int32 ElementSize, int32 ElementAlignment)
diff --git a/Engine/Source/Runtime/Core/Public/Containers/SparseArray.h b/Engine/Source/Runtime/Core/Public/Containers/SparseArray.h
--- a/Engine/Source/Runtime/Core/Public/Containers/SparseArray.h
+++ b/Engine/Source/Runtime/Core/Public/Containers/SparseArray.h
@@ -1223,7 +1223,7 @@ struct FScriptSparseArrayLayout
 template <typename AllocatorType, typename InDerivedType>
 class TScriptSparseArray
 {
-	using DerivedType = typename TChooseClass<TIsVoidType<InDerivedType>::Value, TScriptSparseArray, InDerivedType>::Result;
+	using DerivedType = typename TChooseClass<std::is_void_v<InDerivedType>, TScriptSparseArray, InDerivedType>::Result;
 
 public:
 	static FScriptSparseArrayLayout GetScriptLayout(int32 ElementSize, int32 ElementAlignment)
diff --git a/Engine/Source/Runtime/Core/Public/Memory/MemoryView.h b/Engine/Source/Runtime/Core/Public/Memory/MemoryView.h
--- a/Engine/Source/Runtime/Core/Public/Memory/MemoryView.h
+++ b/Engine/Source/Runtime/Core/Public/Memory/MemoryView.h
@@ -15,7 +15,6 @@
 #include "Templates/PointerIsConvertibleFromTo.h"
 #include "Templates/UnrealTemplate.h"
 #include "Traits/IsContiguousContainer.h"
-#include "Traits/IsVoidType.h"
 
 #include <initializer_list>
 
@@ -29,7 +28,7 @@
 template <typename DataType>
 class TMemoryView
 {
-	static_assert(TIsVoidType<DataType>::Value, "DataType must be cv-qualified void");
+	static_assert(std::is_void_v<DataType>, "DataType must be cv-qualified void");
 
 	using ByteType = typename TChooseClass<TIsConst<DataType>::Value, const uint8, uint8>::Result;
 
diff --git a/Engine/Source/Runtime/Core/Public/Templates/UnrealTypeTraits.h b/Engine/Source/Runtime/Core/Public/Templates/UnrealTypeTraits.h
--- a/Engine/Source/Runtime/Core/Public/Templates/UnrealTypeTraits.h
+++ b/Engine/Source/Runtime/Core/Public/Templates/UnrealTypeTraits.h
@@ -17,7 +17,6 @@
 #include "Templates/IsArithmetic.h"
 #include "Templates/IsEnum.h"
 #include "Templates/RemoveCV.h"
-#include "Traits/IsVoidType.h"
 #include "Templates/Models.h"
 
 #include "Templates/IsPODType.h"
@@ -161,7 +160,7 @@ template<typename T> struct TIsRValueReferenceType<T&&> { enum { Value = true  }
 template<typename T> 
 struct TIsFundamentalType 
 { 
-	enum { Value = TOr<TIsArithmetic<T>, TIsVoidType<T>>::Value };
+	enum { Value = TIsArithmetic<T>::Value || std::is_void_v<T> };
 };
 
 /**
diff --git a/Engine/Source/Runtime/Core/Public/Traits/IsVoidType.h b/Engine/Source/Runtime/Core/Public/Traits/IsVoidType.h
--- a/Engine/Source/Runtime/Core/Public/Traits/IsVoidType.h
+++ b/Engine/Source/Runtime/Core/Public/Traits/IsVoidType.h
@@ -2,11 +2,13 @@
 
 #pragma once
 
+template<typename T> struct TIsVoidTypeBase { enum { Value = false }; };
+template<> struct TIsVoidTypeBase<void> { enum { Value = true }; };
+template<> struct TIsVoidTypeBase<void const> { enum { Value = true }; };
+template<> struct TIsVoidTypeBase<void volatile> { enum { Value = true }; };
+template<> struct TIsVoidTypeBase<void const volatile> { enum { Value = true }; };
+
 /**
  * TIsVoidType
  */
-template<typename T> struct TIsVoidType { enum { Value = false }; };
-template<> struct TIsVoidType<void> { enum { Value = true }; };
-template<> struct TIsVoidType<void const> { enum { Value = true }; };
-template<> struct TIsVoidType<void volatile> { enum { Value = true }; };
-template<> struct TIsVoidType<void const volatile> { enum { Value = true }; };
+template<typename T> struct UE_DEPRECATED(5.2, "TIsVoidType has been deprecated, please use std::is_void instead.") TIsVoidType : TIsVoidTypeBase<T> {};
diff --git a/Engine/Source/Runtime/Engine/Public/CustomBoneIndexArray.h b/Engine/Source/Runtime/Engine/Public/CustomBoneIndexArray.h
--- a/Engine/Source/Runtime/Engine/Public/CustomBoneIndexArray.h
+++ b/Engine/Source/Runtime/Engine/Public/CustomBoneIndexArray.h
@@ -29,10 +29,9 @@ struct TIsCompatibleRangeIndexType
 {
 	enum
 	{
-		Value = TOr<
-			TIsSame<IndexType, typename TRangeIndexType<RangeType>::Type>,
-			TIsVoidType<typename TRangeIndexType<RangeType>::Type>
-		>::Value
+		Value =
+			std::is_same_v<IndexType, typename TRangeIndexType<RangeType>::Type> ||
+			std::is_void_v<typename TRangeIndexType<RangeType>::Type>
 	};
 };
 
