Network Prediction plugin: fixed issue where NP-simulated actors could be invisible on other clients when sim proxy interpolation mode in use. This was due to a replication race condition on the actor's bHidden property.

diff --git a/Engine/Plugins/Runtime/NetworkPrediction/Source/NetworkPrediction/Public/NetworkPredictionDriver.h b/Engine/Plugins/Runtime/NetworkPrediction/Source/NetworkPrediction/Public/NetworkPredictionDriver.h
--- a/Engine/Plugins/Runtime/NetworkPrediction/Source/NetworkPrediction/Public/NetworkPredictionDriver.h
+++ b/Engine/Plugins/Runtime/NetworkPrediction/Source/NetworkPrediction/Public/NetworkPredictionDriver.h
@@ -521,12 +521,12 @@ struct FNetworkPredictionDriverBase
 
 	static void CallSetHiddenForInterpolationFallback(AActor* Driver, bool bHide)
 	{
-		Driver->SetHidden(bHide);
+		Driver->SetActorHiddenInGame(bHide);
 	}
 
 	static void CallSetHiddenForInterpolationFallback(UActorComponent* Driver, bool bHide)
 	{
-		Driver->GetOwner()->SetHidden(bHide);
+		Driver->GetOwner()->SetActorHiddenInGame(bHide);
 	}	
 
 	// -----------------------------------------------------------------------------------------------------------------------------------
