Add a comment for new FNetworkPredictionProxy parameter struct.

diff --git a/Engine/Plugins/Runtime/NetworkPrediction/Source/NetworkPrediction/Public/NetworkPredictionProxy.h b/Engine/Plugins/Runtime/NetworkPrediction/Source/NetworkPrediction/Public/NetworkPredictionProxy.h
--- a/Engine/Plugins/Runtime/NetworkPrediction/Source/NetworkPrediction/Public/NetworkPredictionProxy.h
+++ b/Engine/Plugins/Runtime/NetworkPrediction/Source/NetworkPrediction/Public/NetworkPredictionProxy.h
@@ -37,6 +37,7 @@ struct FNetworkPredictionProxy
 {
 	GENERATED_BODY();
 
+	// Parameter struct to be used with the Init function overload below, when a UWorld isn't available, such as in unit tests.
 	template<typename ModelDef>
 	struct FInitParams
 	{
