Fix behavior, and related ensure, of simulated proxies controlled by the network prediction plugin when there's no autonomous proxy, such as when a client unpossesses a pawn and doesn't immediately possess another one.
Interpolation now uses the maximum of the LatestRecvFrame between autonomous and simulated proxies, instead of just autonomous proxies (if it's ever been set). Allows interpolation to continue while there's no local autonomous proxy.
Enabled the possession test by default now that it passes.

diff --git a/Engine/Plugins/Runtime/NetworkPrediction/Source/NetworkPrediction/Private/NetworkPredictionWorldManager.cpp b/Engine/Plugins/Runtime/NetworkPrediction/Source/NetworkPrediction/Private/NetworkPredictionWorldManager.cpp
--- a/Engine/Plugins/Runtime/NetworkPrediction/Source/NetworkPrediction/Private/NetworkPredictionWorldManager.cpp
+++ b/Engine/Plugins/Runtime/NetworkPrediction/Source/NetworkPrediction/Private/NetworkPredictionWorldManager.cpp
@@ -403,7 +403,7 @@ void UNetworkPredictionWorldManager::BeginNewSimulationFrame_Internal(float Delt
 	// -------------------------------------------------------------------------
 	if (Services.FixedInterpolate.Array.Num() > 0)
 	{
-		const int32 LatestRecvFrame = FixedTickState.Interpolation.LatestRecvFrameAP != INDEX_NONE ? FixedTickState.Interpolation.LatestRecvFrameAP : FixedTickState.Interpolation.LatestRecvFrameSP;
+		const int32 LatestRecvFrame = FMath::Max(FixedTickState.Interpolation.LatestRecvFrameAP, FixedTickState.Interpolation.LatestRecvFrameSP);
 		if (LatestRecvFrame != INDEX_NONE)
 		{
 			// We want 100ms of buffered time. As long a actors replicate at >= 10hz, this is should be good
diff --git a/Engine/Source/Programs/NetworkPredictionTests/Private/WorldManagerTests_Interpolation.cpp b/Engine/Source/Programs/NetworkPredictionTests/Private/WorldManagerTests_Interpolation.cpp
--- a/Engine/Source/Programs/NetworkPredictionTests/Private/WorldManagerTests_Interpolation.cpp
+++ b/Engine/Source/Programs/NetworkPredictionTests/Private/WorldManagerTests_Interpolation.cpp
@@ -49,7 +49,7 @@ private:
 };
 
 // This case includes a repro for UE-170936, disabled until it's fixed.
-TEST_CASE("NetworkPrediction interpolation mode, one simulated proxy", "[.][possession]")
+TEST_CASE("NetworkPrediction interpolation mode, one simulated proxy", "[possession]")
 {
 	UNetworkPredictionSettingsObject* Settings = NewObject<UNetworkPredictionSettingsObject>();
 	Settings->Settings.SimulatedProxyNetworkLOD = ENetworkLOD::Interpolated;
