Add some high-level Insights CPU scopes in the network prediction plugin.

diff --git a/Engine/Plugins/Runtime/NetworkPrediction/Source/NetworkPrediction/Private/NetworkPredictionWorldManager.cpp b/Engine/Plugins/Runtime/NetworkPrediction/Source/NetworkPrediction/Private/NetworkPredictionWorldManager.cpp
--- a/Engine/Plugins/Runtime/NetworkPrediction/Source/NetworkPrediction/Private/NetworkPredictionWorldManager.cpp
+++ b/Engine/Plugins/Runtime/NetworkPrediction/Source/NetworkPrediction/Private/NetworkPredictionWorldManager.cpp
@@ -102,8 +102,6 @@ void UNetworkPredictionWorldManager::OnWorldPreTick_Internal(float InDeltaSecond
 
 void UNetworkPredictionWorldManager::ReconcileSimulationsPostNetworkUpdate()
 {
-	QUICK_SCOPE_CYCLE_COUNTER(STAT_NP_RECONCILE);
-
 	UWorld* World = GetWorld();
 	if (World->GetNetMode() != NM_Client)
 	{
@@ -115,6 +113,9 @@ void UNetworkPredictionWorldManager::ReconcileSimulationsPostNetworkUpdate()
 
 void UNetworkPredictionWorldManager::ReconcileSimulationsPostNetworkUpdate_Internal()
 {
+	QUICK_SCOPE_CYCLE_COUNTER(STAT_NP_RECONCILE);
+	TRACE_CPUPROFILER_EVENT_SCOPE(NetworkPrediction::Reconcile);
+
 	ActiveInstance = this;
 	bLockServices = true;
 
@@ -166,7 +167,8 @@ void UNetworkPredictionWorldManager::ReconcileSimulationsPostNetworkUpdate_Inter
 
 	if (RollbackFrame != INDEX_NONE)
 	{
-		QUICK_SCOPE_CYCLE_COUNTER(STAT_NP_ROLLBACK);		
+		QUICK_SCOPE_CYCLE_COUNTER(STAT_NP_ROLLBACK);
+		TRACE_CPUPROFILER_EVENT_SCOPE(NetworkPrediction::Rollback);
 
 		if (RollbackFrame < FixedTickState.PendingFrame)
 		{
@@ -283,19 +285,21 @@ void UNetworkPredictionWorldManager::BeginNewSimulationFrame(UWorld* InWorld, EL
 
 void UNetworkPredictionWorldManager::BeginNewSimulationFrame_Internal(float DeltaTimeSeconds)
 {
+	QUICK_SCOPE_CYCLE_COUNTER(STAT_NP_TICK);
+	TRACE_CPUPROFILER_EVENT_SCOPE(NetworkPrediction::Tick);
+	
 	ActiveInstance = this;
 	bLockServices = true;
 
 	const float fEngineFrameDeltaTimeMS = DeltaTimeSeconds * 1000.f;
 
-	QUICK_SCOPE_CYCLE_COUNTER(STAT_NP_TICK);
-
 	// -------------------------------------------------------------------------
 	//	Fixed Tick
 	// -------------------------------------------------------------------------
 	if (Physics.bUsingPhysics || Services.FixedTick.Array.Num() > 0)
 	{
 		QUICK_SCOPE_CYCLE_COUNTER(STAT_NP_TICK_FIXED);
+		TRACE_CPUPROFILER_EVENT_SCOPE(NetworkPrediction::FixedTick);
 
 		FixedTickState.UnspentTimeMS += fEngineFrameDeltaTimeMS;
 		
@@ -352,6 +356,8 @@ void UNetworkPredictionWorldManager::BeginNewSimulationFrame_Internal(float Delt
 	//	Local Independent Tick
 	// -------------------------------------------------------------------------
 	{
+		TRACE_CPUPROFILER_EVENT_SCOPE(NetworkPrediction::IndependentTick);
+
 		// Update VariableTickState	
 		constexpr int32 MinStepMS = 1;
 		constexpr int32 MaxStepMS = 100;
@@ -401,118 +407,127 @@ void UNetworkPredictionWorldManager::BeginNewSimulationFrame_Internal(float Delt
 	// -------------------------------------------------------------------------
 	// Interpolation
 	// -------------------------------------------------------------------------
-	if (Services.FixedInterpolate.Array.Num() > 0)
 	{
-		const int32 LatestRecvFrame = FMath::Max(FixedTickState.Interpolation.LatestRecvFrameAP, FixedTickState.Interpolation.LatestRecvFrameSP);
-		if (LatestRecvFrame != INDEX_NONE)
-		{
-			// We want 100ms of buffered time. As long a actors replicate at >= 10hz, this is should be good
-			// Its better to keep this simple with a single time rather than trying to coordinate lowest amount of buffered time
-			// between all the registered instances in the different ModelDefs
-			const int32 DesiredBufferedMS = Settings.FixedTickInterpolationBufferedMS;
+		TRACE_CPUPROFILER_EVENT_SCOPE(NetworkPrediction::Interpolation);
 
-			float InterpolateRate = 1.f;
-			if (FixedTickState.Interpolation.ToFrame == INDEX_NONE)
+		if (Services.FixedInterpolate.Array.Num() > 0)
+		{
+			const int32 LatestRecvFrame = FMath::Max(FixedTickState.Interpolation.LatestRecvFrameAP, FixedTickState.Interpolation.LatestRecvFrameSP);
+			if (LatestRecvFrame != INDEX_NONE)
 			{
-				const int32 NumBufferedFrames = LatestRecvFrame;
-				const int32 BufferedMS = NumBufferedFrames * FixedTickState.FixedStepMS;
-
-				//UE_LOG(LogTemp, Warning, TEXT("BufferedMS: %d Frames: %d (No ToFrame)"), BufferedMS, NumBufferedFrames);
+				// We want 100ms of buffered time. As long a actors replicate at >= 10hz, this is should be good
+				// Its better to keep this simple with a single time rather than trying to coordinate lowest amount of buffered time
+				// between all the registered instances in the different ModelDefs
+				const int32 DesiredBufferedMS = Settings.FixedTickInterpolationBufferedMS;
 
-				if (BufferedMS < DesiredBufferedMS)
+				float InterpolateRate = 1.f;
+				if (FixedTickState.Interpolation.ToFrame == INDEX_NONE)
 				{
-					// Not enough time buffered yet to start interpolating
-					InterpolateRate = 0.f;
+					const int32 NumBufferedFrames = LatestRecvFrame;
+					const int32 BufferedMS = NumBufferedFrames * FixedTickState.FixedStepMS;
+
+					//UE_LOG(LogTemp, Warning, TEXT("BufferedMS: %d Frames: %d (No ToFrame)"), BufferedMS, NumBufferedFrames);
+
+					if (BufferedMS < DesiredBufferedMS)
+					{
+						// Not enough time buffered yet to start interpolating
+						InterpolateRate = 0.f;
+					}
+					else
+					{
+						// Begin interpolation
+						const int32 DesiredNumBufferedFrames = (DesiredBufferedMS / FixedTickState.FixedStepMS);
+						FixedTickState.Interpolation.ToFrame = LatestRecvFrame - DesiredNumBufferedFrames;
+
+						FixedTickState.Interpolation.PCT = 0.f;
+						FixedTickState.Interpolation.AccumulatedTimeMS = 0.f;
+
+						// We need to force a reconcile here since we supress the call until interpolation starts
+						for (TUniquePtr<IFixedInterpolateService>& Ptr : Services.FixedInterpolate.Array)
+						{
+							Ptr->Reconcile(&FixedTickState);
+						}
+					}
 				}
 				else
 				{
-					// Begin interpolation
-					const int32 DesiredNumBufferedFrames = (DesiredBufferedMS / FixedTickState.FixedStepMS);
-					FixedTickState.Interpolation.ToFrame = LatestRecvFrame - DesiredNumBufferedFrames;
+					const int32 NumBufferedFrames = LatestRecvFrame - FixedTickState.Interpolation.ToFrame;
+					const int32 BufferedMS = NumBufferedFrames * FixedTickState.FixedStepMS;
 
-					FixedTickState.Interpolation.PCT = 0.f;
-					FixedTickState.Interpolation.AccumulatedTimeMS = 0.f;
+					//UE_LOG(LogTemp, Warning, TEXT("BufferedMS: %d Frames: %d"), BufferedMS, NumBufferedFrames);
 
-					// We need to force a reconcile here since we supress the call until interpolation starts
-					for (TUniquePtr<IFixedInterpolateService>& Ptr : Services.FixedInterpolate.Array)
+					if (NumBufferedFrames <= 0)
 					{
-						Ptr->Reconcile(&FixedTickState);
+						InterpolateRate = 0.f;
 					}
 				}
-			}
-			else
-			{
-				const int32 NumBufferedFrames = LatestRecvFrame - FixedTickState.Interpolation.ToFrame;
-				const int32 BufferedMS = NumBufferedFrames * FixedTickState.FixedStepMS;
-
-				//UE_LOG(LogTemp, Warning, TEXT("BufferedMS: %d Frames: %d"), BufferedMS, NumBufferedFrames);
 
-				if (NumBufferedFrames <= 0)
+				int32 AdvanceFrames = 0;
+				if (InterpolateRate > 0.f)
 				{
-					InterpolateRate = 0.f;
-				}
-			}
+					const float fScaledDeltaTimeMS = (InterpolateRate * fEngineFrameDeltaTimeMS);
 
-			int32 AdvanceFrames = 0;
-			if (InterpolateRate > 0.f)
-			{
-				const float fScaledDeltaTimeMS = (InterpolateRate * fEngineFrameDeltaTimeMS);
-
-				FixedTickState.Interpolation.AccumulatedTimeMS += fScaledDeltaTimeMS;
-				AdvanceFrames = (int32)FixedTickState.Interpolation.AccumulatedTimeMS / FixedTickState.FixedStepRealTimeMS;
-				if (AdvanceFrames > 0)
-				{
-					FixedTickState.Interpolation.ToFrame += AdvanceFrames;
-					FixedTickState.Interpolation.AccumulatedTimeMS -= (AdvanceFrames * FixedTickState.FixedStepRealTimeMS);
-				}
-				const float RawPCT = FixedTickState.Interpolation.AccumulatedTimeMS / (float)FixedTickState.FixedStepRealTimeMS;
-				FixedTickState.Interpolation.PCT = FMath::Clamp<float>(RawPCT, 0.f, 1.f);
-				npEnsureMsgf(FixedTickState.Interpolation.PCT >= 0.f && FixedTickState.Interpolation.PCT <= 1.f, TEXT("Interpolation PCT out of range. %f"), FixedTickState.Interpolation.PCT);
+					FixedTickState.Interpolation.AccumulatedTimeMS += fScaledDeltaTimeMS;
+					AdvanceFrames = (int32)FixedTickState.Interpolation.AccumulatedTimeMS / FixedTickState.FixedStepRealTimeMS;
+					if (AdvanceFrames > 0)
+					{
+						FixedTickState.Interpolation.ToFrame += AdvanceFrames;
+						FixedTickState.Interpolation.AccumulatedTimeMS -= (AdvanceFrames * FixedTickState.FixedStepRealTimeMS);
+					}
+					const float RawPCT = FixedTickState.Interpolation.AccumulatedTimeMS / (float)FixedTickState.FixedStepRealTimeMS;
+					FixedTickState.Interpolation.PCT = FMath::Clamp<float>(RawPCT, 0.f, 1.f);
+					npEnsureMsgf(FixedTickState.Interpolation.PCT >= 0.f && FixedTickState.Interpolation.PCT <= 1.f, TEXT("Interpolation PCT out of range. %f"), FixedTickState.Interpolation.PCT);
 
-				const float PCTms = FixedTickState.Interpolation.PCT * (float)FixedTickState.FixedStepMS;
-				FixedTickState.Interpolation.InterpolatedTimeMS = ((FixedTickState.Interpolation.ToFrame-1) * FixedTickState.FixedStepMS) + (int32)PCTms;
+					const float PCTms = FixedTickState.Interpolation.PCT * (float)FixedTickState.FixedStepMS;
+					FixedTickState.Interpolation.InterpolatedTimeMS = ((FixedTickState.Interpolation.ToFrame-1) * FixedTickState.FixedStepMS) + (int32)PCTms;
 
-				//UE_LOG(LogTemp, Warning, TEXT("[Interpolate] %s Interpolating ToFrame %d. PCT: %.2f. Buffered: %d"), *GetPathName(), FixedTickState.Interpolation.ToFrame, FixedTickState.Interpolation.PCT, FixedTickState.Interpolation.LatestRecvFrame - FixedTickState.Interpolation.ToFrame);
+					//UE_LOG(LogTemp, Warning, TEXT("[Interpolate] %s Interpolating ToFrame %d. PCT: %.2f. Buffered: %d"), *GetPathName(), FixedTickState.Interpolation.ToFrame, FixedTickState.Interpolation.PCT, FixedTickState.Interpolation.LatestRecvFrame - FixedTickState.Interpolation.ToFrame);
 
-				for (TUniquePtr<IFixedInterpolateService>& Ptr : Services.FixedInterpolate.Array)
-				{
-					Ptr->FinalizeFrame(DeltaTimeSeconds, &FixedTickState);
+					{
+						TRACE_CPUPROFILER_EVENT_SCOPE(NetworkPrediction::FinalizeFrame);
+						for (TUniquePtr<IFixedInterpolateService>& Ptr : Services.FixedInterpolate.Array)
+						{
+							Ptr->FinalizeFrame(DeltaTimeSeconds, &FixedTickState);
+						}
+					}
 				}
 			}
 		}
-	}
 
-
-	if (Services.IndependentInterpolate.Array.Num() > 0)
-	{
-		const int32 DesiredBufferedMS = Settings.IndependentTickInterpolationBufferedMS;
-		const int32 MaxBufferedMS = Settings.IndependentTickInterpolationMaxBufferedMS;
-		
-		if (VariableTickState.Interpolation.LatestRecvTimeMS > DesiredBufferedMS)
+		if (Services.IndependentInterpolate.Array.Num() > 0)
 		{
-			float InterpolationRate = 1.f;
-
-			const int32 BufferedMS = VariableTickState.Interpolation.LatestRecvTimeMS - (int32)VariableTickState.Interpolation.fTimeMS;
-			if (BufferedMS > MaxBufferedMS)
-			{
-				UE_LOG(LogNetworkPrediction, Warning, TEXT("Independent Interpolation fell behind. BufferedMS: %d"), BufferedMS);
-				VariableTickState.Interpolation.fTimeMS = (float)(VariableTickState.Interpolation.LatestRecvTimeMS - DesiredBufferedMS);
-			}
-			else if (BufferedMS <= 0)
+			const int32 DesiredBufferedMS = Settings.IndependentTickInterpolationBufferedMS;
+			const int32 MaxBufferedMS = Settings.IndependentTickInterpolationMaxBufferedMS;
+		
+			if (VariableTickState.Interpolation.LatestRecvTimeMS > DesiredBufferedMS)
 			{
-				UE_LOG(LogNetworkPrediction, Warning, TEXT("Independent Interpolation starved: %d"), BufferedMS);
-				InterpolationRate = 0.f;
-			}
+				float InterpolationRate = 1.f;
 
-			if (InterpolationRate > 0.f)
-			{
-				const float fScaledDeltaTimeMS = (InterpolationRate * fEngineFrameDeltaTimeMS);
-				VariableTickState.Interpolation.fTimeMS += fScaledDeltaTimeMS;
-			}
+				const int32 BufferedMS = VariableTickState.Interpolation.LatestRecvTimeMS - (int32)VariableTickState.Interpolation.fTimeMS;
+				if (BufferedMS > MaxBufferedMS)
+				{
+					UE_LOG(LogNetworkPrediction, Warning, TEXT("Independent Interpolation fell behind. BufferedMS: %d"), BufferedMS);
+					VariableTickState.Interpolation.fTimeMS = (float)(VariableTickState.Interpolation.LatestRecvTimeMS - DesiredBufferedMS);
+				}
+				else if (BufferedMS <= 0)
+				{
+					UE_LOG(LogNetworkPrediction, Warning, TEXT("Independent Interpolation starved: %d"), BufferedMS);
+					InterpolationRate = 0.f;
+				}
 
-			for (TUniquePtr<IIndependentInterpolateService>& Ptr : Services.IndependentInterpolate.Array)
-			{
-				Ptr->FinalizeFrame(DeltaTimeSeconds, &VariableTickState);
+				if (InterpolationRate > 0.f)
+				{
+					const float fScaledDeltaTimeMS = (InterpolationRate * fEngineFrameDeltaTimeMS);
+					VariableTickState.Interpolation.fTimeMS += fScaledDeltaTimeMS;
+				}
+
+				{
+					TRACE_CPUPROFILER_EVENT_SCOPE(NetworkPrediction::FinalizeFrame);
+					for (TUniquePtr<IIndependentInterpolateService>& Ptr : Services.IndependentInterpolate.Array)
+					{
+						Ptr->FinalizeFrame(DeltaTimeSeconds, &VariableTickState);
+					}
+				}
 			}
 		}
 	}
@@ -529,23 +544,27 @@ void UNetworkPredictionWorldManager::BeginNewSimulationFrame_Internal(float Delt
 	// -------------------------------------------------------------------------
 	//	Finalize
 	// -------------------------------------------------------------------------
-	const int32 FixedTotalSimTimeMS = FixedTickState.GetTotalSimTimeMS();
-	const int32 FixedServerFrame = FixedTickState.PendingFrame + FixedTickState.Offset;
-	for (TUniquePtr<IFinalizeService>& Ptr : Services.FixedFinalize.Array)
 	{
-		Ptr->FinalizeFrame(DeltaTimeSeconds, FixedServerFrame, FixedTotalSimTimeMS, FixedTickState.FixedStepMS);
-	}
+		TRACE_CPUPROFILER_EVENT_SCOPE(NetworkPrediction::FinalizeFrame);
 
-	const int32 IndependentTotalSimTimeMS = VariableTickState.Frames[VariableTickState.PendingFrame].TotalMS;
-	const int32 IndependentFrame = VariableTickState.PendingFrame;
-	for (TUniquePtr<IFinalizeService>& Ptr : Services.IndependentLocalFinalize.Array)
-	{
-		Ptr->FinalizeFrame(DeltaTimeSeconds, IndependentFrame, IndependentTotalSimTimeMS, 0);
-	}
+		const int32 FixedTotalSimTimeMS = FixedTickState.GetTotalSimTimeMS();
+		const int32 FixedServerFrame = FixedTickState.PendingFrame + FixedTickState.Offset;
+		for (TUniquePtr<IFinalizeService>& Ptr : Services.FixedFinalize.Array)
+		{
+			Ptr->FinalizeFrame(DeltaTimeSeconds, FixedServerFrame, FixedTotalSimTimeMS, FixedTickState.FixedStepMS);
+		}
+
+		const int32 IndependentTotalSimTimeMS = VariableTickState.Frames[VariableTickState.PendingFrame].TotalMS;
+		const int32 IndependentFrame = VariableTickState.PendingFrame;
+		for (TUniquePtr<IFinalizeService>& Ptr : Services.IndependentLocalFinalize.Array)
+		{
+			Ptr->FinalizeFrame(DeltaTimeSeconds, IndependentFrame, IndependentTotalSimTimeMS, 0);
+		}
 	
-	for (TUniquePtr<IRemoteFinalizeService>& Ptr : Services.IndependentRemoteFinalize.Array)
-	{
-		Ptr->FinalizeFrame(DeltaTimeSeconds);
+		for (TUniquePtr<IRemoteFinalizeService>& Ptr : Services.IndependentRemoteFinalize.Array)
+		{
+			Ptr->FinalizeFrame(DeltaTimeSeconds);
+		}
 	}
 
 	// -------------------------------------------------------------------------
@@ -559,10 +578,12 @@ void UNetworkPredictionWorldManager::BeginNewSimulationFrame_Internal(float Delt
 	// -------------------------------------------------------------------------
 	// Call server RPC (common)
 	// -------------------------------------------------------------------------
-
-	for (TUniquePtr<IServerRPCService>& Ptr : Services.ServerRPC.Array)
 	{
-		Ptr->CallServerRPC(DeltaTimeSeconds);
+		TRACE_CPUPROFILER_EVENT_SCOPE(NetworkPrediction::CallServerRPC);
+		for (TUniquePtr<IServerRPCService>& Ptr : Services.ServerRPC.Array)
+		{
+			Ptr->CallServerRPC(DeltaTimeSeconds);
+		}
 	}
 }
 
diff --git a/Engine/Plugins/Runtime/NetworkPrediction/Source/NetworkPrediction/Public/Services/NetworkPredictionService_Rollback.inl b/Engine/Plugins/Runtime/NetworkPrediction/Source/NetworkPrediction/Public/Services/NetworkPredictionService_Rollback.inl
--- a/Engine/Plugins/Runtime/NetworkPrediction/Source/NetworkPrediction/Public/Services/NetworkPredictionService_Rollback.inl
+++ b/Engine/Plugins/Runtime/NetworkPrediction/Source/NetworkPrediction/Public/Services/NetworkPredictionService_Rollback.inl
@@ -184,6 +184,7 @@ public:
 			// (this frame or this rollback - they will need to restore their collision data since we are about to retick everyone in step)
 
 			QUICK_SCOPE_CYCLE_COUNTER(NP_Rollback_RestoreFrame);
+			TRACE_CPUPROFILER_EVENT_SCOPE(NetworkPrediction::RestoreFrame);
 
 			for (TConstSetBitIterator<> BitIt(InstanceBitArray); BitIt; ++BitIt)
 			{
