Some more EAllowShrinking fixes.

diff --git a/Engine/Plugins/Runtime/NetworkPrediction/Source/NetworkPrediction/Public/NetworkPredictionCues.h b/Engine/Plugins/Runtime/NetworkPrediction/Source/NetworkPrediction/Public/NetworkPredictionCues.h
--- a/Engine/Plugins/Runtime/NetworkPrediction/Source/NetworkPrediction/Public/NetworkPredictionCues.h
+++ b/Engine/Plugins/Runtime/NetworkPrediction/Source/NetworkPrediction/Public/NetworkPredictionCues.h
@@ -165,7 +165,7 @@ public:
 			FTypeInfo& TypeInfo = RegisteredTypes.CueTypes[i];
 			if (TypeInfo.IDPtr == &TCue::ID)
 			{
-				RegisteredTypes.CueTypes.RemoveAtSwap(i, 1, false);
+				RegisteredTypes.CueTypes.RemoveAtSwap(i, 1, EAllowShrinking::No);
 				RegisteredTypes.bDirty = true;
 				break;
 			}
