Adjusting some comments in NetworkPredictionConfig so options in projects settings have correct/helpful tooltips

diff --git a/Engine/Plugins/Runtime/NetworkPrediction/Source/NetworkPrediction/Public/NetworkPredictionConfig.h b/Engine/Plugins/Runtime/NetworkPrediction/Source/NetworkPrediction/Public/NetworkPredictionConfig.h
--- a/Engine/Plugins/Runtime/NetworkPrediction/Source/NetworkPrediction/Public/NetworkPredictionConfig.h
+++ b/Engine/Plugins/Runtime/NetworkPrediction/Source/NetworkPrediction/Public/NetworkPredictionConfig.h
@@ -8,8 +8,10 @@
 UENUM()
 enum class ENetworkPredictionTickingPolicy : uint8
 {
-	Independent	= 1 << 0,	// Client ticks at local frame rate. Server ticks clients independently at client input cmd rate.
-	Fixed		= 1 << 1,	// Everyone ticks at same fixed rate. Supports group rollback and physics.
+	// Client ticks at local frame rate. Server ticks clients independently at client input cmd rate.
+	Independent	= 1 << 0,
+	// Everyone ticks at same fixed rate. Supports group rollback and physics.
+	Fixed		= 1 << 1,
 
 	All = Independent | Fixed UMETA(Hidden),
 };
@@ -17,17 +19,18 @@ ENUM_CLASS_FLAGS(ENetworkPredictionTickingPolicy);
 
 enum class ENetworkPredictionLocalInputPolicy : uint8
 {
-	Passive,			// Up to the user to write input via FNetSimProxy::WriteInputCmd.
-	PollPerSimFrame,	// ProduceInput is called on the driver before every simulation frame. 
-						// This may be necessary for things like aim assist and fixed step simulations that run multiple sim frames per engine frame
+	// Up to the user to write input via FNetSimProxy::WriteInputCmd.
+	Passive,
+	// ProduceInput is called on the driver before every simulation frame. This may be necessary for things like aim assist and fixed step simulations that run multiple sim frames per engine frame
+	PollPerSimFrame,
 };
 
-// Must be kept in sync with ENP_NetworkLOD
+// Must be kept in sync with ENP_NetworkLOD. Note: SimExtrapolate Not currently implemented so it is hidden
 UENUM()
 enum class ENetworkLOD : uint8
 {
 	Interpolated	= 1 << 0,
-	SimExtrapolate	= 1 << 1 UMETA(Hidden), // Not currently implemented
+	SimExtrapolate	= 1 << 1 UMETA(Hidden),
 	ForwardPredict	= 1 << 2,
 
 	All = Interpolated | SimExtrapolate | ForwardPredict UMETA(Hidden),
