Fixing compile error

#rnx

diff --git a/Engine/Plugins/Runtime/NetworkPrediction/Source/NetworkPrediction/Private/NetworkPredictionModule.cpp b/Engine/Plugins/Runtime/NetworkPrediction/Source/NetworkPrediction/Private/NetworkPredictionModule.cpp
--- a/Engine/Plugins/Runtime/NetworkPrediction/Source/NetworkPrediction/Private/NetworkPredictionModule.cpp
+++ b/Engine/Plugins/Runtime/NetworkPrediction/Source/NetworkPrediction/Private/NetworkPredictionModule.cpp
@@ -4,6 +4,7 @@
 #include "NetworkPredictionTrace.h"
 #include "String/ParseTokens.h"
 #include "NetworkPredictionModelDefRegistry.h"
+#include "Misc/CoreDelegates.h"
 
 #if WITH_EDITOR
 #include "Editor.h"
