Network Prediction: fixes issue where the FixedTickFrameRate setting was not respected for non-physics simulations.

diff --git a/Engine/Plugins/Runtime/NetworkPrediction/Source/NetworkPrediction/Private/NetworkPredictionWorldManager.cpp b/Engine/Plugins/Runtime/NetworkPrediction/Source/NetworkPrediction/Private/NetworkPredictionWorldManager.cpp
--- a/Engine/Plugins/Runtime/NetworkPrediction/Source/NetworkPrediction/Private/NetworkPredictionWorldManager.cpp
+++ b/Engine/Plugins/Runtime/NetworkPrediction/Source/NetworkPrediction/Private/NetworkPredictionWorldManager.cpp
@@ -81,7 +81,10 @@ void UNetworkPredictionWorldManager::OnWorldPreTick(UWorld* InWorld, ELevelTick
 
 	UE_NP_TRACE_WORLD_FRAME_START(InWorld->GetGameInstance(), InDeltaSeconds);
 
-	OnWorldPreTick_Internal(InDeltaSeconds, GEngine->FixedFrameRate);
+	// Defer to the engine ticking rate, if fixed. Otherwise use NPP's setting.
+	const float FixedTickFrameRate = (GEngine && GEngine->bUseFixedFrameRate) ? GEngine->FixedFrameRate : Settings.FixedTickFrameRate;
+
+	OnWorldPreTick_Internal(InDeltaSeconds, FixedTickFrameRate);
 
 	// Instantiate replicated manager on server
 	if (!ReplicatedManager && InWorld->GetNetMode() != NM_Client)
diff --git a/Engine/Plugins/Runtime/NetworkPrediction/Source/NetworkPrediction/Public/NetworkPredictionSettings.h b/Engine/Plugins/Runtime/NetworkPrediction/Source/NetworkPrediction/Public/NetworkPredictionSettings.h
--- a/Engine/Plugins/Runtime/NetworkPrediction/Source/NetworkPrediction/Public/NetworkPredictionSettings.h
+++ b/Engine/Plugins/Runtime/NetworkPrediction/Source/NetworkPrediction/Public/NetworkPredictionSettings.h
@@ -25,7 +25,7 @@ struct FNetworkPredictionSettings
 
 	// ------------------------------------------------------------------------------------------
 
-	// Frame rate to use when running Fixed Tick simulations. Note: Engine::FixedFrameRate will take precedence if manually set.
+	// Frame rate to use when running Fixed Tick simulations. Note: Engine::FixedFrameRate will take precedence if Engine::bUseFixedFrameRate is enabled.
 	UPROPERTY(config, EditAnywhere, Category = FixedTick)
 	int32 FixedTickFrameRate = 60;
 
