Network Prediction plugin: adds partial support for replays and fixes crash when starting playback. Replays work for authority and autonomous proxy simulations, but not well for simulated proxies. Playback only works when SimulatedProxyNetworkLOD = Interpolated.

#11875

diff --git a/Engine/Plugins/Runtime/NetworkPrediction/Source/NetworkPrediction/Private/NetworkPredictionComponent.cpp b/Engine/Plugins/Runtime/NetworkPrediction/Source/NetworkPrediction/Private/NetworkPredictionComponent.cpp
--- a/Engine/Plugins/Runtime/NetworkPrediction/Source/NetworkPrediction/Private/NetworkPredictionComponent.cpp
+++ b/Engine/Plugins/Runtime/NetworkPrediction/Source/NetworkPrediction/Private/NetworkPredictionComponent.cpp
@@ -61,7 +61,7 @@ void UNetworkPredictionComponent::GetLifetimeReplicatedProps(TArray< FLifetimePr
 
 	DOREPLIFETIME( UNetworkPredictionComponent, NetworkPredictionProxy);
 	DOREPLIFETIME_CONDITION( UNetworkPredictionComponent, ReplicationProxy_Autonomous, COND_AutonomousOnly);
-	DOREPLIFETIME_CONDITION( UNetworkPredictionComponent, ReplicationProxy_Simulated, COND_SimulatedOnly);
+	DOREPLIFETIME_CONDITION( UNetworkPredictionComponent, ReplicationProxy_Simulated, COND_SimulatedOnlyNoReplay);
 	DOREPLIFETIME_CONDITION( UNetworkPredictionComponent, ReplicationProxy_Replay, COND_ReplayOnly);
 }
 
diff --git a/Engine/Plugins/Runtime/NetworkPrediction/Source/NetworkPrediction/Public/NetworkPredictionWorldManager.h b/Engine/Plugins/Runtime/NetworkPrediction/Source/NetworkPrediction/Public/NetworkPredictionWorldManager.h
--- a/Engine/Plugins/Runtime/NetworkPrediction/Source/NetworkPrediction/Public/NetworkPredictionWorldManager.h
+++ b/Engine/Plugins/Runtime/NetworkPrediction/Source/NetworkPrediction/Public/NetworkPredictionWorldManager.h
@@ -170,6 +170,17 @@ private:
 	template<typename ReplicatorType, typename ModelDef=typename ReplicatorType::ModelDef>
 	void BindNetSend_IndependentRemote(FNetworkPredictionID ID, FReplicationProxy* RepProxy, TModelDataStore<ModelDef>* DataStore);
 
+	// ----
+	
+	template<typename ReplicatorType, typename ModelDef=typename ReplicatorType::ModelDef>
+	void BindReplayNetSendRecv_Fixed(FNetworkPredictionID ID, FReplicationProxy* RepProxy, TModelDataStore<ModelDef>* DataStore, ENetRole NetRole);
+
+	template<typename ReplicatorType, typename ModelDef = typename ReplicatorType::ModelDef>
+	void BindReplayNetSendRecv_IndependentLocal(FNetworkPredictionID ID, FReplicationProxy* RepProxy, TModelDataStore<ModelDef>* DataStore, ENetRole NetRole);
+
+	template<typename ReplicatorType, typename ModelDef = typename ReplicatorType::ModelDef>
+	void BindReplayNetSendRecv_IndependentRemote(FNetworkPredictionID ID, FReplicationProxy* RepProxy, TModelDataStore<ModelDef>* DataStore, ENetRole NetRole);
+
 	// ---------------------------------------
 
 	template<typename ModelDef>
@@ -223,7 +234,7 @@ void UNetworkPredictionWorldManager::ConfigureInstance(FNetworkPredictionID ID,
 					BindServerNetRecv_Independent<ModelDef>(ID, RepProxies.ServerRPC, DataStore);
 					BindNetSend_IndependentRemote<TIndependentTickReplicator_AP<ModelDef>>(ID, RepProxies.AutonomousProxy, DataStore);
 					BindNetSend_IndependentRemote<TIndependentTickReplicator_SP<ModelDef>>(ID, RepProxies.SimulatedProxy, DataStore);
-					BindNetSend_IndependentRemote<TIndependentTickReplicator_SP<ModelDef>>(ID, RepProxies.Replay, DataStore);
+					BindReplayNetSendRecv_IndependentRemote<TIndependentTickReplicator_SP<ModelDef>>(ID, RepProxies.Replay, DataStore, Role);
 
 					ServiceMask |= ENetworkPredictionService::IndependentRemoteTick;
 					ServiceMask |= ENetworkPredictionService::IndependentRemoteFinalize;
@@ -245,7 +256,7 @@ void UNetworkPredictionWorldManager::ConfigureInstance(FNetworkPredictionID ID,
 				{
 					// Locally controlled
 					BindNetSend_IndependentLocal<TIndependentTickReplicator_SP<ModelDef>>(ID, RepProxies.SimulatedProxy, DataStore);
-					BindNetSend_IndependentLocal<TIndependentTickReplicator_SP<ModelDef>>(ID, RepProxies.Replay, DataStore);
+					BindReplayNetSendRecv_IndependentLocal<TIndependentTickReplicator_SP<ModelDef>>(ID, RepProxies.Replay, DataStore, Role);
 
 					if (FNetworkPredictionDriver<ModelDef>::HasSimulation())
 					{
@@ -264,7 +275,7 @@ void UNetworkPredictionWorldManager::ConfigureInstance(FNetworkPredictionID ID,
 			{
 				BindClientNetRecv_Independent<TIndependentTickReplicator_AP<ModelDef>>(ID, RepProxies.AutonomousProxy, DataStore, Role);
 				BindNetSend_IndependentLocal<TIndependentTickReplicator_Server<ModelDef>>(ID, RepProxies.ServerRPC, DataStore);
-				BindNetSend_IndependentLocal<TIndependentTickReplicator_SP<ModelDef>>(ID, RepProxies.Replay, DataStore);
+				BindReplayNetSendRecv_IndependentLocal<TIndependentTickReplicator_SP<ModelDef>>(ID, RepProxies.Replay, DataStore, Role);
 				
 				npCheckf(FNetworkPredictionDriver<ModelDef>::HasSimulation(), TEXT("AP must have Simulation."));
 				npCheckf(FNetworkPredictionDriver<ModelDef>::HasInput(), TEXT("AP sim doesn't have Input?"));
@@ -281,7 +292,7 @@ void UNetworkPredictionWorldManager::ConfigureInstance(FNetworkPredictionID ID,
 			{
 				BindClientNetRecv_Independent<TIndependentTickReplicator_AP<ModelDef>>(ID, RepProxies.AutonomousProxy, DataStore, Role);
 				BindClientNetRecv_Independent<TIndependentTickReplicator_SP<ModelDef>>(ID, RepProxies.SimulatedProxy, DataStore, Role);
-				BindNetSend_IndependentLocal<TIndependentTickReplicator_SP<ModelDef>>(ID, RepProxies.Replay, DataStore);
+				BindReplayNetSendRecv_IndependentLocal<TIndependentTickReplicator_SP<ModelDef>>(ID, RepProxies.Replay, DataStore, Role);
 
 				// Interpolation is the only supported mode for independently ticked SP simulations
 				// (will add support for sim-extrapolate eventually)
@@ -318,7 +329,7 @@ void UNetworkPredictionWorldManager::ConfigureInstance(FNetworkPredictionID ID,
 				}
 				
 				BindNetSend_Fixed<TFixedTickReplicator_SP<ModelDef>>(ID, RepProxies.SimulatedProxy, DataStore);
-				BindNetSend_Fixed<TFixedTickReplicator_SP<ModelDef>>(ID, RepProxies.Replay, DataStore);
+				BindReplayNetSendRecv_Fixed<TFixedTickReplicator_SP<ModelDef>>(ID, RepProxies.Replay, DataStore, Role);
 				break;
 			}
 			case ENetRole::ROLE_AutonomousProxy:
@@ -330,7 +341,7 @@ void UNetworkPredictionWorldManager::ConfigureInstance(FNetworkPredictionID ID,
 				BindClientNetRecv_Fixed<TFixedTickReplicator_SP<ModelDef>>(ID, RepProxies.SimulatedProxy, DataStore, Role);
 
 				BindNetSend_Fixed<TFixedTickReplicator_Server<ModelDef>>(ID, RepProxies.ServerRPC, DataStore);
-				BindNetSend_Fixed<TFixedTickReplicator_SP<ModelDef>>(ID, RepProxies.Replay, DataStore);
+				BindReplayNetSendRecv_Fixed<TFixedTickReplicator_SP<ModelDef>>(ID, RepProxies.Replay, DataStore, Role);
 				
 				// Poll local input and send to server services
 				ServiceMask |= ENetworkPredictionService::FixedInputLocal;
@@ -342,7 +353,7 @@ void UNetworkPredictionWorldManager::ConfigureInstance(FNetworkPredictionID ID,
 				BindClientNetRecv_Fixed<TFixedTickReplicator_AP<ModelDef>>(ID, RepProxies.AutonomousProxy, DataStore, Role);
 				BindClientNetRecv_Fixed<TFixedTickReplicator_SP<ModelDef>>(ID, RepProxies.SimulatedProxy, DataStore, Role);
 
-				BindNetSend_Fixed<TFixedTickReplicator_SP<ModelDef>>(ID, RepProxies.Replay, DataStore);
+				BindReplayNetSendRecv_Fixed<TFixedTickReplicator_SP<ModelDef>>(ID, RepProxies.Replay, DataStore, Role);
 				break;
 			}
 		};
@@ -620,6 +631,117 @@ void UNetworkPredictionWorldManager::BindNetSend_IndependentRemote(FNetworkPredi
 	};
 }
 
+// ---------------------------------------------------------------------------------------
+//	Replay Bindings
+//		-Binds RepProxy serialize lambda to a Replicator function
+//		-3 versions: Fixed, Independent (Local Tick), Independent (Remote Ticked)
+// ---------------------------------------------------------------------------------------
+
+template<typename ReplicatorType, typename ModelDef>
+void UNetworkPredictionWorldManager::BindReplayNetSendRecv_Fixed(FNetworkPredictionID ID, FReplicationProxy* RepProxy, TModelDataStore<ModelDef>* DataStore, ENetRole NetRole)
+{
+	if (!RepProxy)
+	{
+		return;
+	}
+
+	const int32 ClientRecvIdx = DataStore->ClientRecv.GetIndex(ID);
+	NpResizeAndSetBit(DataStore->ClientRecvBitMask, ClientRecvIdx, false);
+
+	TClientRecvData<ModelDef>& ClientRecvData = DataStore->ClientRecv.GetByIndexChecked(ClientRecvIdx);
+	InitClientRecvData<ModelDef>(ID, ClientRecvData, DataStore, NetRole);
+
+	FFixedTickState* TickState = &this->FixedTickState;
+	RepProxy->NetSerializeFunc = [ID, DataStore, ClientRecvIdx, TickState](const FNetSerializeParams& P)
+	{
+		if (P.Ar.IsLoading()) // Receiving replay data
+		{
+			DataStore->ClientRecvBitMask[ClientRecvIdx] = true;
+			auto& ClientRecvData = DataStore->ClientRecv.GetByIndexChecked(ClientRecvIdx);
+
+			UE_NP_TRACE_SIM(ClientRecvData.TraceID);
+			ReplicatorType::NetRecv(P, ClientRecvData, DataStore, TickState);
+		}
+		else // Sending replay data
+		{
+			UE_NP_TRACE_SIM(ID.GetTraceID());
+			ReplicatorType::NetSend(P, ID, DataStore, TickState);
+		}
+	};
+}
+
+
+template<typename ReplicatorType, typename ModelDef>
+void UNetworkPredictionWorldManager::BindReplayNetSendRecv_IndependentLocal(FNetworkPredictionID ID, FReplicationProxy* RepProxy, TModelDataStore<ModelDef>* DataStore, ENetRole NetRole)
+{
+	if (!RepProxy)
+	{
+		return;
+	}
+
+	const int32 ClientRecvIdx = DataStore->ClientRecv.GetIndex(ID);
+	NpResizeAndSetBit(DataStore->ClientRecvBitMask, ClientRecvIdx, false);
+
+	TClientRecvData<ModelDef>& ClientRecvData = DataStore->ClientRecv.GetByIndexChecked(ClientRecvIdx);
+	InitClientRecvData<ModelDef>(ID, ClientRecvData, DataStore, NetRole);
+
+	FVariableTickState* TickState = &this->VariableTickState;
+	RepProxy->NetSerializeFunc = [ID, DataStore, ClientRecvIdx, TickState](const FNetSerializeParams& P)
+	{
+		if (P.Ar.IsLoading())	// Receiving replay data
+		{
+			DataStore->ClientRecvBitMask[ClientRecvIdx] = true;
+			auto& ClientRecvData = DataStore->ClientRecv.GetByIndexChecked(ClientRecvIdx);
+
+			UE_NP_TRACE_SIM(ClientRecvData.TraceID);
+			ReplicatorType::NetRecv(P, ClientRecvData, DataStore, TickState);
+		}
+		else // Sending replay data
+		{
+			UE_NP_TRACE_SIM(ID.GetTraceID());
+			ReplicatorType::NetSend(P, ID, DataStore, TickState);
+		}
+	};
+}
+
+
+template<typename ReplicatorType, typename ModelDef>
+void UNetworkPredictionWorldManager::BindReplayNetSendRecv_IndependentRemote(FNetworkPredictionID ID, FReplicationProxy* RepProxy, TModelDataStore<ModelDef>* DataStore, ENetRole NetRole)
+{
+	if (!RepProxy)
+	{
+		return;
+	}
+
+	const int32 ClientRecvIdx = DataStore->ClientRecv.GetIndex(ID);
+	NpResizeAndSetBit(DataStore->ClientRecvBitMask, ClientRecvIdx, false);
+
+	TClientRecvData<ModelDef>& ClientRecvData = DataStore->ClientRecv.GetByIndexChecked(ClientRecvIdx);
+	InitClientRecvData<ModelDef>(ID, ClientRecvData, DataStore, NetRole);
+
+	FVariableTickState* TickState = &this->VariableTickState;
+
+	const int32 ServerRecvIdx = DataStore->ServerRecv_IndependentTick.GetIndex(ID);
+
+	RepProxy->NetSerializeFunc = [ID, DataStore, ClientRecvIdx, TickState, ServerRecvIdx](const FNetSerializeParams& P)
+	{
+		if (P.Ar.IsLoading())	// Receiving replay data
+		{
+			DataStore->ClientRecvBitMask[ClientRecvIdx] = true;
+			auto& ClientRecvData = DataStore->ClientRecv.GetByIndexChecked(ClientRecvIdx);
+
+			UE_NP_TRACE_SIM(ClientRecvData.TraceID);
+			ReplicatorType::NetRecv(P, ClientRecvData, DataStore, TickState);
+		}
+		else // Sending replay data
+		{
+			TServerRecvData_Independent<ModelDef>& ServerRecv = DataStore->ServerRecv_IndependentTick.GetByIndexChecked(ServerRecvIdx);
+			UE_NP_TRACE_SIM(ID.GetTraceID());
+			ReplicatorType::NetSend(P, ID, DataStore, ServerRecv, TickState);
+		}
+	};
+}
+
 // -----------------
 
 template<typename ModelDef>
