Network Prediction: added fixed tick smoothing service (off by default)
Mover:  used  this to add optional visual component offset smoothing to Mover.
Adapted from PR.

diff --git a/Engine/Plugins/Experimental/Mover/Source/Mover/Private/Backends/MoverNetworkPredictionLiaison.cpp b/Engine/Plugins/Experimental/Mover/Source/Mover/Private/Backends/MoverNetworkPredictionLiaison.cpp
--- a/Engine/Plugins/Experimental/Mover/Source/Mover/Private/Backends/MoverNetworkPredictionLiaison.cpp
+++ b/Engine/Plugins/Experimental/Mover/Source/Mover/Private/Backends/MoverNetworkPredictionLiaison.cpp
@@ -72,6 +72,12 @@ void UMoverNetworkPredictionLiaisonComponent::FinalizeFrame(const FMoverSyncStat
 	MoverComp->FinalizeFrame(SyncState, AuxState);
 }
 
+void UMoverNetworkPredictionLiaisonComponent::FinalizeSmoothingFrame(const FMoverSyncState* SyncState, const FMoverAuxStateContext* AuxState)
+{
+	check(MoverComp);
+	MoverComp->FinalizeSmoothingFrame(SyncState, AuxState);
+}
+
 void UMoverNetworkPredictionLiaisonComponent::InitializeSimulationState(FMoverSyncState* OutSync, FMoverAuxStateContext* OutAux)
 {
 	check(MoverComp);
@@ -111,7 +117,7 @@ int32 UMoverNetworkPredictionLiaisonComponent::GetCurrentSimFrame()
 
 bool UMoverNetworkPredictionLiaisonComponent::ReadPendingSyncState(OUT FMoverSyncState& OutSyncState)
 {
-	if (const FMoverSyncState* PendingSyncState = NetworkPredictionProxy.ReadSyncState<FMoverSyncState>())
+	if (const FMoverSyncState* PendingSyncState = NetworkPredictionProxy.ReadSyncState<FMoverSyncState>(ENetworkPredictionStateRead::Simulation))
 	{
 		OutSyncState = *PendingSyncState;
 		return true;
@@ -122,14 +128,61 @@ bool UMoverNetworkPredictionLiaisonComponent::ReadPendingSyncState(OUT FMoverSyn
 
 bool UMoverNetworkPredictionLiaisonComponent::WritePendingSyncState(const FMoverSyncState& SyncStateToWrite)
 {
-	NetworkPredictionProxy.WriteSyncState<FMoverSyncState>([&SyncStateToWrite](FMoverSyncState& PendingSyncStateRef)
+	bool bDidWriteSucceed = NetworkPredictionProxy.WriteSyncState<FMoverSyncState>([&SyncStateToWrite](FMoverSyncState& PendingSyncStateRef)
 		{
 			PendingSyncStateRef = SyncStateToWrite;
-		});
+		}) != nullptr;
+
+	return bDidWriteSucceed;
+}
+
+
+bool UMoverNetworkPredictionLiaisonComponent::ReadPresentationSyncState(OUT FMoverSyncState& OutSyncState)
+{
+	if (const FMoverSyncState* PendingSyncState = NetworkPredictionProxy.ReadSyncState<FMoverSyncState>(ENetworkPredictionStateRead::Presentation))
+	{
+		OutSyncState = *PendingSyncState;
+		return true;
+	}
+
+	return false;
+}
+
+
+bool UMoverNetworkPredictionLiaisonComponent::WritePresentationSyncState(const FMoverSyncState& SyncStateToWrite)
+{
+	bool bDidWriteSucceed = NetworkPredictionProxy.WritePresentationSyncState<FMoverSyncState>([&SyncStateToWrite](FMoverSyncState& PresentationSyncStateRef)
+		{
+			PresentationSyncStateRef = SyncStateToWrite;
+		}) != nullptr;
+
+	return bDidWriteSucceed;
+}
+
+
+bool UMoverNetworkPredictionLiaisonComponent::ReadPrevPresentationSyncState(FMoverSyncState& OutSyncState)
+{
+	if (const FMoverSyncState* PrevPresentationSyncState = NetworkPredictionProxy.ReadPrevPresentationSyncState<FMoverSyncState>())
+	{
+		OutSyncState = *PrevPresentationSyncState;
+		return true;
+	}
 
-	return true;
+	return false;
 }
 
+
+bool UMoverNetworkPredictionLiaisonComponent::WritePrevPresentationSyncState(const FMoverSyncState& SyncStateToWrite)
+{
+	bool bDidWriteSucceed = NetworkPredictionProxy.WritePrevPresentationSyncState<FMoverSyncState>([&SyncStateToWrite](FMoverSyncState& PresentationSyncStateRef)
+		{
+			PresentationSyncStateRef = SyncStateToWrite;
+		}) != nullptr;
+
+	return bDidWriteSucceed;
+}
+
+
 #if WITH_EDITOR
 EDataValidationResult UMoverNetworkPredictionLiaisonComponent::ValidateData(FDataValidationContext& Context, const UMoverComponent& ValidationMoverComp) const
 {
diff --git a/Engine/Plugins/Experimental/Mover/Source/Mover/Private/MoveLibrary/BasedMovementUtils.cpp b/Engine/Plugins/Experimental/Mover/Source/Mover/Private/MoveLibrary/BasedMovementUtils.cpp
--- a/Engine/Plugins/Experimental/Mover/Source/Mover/Private/MoveLibrary/BasedMovementUtils.cpp
+++ b/Engine/Plugins/Experimental/Mover/Source/Mover/Private/MoveLibrary/BasedMovementUtils.cpp
@@ -425,23 +425,73 @@ void UBasedMovementUtils::UpdateSimpleBasedMovement(UMoverComponent* TargetMover
 						CurrentBaseInfo.ContactLocalPosition -= UnachievedLocalDelta;
 					}
 
-					// Read, edit, write out the sync state based on the move results
-					FMoverSyncState PendingSyncState;
-					if (TargetMoverComp->BackendLiaisonComp->ReadPendingSyncState(OUT PendingSyncState))
+					// Propagate the movement changes to the backend's state, if supported
+
+					// Note that this is occurring out-of-band with the movement simulation, in order to support based movement regardless of update order or
+					// whether the movement base is also simulated through Mover.
+
+					FMoverSyncState PendingSimSyncState;
+					if (TargetMoverComp->BackendLiaisonComp->ReadPendingSyncState(OUT PendingSimSyncState))
 					{
-						if (FMoverDefaultSyncState* DefaultSyncState = PendingSyncState.SyncStateCollection.FindMutableDataByType<FMoverDefaultSyncState>())
+						// Modify the PENDING sync state that has not yet been committed to simulation history nor replicated
+						if (FMoverDefaultSyncState* PendingMoverState = PendingSimSyncState.SyncStateCollection.FindMutableDataByType<FMoverDefaultSyncState>())
 						{
-							FVector Velocity = DefaultSyncState->GetVelocity_WorldSpace();
-							DefaultSyncState->SetTransforms_WorldSpace(UpdatedComponent->GetComponentLocation(),
-								UpdatedComponent->GetComponentRotation(),
-								Velocity,
-								CurrentBaseInfo.MovementBase.Get(), CurrentBaseInfo.BoneName);
-
-							TargetMoverComp->BackendLiaisonComp->WritePendingSyncState(PendingSyncState);
+							FTransform OldSyncTransformWs = PendingMoverState->GetTransform_WorldSpace();
+							FTransform NewSyncTransformWs = UpdatedComponent->GetComponentTransform();
+
+							PendingMoverState->SetTransforms_WorldSpace(
+								NewSyncTransformWs.GetLocation(),
+								NewSyncTransformWs.GetRotation().Rotator(),
+								PendingMoverState->GetVelocity_WorldSpace(),	// keep same velocity and base
+								PendingMoverState->GetMovementBase(), PendingMoverState->GetMovementBaseBoneName());
+
+							TargetMoverComp->BackendLiaisonComp->WritePendingSyncState(PendingSimSyncState);	// writes pending Simulation state
+
+							// If smoothing, modify presentation-related states as well so that the visual offset location stays anchored to the movement base
+							if (TargetMoverComp->SmoothingMode != EMoverSmoothingMode::None)
+							{
+								const FTransform OldToNewTransform = NewSyncTransformWs.GetRelativeTransform(OldSyncTransformWs);
+
+								// Modify the PRESENTATION sync state that we're smoothing TO
+								FMoverSyncState PresentationSyncState;
+								if (TargetMoverComp->BackendLiaisonComp->ReadPresentationSyncState(OUT PresentationSyncState))
+								{
+									if (FMoverDefaultSyncState* PresentationMoverState = PresentationSyncState.SyncStateCollection.FindMutableDataByType<FMoverDefaultSyncState>())
+									{
+										OldSyncTransformWs = PresentationMoverState->GetTransform_WorldSpace();
+										NewSyncTransformWs = OldToNewTransform * OldSyncTransformWs;
+
+										PresentationMoverState->SetTransforms_WorldSpace(
+											NewSyncTransformWs.GetLocation(),
+											NewSyncTransformWs.GetRotation().Rotator(),
+											PresentationMoverState->GetVelocity_WorldSpace(),	// keep same velocity and base
+											PresentationMoverState->GetMovementBase(), PresentationMoverState->GetMovementBaseBoneName());
+
+										TargetMoverComp->BackendLiaisonComp->WritePresentationSyncState(PresentationSyncState);
+									}
+								}
+
+								// Modify the PREV PRESENTATION sync state that we're smoothing FROM
+								FMoverSyncState PrevPresentationSyncState;
+								if (TargetMoverComp->BackendLiaisonComp->ReadPrevPresentationSyncState(OUT PrevPresentationSyncState))
+								{
+									if (FMoverDefaultSyncState* PrevPresentationMoverState = PrevPresentationSyncState.SyncStateCollection.FindMutableDataByType<FMoverDefaultSyncState>())
+									{
+										OldSyncTransformWs = PrevPresentationMoverState->GetTransform_WorldSpace();
+										NewSyncTransformWs = OldToNewTransform * OldSyncTransformWs;
+
+										PrevPresentationMoverState->SetTransforms_WorldSpace(
+											NewSyncTransformWs.GetLocation(),
+											NewSyncTransformWs.GetRotation().Rotator(),
+											PrevPresentationMoverState->GetVelocity_WorldSpace(),	// keep same velocity and base
+											PrevPresentationMoverState->GetMovementBase(), PrevPresentationMoverState->GetMovementBaseBoneName());
+
+										TargetMoverComp->BackendLiaisonComp->WritePrevPresentationSyncState(PrevPresentationSyncState);	
+									}
+								}
+							}
 						}
 					}
-
-
 				}
 
 				SimBlackboard->Set(CommonBlackboard::LastAppliedDynamicMovementBase, CurrentBaseInfo);
@@ -484,3 +534,4 @@ FName FMoverDynamicBasedMovementTickFunction::DiagnosticContext(bool bDetailed)
 	}
 	return FName(TEXT("FMoverDynamicBasedMovementTickFunction"));
 }
+
diff --git a/Engine/Plugins/Experimental/Mover/Source/Mover/Private/MoverComponent.cpp b/Engine/Plugins/Experimental/Mover/Source/Mover/Private/MoverComponent.cpp
--- a/Engine/Plugins/Experimental/Mover/Source/Mover/Private/MoverComponent.cpp
+++ b/Engine/Plugins/Experimental/Mover/Source/Mover/Private/MoverComponent.cpp
@@ -171,6 +171,11 @@ void UMoverComponent::BeginPlay()
 
 		ensureMsgf(UpdatedComponent && (PrimaryVisualComponent != UpdatedComponent), TEXT("A Mover actor (%s) must have an UpdatedComponent and cannot have a PrimaryVisualComponent that is the same as UpdatedComponent"), *GetNameSafe(MyActor));
 
+		if (PrimaryVisualComponent)
+		{
+			BaseVisualComponentTransform = PrimaryVisualComponent->GetRelativeTransform();
+		}
+
 		// Optional motion warping support
 		if (UMotionWarpingComponent* WarpingComp = MyActor->FindComponentByClass<UMotionWarpingComponent>())
 		{
@@ -263,6 +268,29 @@ void UMoverComponent::FinalizeFrame(const FMoverSyncState* SyncState, const FMov
 	}
 }
 
+void UMoverComponent::FinalizeSmoothingFrame(const FMoverSyncState* SyncState, const FMoverAuxStateContext* AuxState)
+{
+	if (PrimaryVisualComponent)
+	{
+		if (SmoothingMode == EMoverSmoothingMode::VisualComponentOffset)
+		{
+			// Offset the visual component so it aligns with the smoothed state transform, while leaving the actual root component in place
+			if (const FMoverDefaultSyncState* MoverState = SyncState->SyncStateCollection.FindDataByType<FMoverDefaultSyncState>())
+			{
+				FTransform ActorTransform = FTransform(MoverState->GetOrientation_WorldSpace(), MoverState->GetLocation_WorldSpace(), FVector::OneVector);
+				PrimaryVisualComponent->SetWorldTransform(BaseVisualComponentTransform * ActorTransform);	// smoothed location with base offset applied
+			}
+		}
+		else
+		{
+			if (!PrimaryVisualComponent->GetRelativeTransform().Equals(BaseVisualComponentTransform))
+			{
+				PrimaryVisualComponent->SetRelativeTransform(BaseVisualComponentTransform);
+			}
+		}
+	}
+}
+
 void UMoverComponent::TickInterpolatedSimProxy(const FMoverTimeStep& TimeStep, const FMoverInputCmdContext& InputCmd, UMoverComponent* MoverComp, const FMoverSyncState& CachedSyncState, const FMoverSyncState& SyncState, const FMoverAuxStateContext& AuxState)
 {
 	TArray<TSharedPtr<FMovementModifierBase>> ModifiersToStart;
diff --git a/Engine/Plugins/Experimental/Mover/Source/Mover/Private/MoverDataModelTypes.cpp b/Engine/Plugins/Experimental/Mover/Source/Mover/Private/MoverDataModelTypes.cpp
--- a/Engine/Plugins/Experimental/Mover/Source/Mover/Private/MoverDataModelTypes.cpp
+++ b/Engine/Plugins/Experimental/Mover/Source/Mover/Private/MoverDataModelTypes.cpp
@@ -375,6 +375,21 @@ FRotator FMoverDefaultSyncState::GetOrientation_BaseSpace() const
 	return Orientation;
 }
 
+FTransform FMoverDefaultSyncState::GetTransform_WorldSpace() const
+{
+	if (MovementBase)
+	{
+		return FTransform(Orientation, Location) * FTransform(MovementBaseQuat, MovementBasePos);
+	}
+
+	return FTransform(Orientation, Location);
+}
+
+FTransform FMoverDefaultSyncState::GetTransform_BaseSpace() const
+{
+	return FTransform(Orientation, Location);
+}
+
 
 // UMoverDataModelBlueprintLibrary ///////////////////////////////////////////////////
 
diff --git a/Engine/Plugins/Experimental/Mover/Source/Mover/Public/Backends/MoverBackendLiaison.h b/Engine/Plugins/Experimental/Mover/Source/Mover/Public/Backends/MoverBackendLiaison.h
--- a/Engine/Plugins/Experimental/Mover/Source/Mover/Public/Backends/MoverBackendLiaison.h
+++ b/Engine/Plugins/Experimental/Mover/Source/Mover/Public/Backends/MoverBackendLiaison.h
@@ -29,8 +29,17 @@ public:
 	virtual float GetCurrentSimTimeMs() = 0;
 	virtual int32 GetCurrentSimFrame() = 0;
 
+	// Pending State: the simulation state currently being authored
 	virtual bool ReadPendingSyncState(OUT FMoverSyncState& OutSyncState) { return false; }
 	virtual bool WritePendingSyncState(const FMoverSyncState& SyncStateToWrite) { return false; }
+	
+	// Presentation State: the most recent presentation state, possibly the result of interpolation or smoothing. Writing to it does not affect the official simulation record.
+	virtual bool ReadPresentationSyncState(OUT FMoverSyncState& OutSyncState) { return false; }
+	virtual bool WritePresentationSyncState(const FMoverSyncState& SyncStateToWrite) { return false; }
+
+	// Previous Presentation State: the state that our optional smoothing process is moving away from, towards a more recent state. Writing to it does not affect the official simulation record.
+	virtual bool ReadPrevPresentationSyncState(OUT FMoverSyncState& OutSyncState) { return false; }
+	virtual bool WritePrevPresentationSyncState(const FMoverSyncState& SyncStateToWrite) { return false; }
 
 #if WITH_EDITOR
 	virtual EDataValidationResult ValidateData(FDataValidationContext& Context, const UMoverComponent& ValidationMoverComp) const { return EDataValidationResult::Valid; }
diff --git a/Engine/Plugins/Experimental/Mover/Source/Mover/Public/Backends/MoverNetworkPredictionLiaison.h b/Engine/Plugins/Experimental/Mover/Source/Mover/Public/Backends/MoverNetworkPredictionLiaison.h
--- a/Engine/Plugins/Experimental/Mover/Source/Mover/Public/Backends/MoverNetworkPredictionLiaison.h
+++ b/Engine/Plugins/Experimental/Mover/Source/Mover/Public/Backends/MoverNetworkPredictionLiaison.h
@@ -37,6 +37,9 @@ public:
 	// Take output for simulation. Called by Network Prediction system.
 	void FinalizeFrame(const FMoverSyncState* SyncState, const FMoverAuxStateContext* AuxState);
 
+	// Take output for smoothing. Called by Network Prediction system.
+	void FinalizeSmoothingFrame(const FMoverSyncState* SyncState, const FMoverAuxStateContext* AuxState);
+
 	// Seed initial values based on component's state. Called by Network Prediction system.
 	void InitializeSimulationState(FMoverSyncState* OutSync, FMoverAuxStateContext* OutAux);
 
@@ -49,6 +52,10 @@ public:
 	virtual int32 GetCurrentSimFrame() override;
 	virtual bool ReadPendingSyncState(OUT FMoverSyncState& OutSyncState) override;
 	virtual bool WritePendingSyncState(const FMoverSyncState& SyncStateToWrite) override;
+	virtual bool ReadPresentationSyncState(OUT FMoverSyncState& OutSyncState) override;
+	virtual bool WritePresentationSyncState(const FMoverSyncState& SyncStateToWrite) override;
+	virtual bool ReadPrevPresentationSyncState(FMoverSyncState& OutSyncState) override;
+	virtual bool WritePrevPresentationSyncState(const FMoverSyncState& SyncStateToWrite) override;
 #if WITH_EDITOR
 	virtual EDataValidationResult ValidateData(FDataValidationContext& Context, const UMoverComponent& ValidationMoverComp) const override;
 #endif
@@ -74,4 +81,4 @@ protected:
 	TObjectPtr<UMoverComponent> MoverComp;	// the component that we're in charge of driving
 	FMoverSyncState* StartingOutSync;
 	FMoverAuxStateContext* StartingOutAux;
-};
\ No newline at end of file
+};
diff --git a/Engine/Plugins/Experimental/Mover/Source/Mover/Public/MoverComponent.h b/Engine/Plugins/Experimental/Mover/Source/Mover/Public/MoverComponent.h
--- a/Engine/Plugins/Experimental/Mover/Source/Mover/Public/MoverComponent.h
+++ b/Engine/Plugins/Experimental/Mover/Source/Mover/Public/MoverComponent.h
@@ -111,22 +111,25 @@ public:
 	// NP Driver
 	// --------------------------------------------------------------------------------
 
-	// Get latest local input prior to simulation step. Called by Network Prediction system on owner's instance (autonomous or authority).
+	// Get latest local input prior to simulation step. Called by backend system on owner's instance (autonomous or authority).
 	void ProduceInput(const int32 DeltaTimeMS, FMoverInputCmdContext* Cmd);
 
-	// Restore a previous frame prior to resimulating. Called by Network Prediction system.
+	// Restore a previous frame prior to resimulating. Called by backend system.
 	void RestoreFrame(const FMoverSyncState* SyncState, const FMoverAuxStateContext* AuxState);
 
-	// Take output for simulation. Called by Network Prediction system.
+	// Take output for simulation. Called by backend system.
 	void FinalizeFrame(const FMoverSyncState* SyncState, const FMoverAuxStateContext* AuxState);
 
+	// Take smoothed simulation state. Called by backend system, if supported.
+	void FinalizeSmoothingFrame(const FMoverSyncState* SyncState, const FMoverAuxStateContext* AuxState);
+
 	// This is an opportunity to run code on the code on the simproxy in interpolated mode - currently used to help activate and deactivate modifiers on the simproxy in interpolated mode
 	void TickInterpolatedSimProxy(const FMoverTimeStep& TimeStep, const FMoverInputCmdContext& InputCmd, UMoverComponent* MoverComp, const FMoverSyncState& CachedSyncState, const FMoverSyncState& SyncState, const FMoverAuxStateContext& AuxState);
 	
-	// Seed initial values based on component's state. Called by Network Prediction system.
+	// Seed initial values based on component's state. Called by backend system.
 	void InitializeSimulationState(FMoverSyncState* OutSync, FMoverAuxStateContext* OutAux);
 
-	// Primary movement simulation update. Given an starting state and timestep, produce a new state. Called by Network Prediction system.
+	// Primary movement simulation update. Given an starting state and timestep, produce a new state. Called by backend system.
 	void SimulationTick(const FMoverTimeStep& InTimeStep, const FMoverTickStartData& SimInput, OUT FMoverTickEndData& SimOutput);
 
 	// Specifies which supporting back end class should drive this Mover actor
@@ -232,6 +235,10 @@ public:
 	// Sets planar constraint that can limit movement direction
 	UFUNCTION(BlueprintCallable, Category = Mover)
 	void SetPlanarConstraint(const FPlanarConstraint& InConstraint);
+	
+	// If enabled, the movement of the primary visual component will be smoothed via an offset from the root moving component. This is useful in fixed-tick simulations with variable rendering rates.
+	UPROPERTY(EditAnywhere, BlueprintReadWrite, Category = "Mover")
+	EMoverSmoothingMode SmoothingMode = EMoverSmoothingMode::VisualComponentOffset;
 
 public:
 
@@ -502,6 +509,9 @@ protected:
 	UPROPERTY(Transient)
 	TObjectPtr<USceneComponent> PrimaryVisualComponent;
 
+	/** Cached original offset from the visual component, used for cases where we want to move the visual component away from the root component (for smoothing, corrections, etc.) */
+	FTransform BaseVisualComponentTransform = FTransform::Identity;
+
 	bool bHasValidLastProducedInput = false;
 	FMoverInputCmdContext CachedLastProducedInputCmd;
 
diff --git a/Engine/Plugins/Experimental/Mover/Source/Mover/Public/MoverDataModelTypes.h b/Engine/Plugins/Experimental/Mover/Source/Mover/Public/MoverDataModelTypes.h
--- a/Engine/Plugins/Experimental/Mover/Source/Mover/Public/MoverDataModelTypes.h
+++ b/Engine/Plugins/Experimental/Mover/Source/Mover/Public/MoverDataModelTypes.h
@@ -214,6 +214,8 @@ public:
 	FRotator GetOrientation_WorldSpace() const;
 	FRotator GetOrientation_BaseSpace() const;
 
+	FTransform GetTransform_WorldSpace() const;
+	FTransform GetTransform_BaseSpace() const;
 };
 
 template<>
diff --git a/Engine/Plugins/Experimental/Mover/Source/Mover/Public/MoverTypes.h b/Engine/Plugins/Experimental/Mover/Source/Mover/Public/MoverTypes.h
--- a/Engine/Plugins/Experimental/Mover/Source/Mover/Public/MoverTypes.h
+++ b/Engine/Plugins/Experimental/Mover/Source/Mover/Public/MoverTypes.h
@@ -18,6 +18,19 @@ MOVER_API UE_DECLARE_GAMEPLAY_TAG_EXTERN(Mover_IsFlying);
 MOVER_API UE_DECLARE_GAMEPLAY_TAG_EXTERN(Mover_IsSwimming);
 MOVER_API UE_DECLARE_GAMEPLAY_TAG_EXTERN(Mover_IsCrouching);
 
+
+/** Options for how to handle smoothing frame data from the backend. Typically this is for advancing the simulation at a lower or fixed rate versus the game thread/render rate. */
+UENUM(BlueprintType)
+enum class EMoverSmoothingMode : uint8
+{
+	/** Smoothed frames will be ignored */
+	None,
+
+	/** Use the smoothed state data to offset the visual root component only, without smoothing the root moving component or any other state data */
+	VisualComponentOffset,
+};
+
+
 // Struct to hold params for when an impact happens. This contains all of the data for impacts including what gets passed to the FMover_OnImpact delegate
 USTRUCT(BlueprintType, meta = (DisplayName = "Impact Data"))
 struct MOVER_API FMoverOnImpactParams
diff --git a/Engine/Plugins/Runtime/NetworkPrediction/Source/NetworkPrediction/Private/NetworkPredictionWorldManager.cpp b/Engine/Plugins/Runtime/NetworkPrediction/Source/NetworkPrediction/Private/NetworkPredictionWorldManager.cpp
--- a/Engine/Plugins/Runtime/NetworkPrediction/Source/NetworkPrediction/Private/NetworkPredictionWorldManager.cpp
+++ b/Engine/Plugins/Runtime/NetworkPrediction/Source/NetworkPrediction/Private/NetworkPredictionWorldManager.cpp
@@ -347,6 +347,14 @@ void UNetworkPredictionWorldManager::BeginNewSimulationFrame_Internal(float Delt
 				Ptr->Tick(Step, ServiceStep);
 			}
 			
+			if (Settings.bEnableFixedTickSmoothing)
+			{
+				for (TUniquePtr<IFixedSmoothingService>& Ptr : Services.FixedSmoothing.Array)
+				{
+					Ptr->UpdateSmoothing(ServiceStep, &FixedTickState);
+				}
+			}
+
 			if (bSingleTick)
 			{
 				FixedTickState.UnspentTimeMS = 0.f;
@@ -557,6 +565,14 @@ void UNetworkPredictionWorldManager::BeginNewSimulationFrame_Internal(float Delt
 			Ptr->FinalizeFrame(DeltaTimeSeconds, FixedServerFrame, FixedTotalSimTimeMS, FixedTickState.FixedStepMS);
 		}
 
+		if (Settings.bEnableFixedTickSmoothing)
+		{
+			for (TUniquePtr<IFixedSmoothingService>& Ptr : Services.FixedSmoothing.Array)
+			{
+				Ptr->FinalizeSmoothingFrame(&FixedTickState);
+			}
+		}
+
 		const int32 IndependentTotalSimTimeMS = VariableTickState.Frames[VariableTickState.PendingFrame].TotalMS;
 		const int32 IndependentFrame = VariableTickState.PendingFrame;
 		for (TUniquePtr<IFinalizeService>& Ptr : Services.IndependentLocalFinalize.Array)
diff --git a/Engine/Plugins/Runtime/NetworkPrediction/Source/NetworkPrediction/Public/NetworkPredictionDriver.h b/Engine/Plugins/Runtime/NetworkPrediction/Source/NetworkPrediction/Public/NetworkPredictionDriver.h
--- a/Engine/Plugins/Runtime/NetworkPrediction/Source/NetworkPrediction/Public/NetworkPredictionDriver.h
+++ b/Engine/Plugins/Runtime/NetworkPrediction/Source/NetworkPrediction/Public/NetworkPredictionDriver.h
@@ -355,6 +355,38 @@ struct FNetworkPredictionDriverBase
 		npCheckf(!HasNpState(), TEXT("No FinalizeFrame implementation found. Implement DriverType::FinalizeFrame or ModelDef::FinalizeFrame"));
 	}
 
+	// -----------------------------------------------------------------------------------------------------------------------------------
+	//	FinalizeSmoothingFrame
+	//
+	//	When a smoothing service is active, this is called every engine frame to push the final smoothed state to the driver. This is optional.
+	// -----------------------------------------------------------------------------------------------------------------------------------
+	static void FinalizeSmoothingFrame(DriverType* Driver, const SyncType* SyncState, const AuxType* AuxState)
+	{
+		CallFinalizeSmoothingFrameMemberFunc(Driver, SyncState, AuxState);
+	}
+	
+	struct CFinalizeSmoothingFrameMemberFuncable
+	{
+		template <typename InDriverType, typename...>
+		auto Requires(InDriverType* Driver, const SyncType* S, const AuxType* A) -> decltype(Driver->FinalizeSmoothingFrame(S, A));
+	};
+
+	static constexpr bool HasFinalizeSmoothingFrame = TModels_V<CFinalizeSmoothingFrameMemberFuncable, DriverType, SyncType, AuxType>;
+
+	template<bool HasFunc=HasFinalizeSmoothingFrame>
+	static typename TEnableIf<HasFunc>::Type CallFinalizeSmoothingFrameMemberFunc(DriverType* Driver, const SyncType* SyncState, const AuxType* AuxState)
+	{
+		npCheckSlow(Driver);
+		Driver->FinalizeSmoothingFrame(SyncState, AuxState);
+	}
+
+	template<bool HasFunc=HasFinalizeSmoothingFrame>
+	static typename TEnableIf<!HasFunc>::Type CallFinalizeSmoothingFrameMemberFunc(DriverType* Driver, const SyncType* SyncState, const AuxType* AuxState)
+	{
+		// FinalizeSmoothingFrame isn't required, but the driver/model won't get smoothed state
+	}
+	
+
 	// -----------------------------------------------------------------------------------------------------------------------------------
 	//	RestoreFrame
 	//
diff --git a/Engine/Plugins/Runtime/NetworkPrediction/Source/NetworkPrediction/Public/NetworkPredictionProxy.h b/Engine/Plugins/Runtime/NetworkPrediction/Source/NetworkPrediction/Public/NetworkPredictionProxy.h
--- a/Engine/Plugins/Runtime/NetworkPrediction/Source/NetworkPrediction/Public/NetworkPredictionProxy.h
+++ b/Engine/Plugins/Runtime/NetworkPrediction/Source/NetworkPrediction/Public/NetworkPredictionProxy.h
@@ -115,6 +115,19 @@ struct FNetworkPredictionProxy
 		return static_cast<TSyncState*>(View.PendingSyncState);
 	}
 
+	// Returns Prev Presentation SyncState by default, if it is set. Otherwise returns pending Simulation value.
+	template<typename TSyncState>
+	const TSyncState* ReadPrevPresentationSyncState() const
+	{
+		
+		if (View.PrevPresentationSyncState)
+		{
+			return static_cast<TSyncState*>(View.PrevPresentationSyncState);
+		}
+		
+		return static_cast<TSyncState*>(View.PendingSyncState);
+	}
+
 	// Returns Presentation AuxState by default, if it is set. Otherwise returns pending Simulation value.
 	template<typename TAuxState>
 	const TAuxState* ReadAuxState(ENetworkPredictionStateRead ReadType = ENetworkPredictionStateRead::Presentation) const
@@ -123,7 +136,19 @@ struct FNetworkPredictionProxy
 		{
 			return static_cast<TAuxState*>(View.PresentationAuxState);
 		}
+		
+		return static_cast<TAuxState*>(View.PendingAuxState);
+	}
 
+	// Returns Prev Presentation AuxState by default, if it is set. Otherwise returns pending Simulation value.
+	template<typename TAuxState>
+	const TAuxState* ReadPrevPresentationAuxState() const
+	{
+		if (View.PrevPresentationAuxState)
+		{
+			return static_cast<TAuxState*>(View.PrevPresentationAuxState);
+		}
+		
 		return static_cast<TAuxState*>(View.PendingAuxState);
 	}
 
@@ -134,10 +159,23 @@ struct FNetworkPredictionProxy
 
 	template<typename TSyncState>
 	const TSyncState* WriteSyncState(TFunctionRef<void(TSyncState&)> WriteFunc, const FAnsiStringView& TraceMsg=FAnsiStringView());
+	
+	template<typename TSyncState>
+	const TSyncState* WritePresentationSyncState(TFunctionRef<void(TSyncState&)> WriteFunc, const FAnsiStringView& TraceMsg=FAnsiStringView());
+	
+	template<typename TSyncState>
+	const TSyncState* WritePrevPresentationSyncState(TFunctionRef<void(TSyncState&)> WriteFunc, const FAnsiStringView& TraceMsg=FAnsiStringView());
 
 	template<typename TAuxState>
 	const TAuxState* WriteAuxState(TFunctionRef<void(TAuxState&)> WriteFunc, const FAnsiStringView& TraceMsg=FAnsiStringView());
 
+	template<typename TAuxState>
+	const TAuxState* WritePresentationAuxState(TFunctionRef<void(TAuxState&)> WriteFunc, const FAnsiStringView& TraceMsg=FAnsiStringView());
+
+	template<typename TAuxState>
+	const TAuxState* WritePrevPresentationAuxState(TFunctionRef<void(TAuxState&)> WriteFunc, const FAnsiStringView& TraceMsg=FAnsiStringView());
+	
+
 	// ------------------------------------------------------------------------------------
 
 	FNetSimCueDispatcher* GetCueDispatcher() const
diff --git a/Engine/Plugins/Runtime/NetworkPrediction/Source/NetworkPrediction/Public/NetworkPredictionProxyWrite.h b/Engine/Plugins/Runtime/NetworkPrediction/Source/NetworkPrediction/Public/NetworkPredictionProxyWrite.h
--- a/Engine/Plugins/Runtime/NetworkPrediction/Source/NetworkPrediction/Public/NetworkPredictionProxyWrite.h
+++ b/Engine/Plugins/Runtime/NetworkPrediction/Source/NetworkPrediction/Public/NetworkPredictionProxyWrite.h
@@ -45,6 +45,32 @@ const TSyncState* FNetworkPredictionProxy::WriteSyncState(TFunctionRef<void(TSyn
 	return nullptr;
 }
 
+template<typename TSyncState>
+const TSyncState* FNetworkPredictionProxy::WritePresentationSyncState(TFunctionRef<void(TSyncState&)> WriteFunc, const FAnsiStringView& TraceMsg)
+{
+	if (TSyncState* SyncState = static_cast<TSyncState*>(View.PresentationSyncState))
+	{
+		WriteFunc(*SyncState);
+		UE_NP_TRACE_OOB_STATE_MOD(ID.GetTraceID(), View.PendingFrame, TraceMsg);
+		ConfigFunc(this, FNetworkPredictionID(), EConfigAction::TraceSync);
+		return SyncState;
+	}
+	return nullptr;
+}
+
+template<typename TSyncState>
+const TSyncState* FNetworkPredictionProxy::WritePrevPresentationSyncState(TFunctionRef<void(TSyncState&)> WriteFunc, const FAnsiStringView& TraceMsg)
+{
+	if (TSyncState* SyncState = static_cast<TSyncState*>(View.PrevPresentationSyncState))
+	{
+		WriteFunc(*SyncState);
+		UE_NP_TRACE_OOB_STATE_MOD(ID.GetTraceID(), View.PendingFrame, TraceMsg);
+		ConfigFunc(this, FNetworkPredictionID(), EConfigAction::TraceSync);
+		return SyncState;
+	}
+	return nullptr;
+}
+
 template<typename TAuxState>
 const TAuxState* FNetworkPredictionProxy::WriteAuxState(TFunctionRef<void(TAuxState&)> WriteFunc, const FAnsiStringView& TraceMsg)
 {
@@ -56,4 +82,30 @@ const TAuxState* FNetworkPredictionProxy::WriteAuxState(TFunctionRef<void(TAuxSt
 		return AuxState;
 	}
 	return nullptr;
+}
+
+template<typename TAuxState>
+const TAuxState* FNetworkPredictionProxy::WritePresentationAuxState(TFunctionRef<void(TAuxState&)> WriteFunc, const FAnsiStringView& TraceMsg)
+{
+	if (TAuxState* AuxState = static_cast<TAuxState*>(View.PresentationAuxState))
+	{
+		WriteFunc(*AuxState);
+		UE_NP_TRACE_OOB_STATE_MOD(ID.GetTraceID(), View.PendingFrame, TraceMsg);
+		ConfigFunc(this, FNetworkPredictionID(), EConfigAction::TraceAux);
+		return AuxState;
+	}
+	return nullptr;
+}
+
+template<typename TAuxState>
+const TAuxState* FNetworkPredictionProxy::WritePrevPresentationAuxState(TFunctionRef<void(TAuxState&)> WriteFunc, const FAnsiStringView& TraceMsg)
+{
+	if (TAuxState* AuxState = static_cast<TAuxState*>(View.PrevPresentationAuxState))
+	{
+		WriteFunc(*AuxState);
+		UE_NP_TRACE_OOB_STATE_MOD(ID.GetTraceID(), View.PendingFrame, TraceMsg);
+		ConfigFunc(this, FNetworkPredictionID(), EConfigAction::TraceAux);
+		return AuxState;
+	}
+	return nullptr;
 }
\ No newline at end of file
diff --git a/Engine/Plugins/Runtime/NetworkPrediction/Source/NetworkPrediction/Public/NetworkPredictionSettings.h b/Engine/Plugins/Runtime/NetworkPrediction/Source/NetworkPrediction/Public/NetworkPredictionSettings.h
--- a/Engine/Plugins/Runtime/NetworkPrediction/Source/NetworkPrediction/Public/NetworkPredictionSettings.h
+++ b/Engine/Plugins/Runtime/NetworkPrediction/Source/NetworkPrediction/Public/NetworkPredictionSettings.h
@@ -38,6 +38,10 @@ struct FNetworkPredictionSettings
 	UPROPERTY(config, EditAnywhere, Category = FixedTick)
 	ENetworkLOD SimulatedProxyNetworkLOD = ENetworkLOD::ForwardPredict;
 
+	// If enabled, the fixed tick smoothing service will be active, allowing drivers to receive smoothly interpolated updates for any locally-simulated objects (including forward-predicted sim proxies).
+	UPROPERTY(config, EditAnywhere, Category = FixedTick)
+	bool bEnableFixedTickSmoothing = false;
+
 	// ------------------------------------------------------------------------------------------
 
 	// How much buffered time to keep for fixed ticking interpolated sims (client only).
diff --git a/Engine/Plugins/Runtime/NetworkPrediction/Source/NetworkPrediction/Public/NetworkPredictionStateView.h b/Engine/Plugins/Runtime/NetworkPrediction/Source/NetworkPrediction/Public/NetworkPredictionStateView.h
--- a/Engine/Plugins/Runtime/NetworkPrediction/Source/NetworkPrediction/Public/NetworkPredictionStateView.h
+++ b/Engine/Plugins/Runtime/NetworkPrediction/Source/NetworkPrediction/Public/NetworkPredictionStateView.h
@@ -34,6 +34,11 @@ struct FNetworkPredictionStateView
 	void* PresentationSyncState = nullptr;
 	void* PresentationAuxState = nullptr;
 
+	// Previous presentation states: these may be used by a smoothing service to represent state that we're smoothing away from.
+	// These will not be fed back into the simulation, and they will be null in cases where there is no smoothing.
+	void* PrevPresentationSyncState = nullptr;
+	void* PrevPresentationAuxState = nullptr;
+
 	// CueDispatcher is exposed so that game code can invoke OOB cues
 	// Future versions may move this / make Invoking OOB cues go through a different API
 	FNetSimCueDispatcher* CueDispatcher = nullptr;
@@ -54,9 +59,17 @@ struct FNetworkPredictionStateView
 		PresentationAuxState = Aux;
 	}
 
+	void UpdatePrevPresentationView(void* Sync, void* Aux)
+	{
+		PrevPresentationSyncState = Sync;
+		PrevPresentationAuxState = Aux;
+	}
+	
 	void ClearPresentationView()
 	{
 		PresentationSyncState = nullptr;
 		PresentationAuxState = nullptr;
+		PrevPresentationSyncState = nullptr;
+		PrevPresentationAuxState = nullptr;
 	}
 };
\ No newline at end of file
diff --git a/Engine/Plugins/Runtime/NetworkPrediction/Source/NetworkPrediction/Public/NetworkPredictionWorldManager.h b/Engine/Plugins/Runtime/NetworkPrediction/Source/NetworkPrediction/Public/NetworkPredictionWorldManager.h
--- a/Engine/Plugins/Runtime/NetworkPrediction/Source/NetworkPrediction/Public/NetworkPredictionWorldManager.h
+++ b/Engine/Plugins/Runtime/NetworkPrediction/Source/NetworkPrediction/Public/NetworkPredictionWorldManager.h
@@ -365,6 +365,12 @@ void UNetworkPredictionWorldManager::ConfigureInstance(FNetworkPredictionID ID,
 			{
 				ServiceMask |= ENetworkPredictionService::FixedTick;
 				ServiceMask |= ENetworkPredictionService::FixedFinalize;
+
+				if (FNetworkPredictionDriver<ModelDef>::HasFinalizeSmoothingFrame && Settings.bEnableFixedTickSmoothing)
+				{
+					ServiceMask |= ENetworkPredictionService::FixedSmoothing;
+				}
+				
 			}
 		}
 		else
@@ -379,6 +385,11 @@ void UNetworkPredictionWorldManager::ConfigureInstance(FNetworkPredictionID ID,
 				{
 					ServiceMask |= ENetworkPredictionService::FixedTick;
 					ServiceMask |= ENetworkPredictionService::FixedFinalize;
+
+					if (FNetworkPredictionDriver<ModelDef>::HasFinalizeSmoothingFrame && Settings.bEnableFixedTickSmoothing)
+					{
+						ServiceMask |= ENetworkPredictionService::FixedSmoothing;
+					}
 				}
 
 				bUsePhysicsRecording = true;
diff --git a/Engine/Plugins/Runtime/NetworkPrediction/Source/NetworkPrediction/Public/Services/NetworkPredictionInstanceData.h b/Engine/Plugins/Runtime/NetworkPrediction/Source/NetworkPrediction/Public/Services/NetworkPredictionInstanceData.h
--- a/Engine/Plugins/Runtime/NetworkPrediction/Source/NetworkPrediction/Public/Services/NetworkPredictionInstanceData.h
+++ b/Engine/Plugins/Runtime/NetworkPrediction/Source/NetworkPrediction/Public/Services/NetworkPredictionInstanceData.h
@@ -24,7 +24,7 @@ enum class ENetworkPredictionService : uint32
 	FixedInputLocal				= 1 << 4,
 	FixedInputRemote			= 1 << 5,
 	FixedTick					= 1 << 6,
-	FixedSmoothing				= 1 << 7,	// TODO
+	FixedSmoothing				= 1 << 7,
 	FixedFinalize				= 1 << 8,
 	FixedPhysics				= 1 << 9,
 	MAX_FIXED					= FixedPhysics,
diff --git a/Engine/Plugins/Runtime/NetworkPrediction/Source/NetworkPrediction/Public/Services/NetworkPredictionServiceRegistry.h b/Engine/Plugins/Runtime/NetworkPrediction/Source/NetworkPrediction/Public/Services/NetworkPredictionServiceRegistry.h
--- a/Engine/Plugins/Runtime/NetworkPrediction/Source/NetworkPrediction/Public/Services/NetworkPredictionServiceRegistry.h
+++ b/Engine/Plugins/Runtime/NetworkPrediction/Source/NetworkPrediction/Public/Services/NetworkPredictionServiceRegistry.h
@@ -10,6 +10,7 @@
 #include "NetworkPredictionService_Finalize.inl"
 #include "NetworkPredictionService_ServerRPC.inl"
 #include "NetworkPredictionService_Physics.inl"
+#include "NetworkPredictionService_Smooth.inl"
 
 // Services do the actual system work on batches of registered instances. UNetworkPredictionWorldManager orchestrates them.
 //
@@ -171,6 +172,7 @@ public:
 	NP_DECLARE_SERVICE(FixedTick,					ILocalTickService);
 	NP_DECLARE_SERVICE(FixedFinalize,				IFinalizeService);
 	NP_DECLARE_SERVICE(FixedPhysics,				IPhysicsService);
+	NP_DECLARE_SERVICE(FixedSmoothing,				IFixedSmoothingService);
 
 	NP_DECLARE_SERVICE(IndependentRollback,			IIndependentRollbackService);
 	NP_DECLARE_SERVICE(IndependentInterpolate,		IIndependentInterpolateService);
@@ -197,6 +199,7 @@ private:
 			NP_DEFINE_SERVICE_CALL(FixedTick,					TLocalTickService);
 			NP_DEFINE_SERVICE_CALL(FixedFinalize,				TFinalizeService);
 			NP_DEFINE_SERVICE_CALL(FixedPhysics,				TPhysicsService);
+			NP_DEFINE_SERVICE_CALL(FixedSmoothing,				TFixedSmoothingService);
 		}
 		else if (EnumHasAnyFlags(Mask, ENetworkPredictionService::ANY_INDEPENDENT))
 		{
diff --git a/Engine/Plugins/Runtime/NetworkPrediction/Source/NetworkPrediction/Public/Services/NetworkPredictionService_Smooth.inl b/Engine/Plugins/Runtime/NetworkPrediction/Source/NetworkPrediction/Public/Services/NetworkPredictionService_Smooth.inl
new file mode 100644
--- /dev/null
+++ b/Engine/Plugins/Runtime/NetworkPrediction/Source/NetworkPrediction/Public/Services/NetworkPredictionService_Smooth.inl
@@ -0,0 +1,168 @@
+// Copyright Epic Games, Inc. All Rights Reserved.
+
+#pragma once
+
+#include "NetworkPredictionCVars.h"
+#include "NetworkPredictionTickState.h"
+#include "NetworkPredictionTrace.h"
+#include "Services/NetworkPredictionInstanceData.h"
+
+
+namespace NetworkPredictionCVars
+{
+	NETSIM_DEVCVAR_SHIPCONST_INT(DisableSmoothing, 0, "np.Smoothing.Disable", "Disables smoothing and just finalizes using the latest simulation state");
+}
+
+
+// ------------------------------------------------------------------------------
+//	FixedTick Smoothing
+//
+//  This first pass service simply performs interpolation between the most recent fixed tick states
+//  and passes the smoothed state to the driver to handle however it chooses. 
+//	
+//  Future improvements could include smoothing out corrections after a reconcile, and expanding
+//  that to smoothing for Independent ticking mode.
+// ------------------------------------------------------------------------------
+class IFixedSmoothingService
+{
+public:
+
+	virtual ~IFixedSmoothingService() = default;
+	virtual void UpdateSmoothing(const FServiceTimeStep& ServiceStep,const FFixedTickState* TickState) = 0;
+	virtual void FinalizeSmoothingFrame(const FFixedTickState* TickState) = 0;
+};
+
+template<typename InModelDef>
+class TFixedSmoothingService : public IFixedSmoothingService
+{
+public:
+
+	using ModelDef = InModelDef;
+	using StateTypes = typename ModelDef::StateTypes;
+	using SyncType = typename StateTypes::SyncType;
+	using AuxType = typename StateTypes::AuxType;
+	using SyncAuxType = TSyncAuxPair<StateTypes>;
+	
+
+	TFixedSmoothingService(TModelDataStore<ModelDef>* InDataStore)
+		: DataStore(InDataStore) { }
+
+	void RegisterInstance(FNetworkPredictionID ID)
+	{
+		const int32 InstanceDataIdx = DataStore->Instances.GetIndex(ID);
+		FSparseArrayAllocationInfo AllocInfo = Instances.InsertUninitialized(InstanceDataIdx);
+		new (AllocInfo.Pointer) FInstance(ID.GetTraceID(), InstanceDataIdx, DataStore->Frames.GetIndex(ID));
+		TInstanceData<ModelDef>& InstanceData = DataStore->Instances.GetByIndexChecked(InstanceDataIdx);
+		// Point the PresentationView to our managed state. Note this only has to be done once
+		FInstance* InternalInstance = (FInstance*)AllocInfo.Pointer;
+		InstanceData.Info.View->UpdatePresentationView(InternalInstance->SyncState, InternalInstance->AuxState);
+		InstanceData.Info.View->UpdatePrevPresentationView(InternalInstance->LastSyncState, InternalInstance->LastAuxState);
+	}
+
+	void UnregisterInstance(FNetworkPredictionID ID)
+	{
+		const int32 Idx = DataStore->Instances.GetIndex(ID);
+		TInstanceData<ModelDef>& InstanceData = DataStore->Instances.GetByIndexChecked(Idx);
+		InstanceData.Info.View->ClearPresentationView();
+		Instances.RemoveAt(Idx);
+	}
+
+	virtual void UpdateSmoothing(const FServiceTimeStep& ServiceStep, const FFixedTickState* TickState) final override
+	{
+		const int32 OutputFrame = ServiceStep.LocalOutputFrame;
+		
+		for (auto& It : Instances)
+		{
+			FInstance& Instance = It;
+			
+			TInstanceFrameState<ModelDef>& Frames = DataStore->Frames.GetByIndexChecked(Instance.FramesIdx);
+
+			typename TInstanceFrameState<ModelDef>::FFrame& OutputFrameData = Frames.Buffer[OutputFrame];
+			if (NetworkPredictionCVars::DisableSmoothing() != 0 || !Instance.bHasTwoFrames)
+			{
+				OutputFrameData.SyncState.CopyTo(Instance.SyncState);
+				OutputFrameData.AuxState.CopyTo(Instance.AuxState);
+				OutputFrameData.SyncState.CopyTo(Instance.LastSyncState);
+				OutputFrameData.AuxState.CopyTo(Instance.LastAuxState);
+				Instance.bHasTwoFrames = true;
+				continue;
+			}
+
+			// TODO: could improve this with a double-buffer that alternates, eliminating one copy
+			// Set Last Presentation States To Current
+			Instance.SyncState.CopyTo(Instance.LastSyncState);
+			Instance.AuxState.CopyTo(Instance.LastAuxState);
+
+			// Set Sync State To The Pending State
+			OutputFrameData.SyncState.CopyTo(Instance.SyncState);
+			OutputFrameData.AuxState.CopyTo(Instance.AuxState);
+		}
+	}
+	
+	virtual void FinalizeSmoothingFrame(const FFixedTickState* TickState) final override
+	{
+		const float RemainingTime = TickState->UnspentTimeMS;
+		const float TimeStep = TickState->FixedStepMS;
+		const float Alpha = FMath::Clamp(RemainingTime / TimeStep,0.f,1.f);
+		
+		if (NetworkPredictionCVars::DisableSmoothing() != 0)
+		{
+			for (auto& It : Instances)
+			{
+				const FInstance& Instance = It;
+				{
+					// Push non-smoothed results to driver
+					const TInstanceData<ModelDef>& InstanceData = DataStore->Instances.GetByIndexChecked(Instance.InstanceIdx);
+					FNetworkPredictionDriver<ModelDef>::FinalizeSmoothingFrame(InstanceData.Info.Driver, Instance.SyncState.Get(), Instance.AuxState.Get());
+				}
+			}
+
+			return;
+		}
+
+		
+		for (auto& It : Instances)
+		{
+			const FInstance& Instance = It;
+			
+			// Interpolate and push smoothed state to driver
+			{
+				TConditionalState<SyncType> SmoothedSyncState;
+				TConditionalState<AuxType> SmoothedAuxState;
+
+				FNetworkPredictionDriver<ModelDef>::Interpolate(SyncAuxType{Instance.LastSyncState.Get(), Instance.LastAuxState.Get()}, SyncAuxType{Instance.SyncState.Get(), Instance.AuxState.Get()},
+					Alpha, SmoothedSyncState, SmoothedAuxState);
+
+				// Push smoothed results to driver
+				const TInstanceData<ModelDef>& InstanceData = DataStore->Instances.GetByIndexChecked(Instance.InstanceIdx);
+				FNetworkPredictionDriver<ModelDef>::FinalizeSmoothingFrame(InstanceData.Info.Driver, SmoothedSyncState, SmoothedAuxState);
+				
+			}
+		}
+	}
+
+private:
+
+	struct FInstance
+	{
+		FInstance(int32 InTraceID, int32 InInstanceIdx, int32 InFramesIdx)
+			: TraceID(InTraceID), InstanceIdx(InInstanceIdx), FramesIdx(InFramesIdx) {}
+
+		int32 TraceID;
+		int32 InstanceIdx;
+		int32 FramesIdx;
+		uint8 bHasTwoFrames : 1 = false;
+		
+		// Latest states to smooth between. Stored here so that we can maintain FNetworkPredictionStateView to them
+		TConditionalState<SyncType> SyncState;
+		TConditionalState<AuxType> AuxState;
+		TConditionalState<SyncType> LastSyncState;
+		TConditionalState<AuxType> LastAuxState;
+		
+	};
+	
+	TSparseArray<FInstance> Instances; // Indices are shared with DataStore->ClientRecv
+	
+	TModelDataStore<ModelDef>* DataStore;
+};
+
