Network Prediction plugin: fixes bug causing ensures and constant corrections/reconciliations due to client and server disagreeing about number of inputs being shared.
Expanded from PR to fix Independent ticking mode as well.

diff --git a/Engine/Plugins/Runtime/NetworkPrediction/Source/NetworkPrediction/Public/NetworkPredictionWorldManager.h b/Engine/Plugins/Runtime/NetworkPrediction/Source/NetworkPrediction/Public/NetworkPredictionWorldManager.h
--- a/Engine/Plugins/Runtime/NetworkPrediction/Source/NetworkPrediction/Public/NetworkPredictionWorldManager.h
+++ b/Engine/Plugins/Runtime/NetworkPrediction/Source/NetworkPrediction/Public/NetworkPredictionWorldManager.h
@@ -538,6 +538,8 @@ void UNetworkPredictionWorldManager::BindClientNetRecv_Fixed(FNetworkPredictionI
 	if (!RepProxy) 
 		return;
 
+	TFixedTickReplicator_Server<ModelDef>::SetNumInputsPerSend(Settings.FixedTickInputSendCount);
+
 	const int32 ClientRecvIdx = DataStore->ClientRecv.GetIndex(ID);
 	NpResizeAndSetBit(DataStore->ClientRecvBitMask, ClientRecvIdx, false);
 
@@ -562,6 +564,8 @@ void UNetworkPredictionWorldManager::BindClientNetRecv_Independent(FNetworkPredi
 	if (!RepProxy) 
 		return;
 
+	TIndependentTickReplicator_Server<ModelDef>::SetNumInputsPerSend(Settings.IndependentTickInputSendCount);
+
 	const int32 ClientRecvIdx = DataStore->ClientRecv.GetIndex(ID);
 	NpResizeAndSetBit(DataStore->ClientRecvBitMask, ClientRecvIdx, false);
 
