Remove copy-pasted macro from FNetworkPredictionTrace::TraceSystemFault.

We can use a StringBuilder here instead and get the same outcome.

diff --git a/Engine/Plugins/Runtime/NetworkPrediction/Source/NetworkPrediction/Private/NetworkPredictionTrace.cpp b/Engine/Plugins/Runtime/NetworkPrediction/Source/NetworkPrediction/Private/NetworkPredictionTrace.cpp
--- a/Engine/Plugins/Runtime/NetworkPrediction/Source/NetworkPrediction/Private/NetworkPredictionTrace.cpp
+++ b/Engine/Plugins/Runtime/NetworkPrediction/Source/NetworkPrediction/Private/NetworkPredictionTrace.cpp
@@ -2,6 +2,7 @@
 
 
 #include "NetworkPredictionTrace.h"
+#include "Containers/StringFwd.h"
 #include "Engine/GameInstance.h"
 #include "Engine/GameInstance.h"
 #include "Engine/World.h"
@@ -312,51 +313,17 @@ void FNetworkPredictionTrace::TraceOOBStateMod(int32 TraceID, int32 Frame, const
 		<< OOBStateMod.Source(StrView.GetData(), StrView.Len());
 }
 
-
-// Copied from VarargsHelper.h
-#define GROWABLE_LOGF(SerializeFunc) \
-	int32	BufferSize	= 1024; \
-	TCHAR*	Buffer		= NULL; \
-	int32	Result		= -1; \
-	/* allocate some stack space to use on the first pass, which matches most strings */ \
-	TCHAR	StackBuffer[512]; \
-	TCHAR*	AllocatedBuffer = NULL; \
-\
-	/* first, try using the stack buffer */ \
-	Buffer = StackBuffer; \
-	GET_TYPED_VARARGS_RESULT( TCHAR, Buffer, UE_ARRAY_COUNT(StackBuffer), UE_ARRAY_COUNT(StackBuffer) - 1, Fmt, Fmt, Result ); \
-\
-	/* if that fails, then use heap allocation to make enough space */ \
-	while(Result == -1) \
-	{ \
-		FMemory::SystemFree(AllocatedBuffer); \
-		/* We need to use malloc here directly as GMalloc might not be safe. */ \
-		Buffer = AllocatedBuffer = (TCHAR*) FMemory::SystemMalloc( BufferSize * sizeof(TCHAR) ); \
-		if (Buffer == NULL) \
-		{ \
-			return; \
-		} \
-		GET_TYPED_VARARGS_RESULT( TCHAR, Buffer, BufferSize, BufferSize-1, Fmt, Fmt, Result ); \
-		BufferSize *= 2; \
-	}; \
-	Buffer[Result] = 0; \
-	; \
-\
-	SerializeFunc; \
-	/*FMemory::SystemFree(AllocatedBuffer);*/
-
-
 void FNetworkPredictionTrace::TraceSystemFault(const TCHAR* Fmt, ...)
 {
-	GROWABLE_LOGF( 
+	TStringBuilder<512> Builder;
 
-		check(Result >= 0 );
+	va_list Args;
+	va_start(Args, Fmt);
+	Builder.AppendV(Fmt, Args);
+	va_end(Args);
 
-	UE_LOG(LogNetworkPrediction, Log, TEXT("SystemFault: %s"), Buffer);
-	);
+	UE_LOG(LogNetworkPrediction, Log, TEXT("SystemFault: %s"), Builder.ToString());
 
 	UE_TRACE_LOG(NetworkPrediction, SystemFault, NetworkPredictionChannel)
-		<< SystemFault.Message(Buffer, Result);
-
-	FMemory::SystemFree(AllocatedBuffer);
+		<< SystemFault.Message(Builder.GetData(), Builder.Len());
 }
