PR #12504: Exposed a public reference to the FixedTickState and VariableTickState in the NetworkPredictionWorldManager

diff --git a/Engine/Plugins/Runtime/NetworkPrediction/Source/NetworkPrediction/Public/NetworkPredictionWorldManager.h b/Engine/Plugins/Runtime/NetworkPrediction/Source/NetworkPrediction/Public/NetworkPredictionWorldManager.h
--- a/Engine/Plugins/Runtime/NetworkPrediction/Source/NetworkPrediction/Public/NetworkPredictionWorldManager.h
+++ b/Engine/Plugins/Runtime/NetworkPrediction/Source/NetworkPrediction/Public/NetworkPredictionWorldManager.h
@@ -92,6 +92,9 @@ public:
 	void SyncNetworkPredictionSettings(const UNetworkPredictionSettingsObject* Settings);
 
 	const FNetworkPredictionSettings& GetSettings() const { return Settings; }
+
+	const FFixedTickState& GetFixedTickState() const { return FixedTickState; }
+	const FVariableTickState& GetVariableTickState() const { return VariableTickState; }
 	
 private:
 
