Network Prediction: fixed crash in smoothing service that could occur after adding more simulation objects during runtime

diff --git a/Engine/Plugins/Runtime/NetworkPrediction/Source/NetworkPrediction/Public/Services/NetworkPredictionService_Smooth.inl b/Engine/Plugins/Runtime/NetworkPrediction/Source/NetworkPrediction/Public/Services/NetworkPredictionService_Smooth.inl
--- a/Engine/Plugins/Runtime/NetworkPrediction/Source/NetworkPrediction/Public/Services/NetworkPredictionService_Smooth.inl
+++ b/Engine/Plugins/Runtime/NetworkPrediction/Source/NetworkPrediction/Public/Services/NetworkPredictionService_Smooth.inl
@@ -50,13 +50,29 @@ public:
 	void RegisterInstance(FNetworkPredictionID ID)
 	{
 		const int32 InstanceDataIdx = DataStore->Instances.GetIndex(ID);
+
+		const int32 PreviousSize = Instances.GetAllocatedSize();
+
 		FSparseArrayAllocationInfo AllocInfo = Instances.InsertUninitialized(InstanceDataIdx);
+
+		// When the array resizes, other instances may have dangling pointers to Presentation/PrevPresentation views and will need updated
+		const bool bNeedsAllViewsUpdated = PreviousSize != Instances.GetAllocatedSize();
+
 		new (AllocInfo.Pointer) FInstance(ID.GetTraceID(), InstanceDataIdx, DataStore->Frames.GetIndex(ID));
-		TInstanceData<ModelDef>& InstanceData = DataStore->Instances.GetByIndexChecked(InstanceDataIdx);
-		// Point the PresentationView to our managed state. Note this only has to be done once
-		FInstance* InternalInstance = (FInstance*)AllocInfo.Pointer;
-		InstanceData.Info.View->UpdatePresentationView(InternalInstance->SyncState, InternalInstance->AuxState);
-		InstanceData.Info.View->UpdatePrevPresentationView(InternalInstance->LastSyncState, InternalInstance->LastAuxState);
+
+		if (bNeedsAllViewsUpdated)
+		{
+			UpdateAllInstancePresentationViews();
+		}
+		else
+		{
+			TInstanceData<ModelDef>& InstanceData = DataStore->Instances.GetByIndexChecked(InstanceDataIdx);
+
+			// Point the PresentationView to our managed state. Note this only has to be done once
+			FInstance* InternalInstance = (FInstance*)AllocInfo.Pointer;
+			InstanceData.Info.View->UpdatePresentationView(InternalInstance->SyncState, InternalInstance->AuxState);
+			InstanceData.Info.View->UpdatePrevPresentationView(InternalInstance->LastSyncState, InternalInstance->LastAuxState);
+		}
 	}
 
 	void UnregisterInstance(FNetworkPredictionID ID)
@@ -142,6 +158,17 @@ public:
 	}
 
 private:
+	void UpdateAllInstancePresentationViews()
+	{
+		for (auto& It : Instances)
+		{
+			FInstance& InternalInstance = It;
+
+			TInstanceData<ModelDef>& RemapInstanceData = DataStore->Instances.GetByIndexChecked(InternalInstance.InstanceIdx);
+			RemapInstanceData.Info.View->UpdatePresentationView(InternalInstance.SyncState, InternalInstance.AuxState);
+			RemapInstanceData.Info.View->UpdatePrevPresentationView(InternalInstance.LastSyncState, InternalInstance.LastAuxState);
+		}
+	}
 
 	struct FInstance
 	{
