Network prediction changing static initialization of FNetworkPredictionModelDefRegistry's singleton to avoid issues with static initialization happening out of order causing registered Network Prediction model definitions to be skipped/removed. Also, making some minor adjustments to FNetworkPredictionModelDefRegistry::FinalizeTypes avoid potential ODR issues.

diff --git a/Engine/Plugins/Runtime/NetworkPrediction/Source/NetworkPrediction/Private/NetworkPredictionModelDefRegistry.cpp b/Engine/Plugins/Runtime/NetworkPrediction/Source/NetworkPrediction/Private/NetworkPredictionModelDefRegistry.cpp
--- a/Engine/Plugins/Runtime/NetworkPrediction/Source/NetworkPrediction/Private/NetworkPredictionModelDefRegistry.cpp
+++ b/Engine/Plugins/Runtime/NetworkPrediction/Source/NetworkPrediction/Private/NetworkPredictionModelDefRegistry.cpp
@@ -1,5 +1,37 @@
 // Copyright Epic Games, Inc. All Rights Reserved.
 #include "NetworkPredictionModelDefRegistry.h"
 
-FNetworkPredictionModelDefRegistry FNetworkPredictionModelDefRegistry::Singleton;
+FNetworkPredictionModelDefRegistry& FNetworkPredictionModelDefRegistry::Get()
+{
+	static FNetworkPredictionModelDefRegistry Singleton;
 
+	return Singleton;
+}
+
+void FNetworkPredictionModelDefRegistry::FinalizeTypes()
+{
+	if (bFinalized)
+	{
+		return;
+	}
+
+	bFinalized = true;
+	ModelDefList.Sort([](const FTypeInfo& LHS, const FTypeInfo& RHS) -> bool
+	{
+		if (LHS.SortPriority == RHS.SortPriority)
+		{
+			UE_LOG(LogNetworkPrediction, Log, TEXT("ModelDefs %s and %s have same sort priority. Using lexical sort as backup"), LHS.Name, RHS.Name);
+			int32 StrCmpResult = FCString::Strcmp(LHS.Name, RHS.Name);
+			npEnsureMsgf(StrCmpResult != 0, TEXT("Duplicate ModelDefs appear to have been registered."));
+			return StrCmpResult > 0;
+		}
+
+		return LHS.SortPriority < RHS.SortPriority; 
+	});
+
+	int32 Count = 1;
+	for (FTypeInfo& TypeInfo : ModelDefList)
+	{
+		*TypeInfo.IDPtr = Count++;
+	}
+}
diff --git a/Engine/Plugins/Runtime/NetworkPrediction/Source/NetworkPrediction/Public/NetworkPredictionModelDefRegistry.h b/Engine/Plugins/Runtime/NetworkPrediction/Source/NetworkPrediction/Public/NetworkPredictionModelDefRegistry.h
--- a/Engine/Plugins/Runtime/NetworkPrediction/Source/NetworkPrediction/Public/NetworkPredictionModelDefRegistry.h
+++ b/Engine/Plugins/Runtime/NetworkPrediction/Source/NetworkPrediction/Public/NetworkPredictionModelDefRegistry.h
@@ -11,12 +11,12 @@
 // This ID is used as indices into the various NP service arrays
 class FNetworkPredictionModelDefRegistry
 {
+	// Declaring this as private to help avoid sidestepping singleton
+	FNetworkPredictionModelDefRegistry() = default;
+	
 public:
 
-	static FNetworkPredictionModelDefRegistry& Get()
-	{
-		return Singleton;
-	}
+	NETWORKPREDICTION_API static FNetworkPredictionModelDefRegistry& Get();
 
 	template<typename ModelDef>
 	void RegisterType()
@@ -38,38 +38,9 @@ public:
 		ModelDefList.Emplace(TypeInfo);
 	}
 
-	void FinalizeTypes()
-	{
-		if (bFinalized)
-		{
-			return;
-		}
-
-		bFinalized = true;
-		ModelDefList.Sort([](const FTypeInfo& LHS, const FTypeInfo& RHS) -> bool
-		{
-			if (LHS.SortPriority == RHS.SortPriority)
-			{
-				UE_LOG(LogNetworkPrediction, Log, TEXT("ModelDefs %s and %s have same sort priority. Using lexical sort as backup"), LHS.Name, RHS.Name);
-				int32 StrCmpResult = FCString::Strcmp(LHS.Name, RHS.Name);
-				npEnsureMsgf(StrCmpResult != 0, TEXT("Duplicate ModelDefs appear to have been registered."));
-				return StrCmpResult > 0;
-			}
-
-			return LHS.SortPriority < RHS.SortPriority; 
-		});
-
-		int32 Count = 1;
-		for (FTypeInfo& TypeInfo : ModelDefList)
-		{
-			*TypeInfo.IDPtr = Count++;
-		}
-	}
+	void FinalizeTypes();
 
 private:
-
-	NETWORKPREDICTION_API static FNetworkPredictionModelDefRegistry Singleton;
-
 	struct FTypeInfo
 	{
 		bool operator==(const FModelDefId* OtherIDPtr) const
