Fix `suppress` misspellings in Engine/Plugins.

This CL tries to avoid risky changes that might cause unexpected ripple effects, so not every misspelling is addressed. Also, this CL is easier to read with whitespace-changes disabled, as a few of the changed files had an excessive amount of trailing whitespace.

`supress` will be marked as a typo: https://p4-swarm.epicgames.net/reviews/46109184

diff --git a/Engine/Plugins/Compositing/CompositeCore/Source/CompositeCore/Private/CompositeCoreSubsystem.cpp b/Engine/Plugins/Compositing/CompositeCore/Source/CompositeCore/Private/CompositeCoreSubsystem.cpp
--- a/Engine/Plugins/Compositing/CompositeCore/Source/CompositeCore/Private/CompositeCoreSubsystem.cpp
+++ b/Engine/Plugins/Compositing/CompositeCore/Source/CompositeCore/Private/CompositeCoreSubsystem.cpp
@@ -66,7 +66,7 @@ void UCompositeCoreSubsystem::Initialize(FSubsystemCollectionBase& Collection)
 void UCompositeCoreSubsystem::Deinitialize()
 {
 	CompositeCoreViewExtension.Reset();
-	
+
 	Super::Deinitialize();
 }
 
@@ -143,7 +143,7 @@ bool UCompositeCoreSubsystem::IsProjectSettingsValid()
 {
 	static const auto CVarAlphaOutput = IConsoleManager::Get().FindConsoleVariable(TEXT("r.PostProcessing.PropagateAlpha"));
 	static const auto CVarSupportPrimitiveAlphaHoldout = IConsoleManager::Get().FindConsoleVariable(TEXT("r.Deferred.SupportPrimitiveAlphaHoldout"));
-	
+
 	const bool bAlphaOutput = CVarAlphaOutput ? CVarAlphaOutput->GetBool() : false;
 	const bool bSupportPrimitiveAlphaHoldout = CVarSupportPrimitiveAlphaHoldout ? CVarSupportPrimitiveAlphaHoldout->GetBool() : false;
 
@@ -211,7 +211,7 @@ void UCompositeCoreSubsystem::PrimitiveHoldoutSettingsNotification()
 		}
 	};
 
-	// If the user has specified to supress this pop up, then just early out and exit	
+	// If the user has specified to suppress this pop up, then just early out and exit
 	if (FSuppressDialogOptions::ShouldSuppressModal())
 	{
 		return;
diff --git a/Engine/Plugins/Compositing/CompositeCore/Source/CompositeCore/Private/HoldoutCompositeComponent.cpp b/Engine/Plugins/Compositing/CompositeCore/Source/CompositeCore/Private/HoldoutCompositeComponent.cpp
--- a/Engine/Plugins/Compositing/CompositeCore/Source/CompositeCore/Private/HoldoutCompositeComponent.cpp
+++ b/Engine/Plugins/Compositing/CompositeCore/Source/CompositeCore/Private/HoldoutCompositeComponent.cpp
@@ -37,7 +37,7 @@ namespace UE
 				{
 					{
 						UPrimitiveComponent* PrimitiveComponent = Cast<UPrimitiveComponent>(InParentComponent);
-						
+
 						if (PluginSettings->IsAllowedPrimitiveClass(PrimitiveComponent))
 						{
 							OutPrimitiveComponents.Add(PrimitiveComponent);
@@ -50,7 +50,7 @@ namespace UE
 					for (USceneComponent* ParentChild : ParentChildComponents)
 					{
 						UPrimitiveComponent* PrimitiveComponent = Cast<UPrimitiveComponent>(ParentChild);
-						
+
 						if (PluginSettings->IsAllowedPrimitiveClass(PrimitiveComponent))
 						{
 							OutPrimitiveComponents.Add(PrimitiveComponent);
@@ -157,7 +157,7 @@ void UHoldoutCompositeComponent::OnAttachmentChanged()
 		{
 #if WITH_EDITOR
 			const UCompositeCorePluginSettings* PluginSettings = GetDefault<UCompositeCorePluginSettings>();
-			
+
 			if (!PluginSettings->IsAllowedComponentClass(SceneComponent))
 			{
 				/** Utility functions for notifications */
@@ -183,7 +183,7 @@ void UHoldoutCompositeComponent::OnAttachmentChanged()
 					}
 				};
 
-				// Skip if the user has specified to supress this pop up.
+				// Skip if the user has specified to suppress this pop up.
 				if (!FSuppressDialogOptions::ShouldSuppressModal())
 				{
 					FNotificationInfo Info(LOCTEXT("CompositeParentNotification",
diff --git a/Engine/Plugins/FX/Niagara/Source/Niagara/Private/NiagaraScript.cpp b/Engine/Plugins/FX/Niagara/Source/Niagara/Private/NiagaraScript.cpp
--- a/Engine/Plugins/FX/Niagara/Source/Niagara/Private/NiagaraScript.cpp
+++ b/Engine/Plugins/FX/Niagara/Source/Niagara/Private/NiagaraScript.cpp
@@ -456,7 +456,7 @@ const FVersionedNiagaraScriptData* UNiagaraScript::GetScriptData(const FGuid& Ve
 	{
 		return nullptr;
 	}
-	
+
 	// check if we even need to support different versions
 	if (!bVersioningEnabled)
     {
@@ -602,7 +602,7 @@ void UNiagaraScript::EnableVersioning()
 	}
 	Modify();
 	ensure(VersionData.Num() == 1);
-	bVersioningEnabled = true;	
+	bVersioningEnabled = true;
 	ExposedVersion = VersionData[0].Version.VersionGuid;
 
 	NiagaraAnalytics::RecordEvent(TEXT("Versioning.ScriptEnabled"), TEXT("Usage"), NiagaraScriptUsageToString(Usage));
@@ -642,7 +642,7 @@ struct FNiagaraScriptVersionDataAccessor final : FNiagaraVersionDataAccessor
 	virtual ENiagaraPythonUpdateScriptReference& GetUpdateScriptExecutionType() override { return ScriptData->UpdateScriptExecution; }
 	virtual FString& GetPythonUpdateScript() override { return ScriptData->PythonUpdateScript; }
 	virtual FFilePath& GetScriptAsset() override { return ScriptData->ScriptAsset; }
-	
+
 	FVersionedNiagaraScriptData* ScriptData;
 };
 
@@ -701,7 +701,7 @@ UNiagaraScript* UNiagaraScript::CreateCompilationCopy()
 
 		Property->CopyCompleteValue(DestinationAddr, SourceAddr);
 	}
-	
+
 	return Result;
 }
 
@@ -1022,14 +1022,14 @@ void UNiagaraScript::ComputeVMCompilationId(FNiagaraVMExecutableDataId& Id, cons
 	Id.bDisableDebugSwitches = false;
 	Id.bRequiresPersistentIDs = false;
 	Id.ScriptVersionID = IsVersioningEnabled() ? (VersionGuid.IsValid() ? VersionGuid : ExposedVersion) : FGuid();
-	
+
 	const UNiagaraSettings* NiagaraSettings = GetDefault<UNiagaraSettings>();
 	Id.AdditionalDefines.Add(TEXT("InvalidNamespaceWrite=") + FString::FormatAsNumber(static_cast<int32>(NiagaraSettings->InvalidNamespaceWriteSeverity)));
 	if (NiagaraSettings->bAccurateQuatInterpolation)
 	{
 		Id.AdditionalDefines.Add(FNiagaraCompileOptions::AccurateQuatInterpolation);
 	}
-	
+
 	TArray<FNiagaraVariableBase> ReferencedStaticVars;
 	TArray<const uint8*> ReferencedStaticVarValues;
 	ENiagaraSimTarget SimTargetToBuild = ENiagaraSimTarget::CPUSim;
@@ -1175,7 +1175,7 @@ void UNiagaraScript::ComputeVMCompilationId(FNiagaraVMExecutableDataId& Id, cons
 		// Gather up the scripts we will use to iterate for static rapid iteration vars later...
 		EmitterData->GetScripts(Scripts, false, true);
 		bool bUsesInterpolatedSpawning = EmitterData->UsesInterpolatedSpawning();
-		
+
 		if ((bUsesInterpolatedSpawning && Usage == ENiagaraScriptUsage::ParticleGPUComputeScript) ||
 			(bUsesInterpolatedSpawning && Usage == ENiagaraScriptUsage::ParticleSpawnScript) ||
 			Usage == ENiagaraScriptUsage::ParticleSpawnScriptInterpolated)
@@ -1312,7 +1312,7 @@ void UNiagaraScript::ComputeVMCompilationId(FNiagaraVMExecutableDataId& Id, cons
 		UNiagaraScript* Script = Scripts[ScriptIdx];
 		const bool bIsThisEmitterUsage = UNiagaraScript::IsEmitterScript(Usage);
 		const bool bIsThisParticleUsage = UNiagaraScript::IsParticleScript(Usage);
-	
+
 		const bool bIsOtherSystemUsage = UNiagaraScript::IsSystemScript(Script->Usage);
 		const bool bIsOtherParticleUsage = UNiagaraScript::IsParticleScript(Script->Usage);
 		const bool bIsOtherEmitterUsage = UNiagaraScript::IsEmitterScript(Script->Usage);
@@ -1320,7 +1320,7 @@ void UNiagaraScript::ComputeVMCompilationId(FNiagaraVMExecutableDataId& Id, cons
 		// Emitter scripts don't depend on static variables from particle scripts.
 		if (bIsThisEmitterUsage && bIsOtherParticleUsage)
 			continue;
-		
+
 		// We assume that latest source is ok because these are all top-level scripts in either systems or emitters.
 		if (Script->GetLatestSource())
 		{
@@ -1337,7 +1337,7 @@ void UNiagaraScript::ComputeVMCompilationId(FNiagaraVMExecutableDataId& Id, cons
 					TArray<FStringView> SplitName;
 					UE::String::ParseTokens(VarNameBuilder.ToView(), TEXT("."), SplitName);
 
-					// Only include System based rapid iteration static variables if we're an emitter/particle script as system scripts can contain the amalgam of 
+					// Only include System based rapid iteration static variables if we're an emitter/particle script as system scripts can contain the amalgam of
 					// it's own and any child emitter script rapid iteration variables.
 					if (SplitName.Num() == 3)
 					{
@@ -1360,7 +1360,7 @@ void UNiagaraScript::ComputeVMCompilationId(FNiagaraVMExecutableDataId& Id, cons
 
 	// If we aren't using rapid iteration parameters, we need to bake them into the hashstate for the compile id. This
 	// makes their values part of the lookup.
-	
+
 	{
 		FSHA1 HashState;
 		auto UpdateHashState = [&](FName InName) -> void
@@ -1775,8 +1775,8 @@ void UNiagaraScript::Serialize(FArchive& Ar)
 	if (Ar.IsCooking())
 	{
 		// Cook optimization can temporarily modify the rapid iteration parameters for the serialize call.  We don't want this to broadcast
-		// on changed messages because it will modify the change id on owning objects so we supress those messages starting here.
-		FNiagaraParameterStore::FScopedSuppressOnChanged SupressOnChanged(RapidIterationParameters);
+		// on changed messages because it will modify the change id on owning objects so we suppress those messages starting here.
+		FNiagaraParameterStore::FScopedSuppressOnChanged SuppressOnChanged(RapidIterationParameters);
 
 		UNiagaraSystem* System = GetTypedOuter<UNiagaraSystem>();
 		bool bUsesRapidIterationParams = System != nullptr && System->ShouldUseRapidIterationParameters();
@@ -1974,7 +1974,7 @@ void FNiagaraModuleDependency::CheckVersionCache() const
 	if (VersionDependencyCache.SourceProperty != RequiredVersion)
 	{
 		VersionDependencyCache.SourceProperty = RequiredVersion;
-		
+
 		FString PatternString(TEXT("^\\s*(\\d+)\\.(\\d+)\\s*(?:(\\+)|(-)\\s*(\\d+)\\.(\\d+))?\\s*$"));
 		FRegexPattern Pattern(PatternString);
 		FRegexMatcher Matcher(Pattern, RequiredVersion);
@@ -2097,7 +2097,7 @@ void UNiagaraScript::PostLoad()
 #if WITH_EDITORONLY_DATA
 	CheckVersionDataAvailable();
 #endif
-	
+
 	RapidIterationParameters.PostLoad(this);
 
 	if (FPlatformProperties::RequiresCookedData() || IsScriptCooked())
@@ -2106,7 +2106,7 @@ void UNiagaraScript::PostLoad()
 
 		// if our bindings aren't valid, then something has gone wrong with our cook and we need to disable this Script, which will in turn
 		// disable the owning script and system
-		
+
 		if (!HasValidParameterBindings())
 		{
 #if WITH_EDITOR && IS_MONOLITHIC
@@ -2114,7 +2114,7 @@ void UNiagaraScript::PostLoad()
 #else
 			UE_LOG(LogNiagara, Error, TEXT("Mismatch between binding between RapidIterationParamters and ScriptExecutionParameters for system %s"), *GetFullName());
 #endif
-			
+
 			CachedScriptVM.Reset();
 			return;
 		}
@@ -2203,10 +2203,10 @@ void UNiagaraScript::PostLoad()
 
 	const bool bRequestHierarchyRootMigration = bMigrateParameterDataToHierarchyRoot;
 	bMigrateParameterDataToHierarchyRoot = false;
-	
+
 	VersionedScriptAdapters.Reserve(VersionData.Num());
 	for (FVersionedNiagaraScriptData& Data : VersionData)
-	{		
+	{
 		UNiagaraScriptSourceBase* Source = Data.Source;
 		if (Source != nullptr)
 		{
@@ -2420,19 +2420,19 @@ void UNiagaraScript::PostEditChangeVersionedProperty(FPropertyChangedEvent& Prop
 	if (UNiagaraScriptSourceBase* Source = GetSource(Version))
 	{
 		FString UnsynchronizedMessage = TEXT("Script property changed");
-		if (PropertyName == GET_MEMBER_NAME_CHECKED(FVersionedNiagaraScriptData, bDeprecated) || 
+		if (PropertyName == GET_MEMBER_NAME_CHECKED(FVersionedNiagaraScriptData, bDeprecated) ||
          PropertyName == GET_MEMBER_NAME_CHECKED(FVersionedNiagaraScriptData, DeprecationMessage) ||
          PropertyName == GET_MEMBER_NAME_CHECKED(FVersionedNiagaraScriptData, DeprecationRecommendation))
 		{
 			UnsynchronizedMessage = TEXT("Deprecation changed.");
 		}
 
-		if (PropertyName == GET_MEMBER_NAME_CHECKED(FVersionedNiagaraScriptData, bExperimental) || 
+		if (PropertyName == GET_MEMBER_NAME_CHECKED(FVersionedNiagaraScriptData, bExperimental) ||
             PropertyName == GET_MEMBER_NAME_CHECKED(FVersionedNiagaraScriptData, ExperimentalMessage))
 		{
 			UnsynchronizedMessage = TEXT("Experimental changed.");
 		}
-	
+
 		if (PropertyName == GET_MEMBER_NAME_CHECKED(FVersionedNiagaraScriptData, NoteMessage))
 		{
 			UnsynchronizedMessage = TEXT("Note changed.");
@@ -2451,7 +2451,7 @@ UNiagaraScriptSourceBase* UNiagaraScript::GetSource(const FGuid& VersionGuid)
 {
 	if (VersionData.Num() == 0)
 	{
-		// this should only happen when loading old assets where our PostLoad method was not yet called 
+		// this should only happen when loading old assets where our PostLoad method was not yet called
 		return Source_DEPRECATED;
 	}
 	FVersionedNiagaraScriptData* ScriptData = GetScriptData(VersionGuid);
@@ -3144,7 +3144,7 @@ bool UNiagaraScript::ApplyRapidIterationParameters(TConstArrayView<FNiagaraVaria
 		{
 			TConstArrayView<FNiagaraVariableWithOffset> ExistingParameters = RapidIterationParameters.ReadParameterVariables();
 			TArray<FNiagaraVariableBase> ParametersToRemove;
-	
+
 			for (const FNiagaraVariableWithOffset& ExistingParameter : ExistingParameters)
 			{
 				const bool bStaleParameter = !InParameters.ContainsByPredicate([&ExistingParameter](const FNiagaraVariable& Variable) -> bool
@@ -3165,7 +3165,7 @@ bool UNiagaraScript::ApplyRapidIterationParameters(TConstArrayView<FNiagaraVaria
 			}
 		}
 
-		// try to add 
+		// try to add
 		for (const FNiagaraVariable& Parameter : InParameters)
 		{
 			constexpr bool bInitInterfaces = false;
@@ -3347,7 +3347,7 @@ void UNiagaraScript::GetAssetRegistryTags(FAssetRegistryTagsContext Context) con
 	const FVersionedNiagaraScriptData* ScriptData = GetLatestScriptData();
 
 	// Dependencies
-	TArray<FName> ProvidedDependencies = ScriptData ? ScriptData->ProvidedDependencies : ProvidedDependencies_DEPRECATED; 
+	TArray<FName> ProvidedDependencies = ScriptData ? ScriptData->ProvidedDependencies : ProvidedDependencies_DEPRECATED;
 	if (ProvidedDependencies.Num() > 0)
 	{
 		FName ProvidedDependenciesName = GET_MEMBER_NAME_CHECKED(FVersionedNiagaraScriptData, ProvidedDependencies);
@@ -3425,7 +3425,7 @@ void UNiagaraScript::GetAssetRegistryTags(FAssetRegistryTagsContext Context) con
 	FString& SuggestedTag = CustomAssetRegistryTagCache->FindOrAdd(SuggestedName);
 	SuggestedTag = FString::FromInt(bSuggested);
 	Context.AddTag(FAssetRegistryTag(SuggestedName, SuggestedTag, FAssetRegistryTag::TT_Hidden));
-	
+
 	// Add the current custom version to the tags so that tags can be fixed up in the future without having to load
 	// the whole asset.
 	const int32 NiagaraVer = GetLinkerCustomVersion(FNiagaraCustomVersion::GUID);
@@ -4518,7 +4518,7 @@ ENiagaraScriptCompileStatus UNiagaraScript::GetLastCompileStatus() const
 				}
 			}
 			bHasWarnings = true;
-		}	
+		}
 	}
 	if (CachedScriptVM.IsValid())
 	{
diff --git a/Engine/Plugins/Importers/USDImporter/Source/USDStage/Private/USDStageActor.cpp b/Engine/Plugins/Importers/USDImporter/Source/USDStage/Private/USDStageActor.cpp
--- a/Engine/Plugins/Importers/USDImporter/Source/USDStage/Private/USDStageActor.cpp
+++ b/Engine/Plugins/Importers/USDImporter/Source/USDStage/Private/USDStageActor.cpp
@@ -4504,7 +4504,7 @@ void AUsdStageActor::OnObjectPropertyChanged(UObject* ObjectBeingModified, FProp
 						if (bIsRootComponent)
 						{
 							// Maybe all the sequence has is a track directly on the actor. That's still enough to
-							// supress a root component animation in case the property is just mirrored on the actor,
+							// suppress a root component animation in case the property is just mirrored on the actor,
 							// so let's try checking for that
 							BindingsToCheck.Add(Sequence->FindBindingFromObject(OwnerActor, PinnedSequencer->GetSharedPlaybackState()));
 						}
diff --git a/Engine/Plugins/MovieScene/MovieRenderPipeline/Source/MovieRenderPipelineCore/Private/Graph/MovieGraphUtils.cpp b/Engine/Plugins/MovieScene/MovieRenderPipeline/Source/MovieRenderPipelineCore/Private/Graph/MovieGraphUtils.cpp
--- a/Engine/Plugins/MovieScene/MovieRenderPipeline/Source/MovieRenderPipelineCore/Private/Graph/MovieGraphUtils.cpp
+++ b/Engine/Plugins/MovieScene/MovieRenderPipeline/Source/MovieRenderPipelineCore/Private/Graph/MovieGraphUtils.cpp
@@ -40,12 +40,12 @@ namespace UE::MovieGraph
 	bool ResolveConnectedPinValue(UMovieGraphPin* InFromPin, const TArray<UMovieGraphPin*>& InConnectedPins, const FMovieGraphEvaluationContext& InEvaluationContext, FString& OutResolvedValue)
 	{
 		TArray<UMovieGraphPin*> ConnectionPath;
-			
+
 		// Iterate up the connection chain and find all pins which might have a value that can be resolved.
 		FMovieGraphEvaluationContext ValueConnectionContext;
 		ValueConnectionContext.PinBeingFollowed = InFromPin;
 		ValueConnectionContext.SubgraphStack = InEvaluationContext.SubgraphStack;
-		
+
 		TArray<UMovieGraphPin*> ConnectedValuePins = InConnectedPins;
 		while (!ConnectedValuePins.IsEmpty())
 		{
@@ -55,7 +55,7 @@ namespace UE::MovieGraph
 				// Can't continue following the connection chain if an invalid pin was found
 				break;
 			}
-			
+
 			if (ConnectionPath.Contains(ConnectedValuePin))
 			{
 				// Recursive connection found
@@ -80,7 +80,7 @@ namespace UE::MovieGraph
 		for (int32 Index = ConnectionPath.Num() - 1; Index >= 0; --Index)
 		{
 			const UMovieGraphPin* ConnectedPin = ConnectionPath[Index];
-				
+
 			const FString ResolvedValue = ConnectedPin->Node->GetResolvedValueForOutputPin(ConnectedPin->Properties.Label, &InEvaluationContext.UserContext);
 			if (ResolvedValue.IsEmpty())
 			{
@@ -90,7 +90,7 @@ namespace UE::MovieGraph
 			OutResolvedValue = ResolvedValue;
 			return true;
 		}
-			
+
 		return false;
 	}
 
@@ -113,7 +113,7 @@ namespace UE::MovieGraph
 			{
 				return static_cast<::Audio::FMixerDevice*>(AudioDevice);
 			}
-	
+
 			return nullptr;
 		}
 
@@ -202,7 +202,7 @@ namespace UE::MovieGraph
 				}
 			};
 
-			// If the user has specified to supress this pop up, then just early out and exit	
+			// If the user has specified to suppress this pop up, then just early out and exit
 			if (FSuppressDialogOptions::ShouldSuppressModal())
 			{
 				return;
diff --git a/Engine/Plugins/Runtime/InputDebugging/Source/InputDebuggingEditor/Private/InputDebuggingEditorModule.cpp b/Engine/Plugins/Runtime/InputDebugging/Source/InputDebuggingEditor/Private/InputDebuggingEditorModule.cpp
--- a/Engine/Plugins/Runtime/InputDebugging/Source/InputDebuggingEditor/Private/InputDebuggingEditorModule.cpp
+++ b/Engine/Plugins/Runtime/InputDebugging/Source/InputDebuggingEditor/Private/InputDebuggingEditorModule.cpp
@@ -21,7 +21,7 @@ public:
 		: FSlateStyleSet("InputDebugger")
 	{
 		SetParentStyleName(FAppStyle::GetAppStyleSetName());
-		
+
 		SetContentRoot(FPaths::EngineContentDir() / TEXT("Editor/Slate"));
 		SetCoreContentRoot(FPaths::EngineContentDir() / TEXT("Slate"));
 
@@ -41,7 +41,7 @@ void FInputDebuggingEditorModule::StartupModule()
 
 	// Listen for when input devices change
 	IPlatformInputDeviceMapper::Get().GetOnInputDeviceConnectionChange().AddRaw(this, &FInputDebuggingEditorModule::OnInputDeviceConnectionChange);
-	
+
 	// Make a new style set for an input debugger, which will register any custom icons for the types in this module
 	StyleSet = MakeShared<FInputDebuggerSlateStyle>();
 	FSlateStyleRegistry::RegisterSlateStyle(*StyleSet.Get());
@@ -56,7 +56,7 @@ void FInputDebuggingEditorModule::ShutdownModule()
 
 	// Remove any input listeners
 	IPlatformInputDeviceMapper::Get().GetOnInputDeviceConnectionChange().RemoveAll(this);
-	
+
 	// Unregister slate stylings
 	if (StyleSet.IsValid())
 	{
@@ -71,7 +71,7 @@ const FText& FInputDebuggingEditorModule::ConnectionStateToText(const EInputDevi
 	static const FText DisconnectedLabel = LOCTEXT("DisconnectedStateLabel", "has been disconnected");
 	static const FText ConnectedLabel = LOCTEXT("ConnectedStateLabel", "has been connected.");
 	static const FText UnknownLabel = LOCTEXT("UnknownStateLabel", "has an unknown state!");
-	
+
 	switch (NewConnectionState)
 	{
 		case EInputDeviceConnectionState::Invalid:
@@ -79,10 +79,10 @@ const FText& FInputDebuggingEditorModule::ConnectionStateToText(const EInputDevi
 
 		case EInputDeviceConnectionState::Disconnected:
 			return DisconnectedLabel;
-		
+
 		case EInputDeviceConnectionState::Connected:
 			return ConnectedLabel;
-		
+
 		case EInputDeviceConnectionState::Unknown:
 		default:
 			return UnknownLabel;
@@ -104,7 +104,7 @@ void FInputDebuggingEditorModule::OnInputDeviceConnectionChange(EInputDeviceConn
 
 		static ECheckBoxState GetDontAskAgainCheckBoxState()
 		{
-			// Check the config for any preferences			
+			// Check the config for any preferences
 			return ShouldSuppressModal() ? ECheckBoxState::Checked : ECheckBoxState::Unchecked;
 		}
 
@@ -114,14 +114,14 @@ void FInputDebuggingEditorModule::OnInputDeviceConnectionChange(EInputDeviceConn
 			const bool bSuppressNotification = (NewState == ECheckBoxState::Checked);
 			GConfig->SetBool(TEXT("InputDebuggingEditor"), TEXT("SuppressInputDeviceConnectionChangeNotification"), bSuppressNotification, GEditorPerProjectIni);
 		}
-	}; 
+	};
 
-	// If the user has specified to supress this pop up, then just early out and exit	
+	// If the user has specified to suppress this pop up, then just early out and exit
 	if (Local::ShouldSuppressModal())
 	{
 		return;
 	}
-	
+
 	// Send a slate notification that there was a new gamepad plugged in
 	FFormatNamedArguments MainMessageArgs;
 	MainMessageArgs.Add(TEXT("StateMessage"), ConnectionStateToText(NewConnectionState));
@@ -149,7 +149,7 @@ void FInputDebuggingEditorModule::OnInputDeviceConnectionChange(EInputDeviceConn
 	Info->CheckBoxState = TAttribute<ECheckBoxState>::Create(&Local::GetDontAskAgainCheckBoxState);
 	Info->CheckBoxStateChanged = FOnCheckStateChanged::CreateStatic(&Local::OnDontAskAgainCheckBoxStateChanged);
 	Info->CheckBoxText = NSLOCTEXT("ModalDialogs", "DefaultCheckBoxMessage", "Don't show this again");
-	
+
 	FSlateNotificationManager::Get().QueueNotification(Info);
 }
 
diff --git a/Engine/Plugins/Runtime/NetworkPrediction/Source/NetworkPrediction/Private/NetworkPredictionWorldManager.cpp b/Engine/Plugins/Runtime/NetworkPrediction/Source/NetworkPrediction/Private/NetworkPredictionWorldManager.cpp
--- a/Engine/Plugins/Runtime/NetworkPrediction/Source/NetworkPrediction/Private/NetworkPredictionWorldManager.cpp
+++ b/Engine/Plugins/Runtime/NetworkPrediction/Source/NetworkPrediction/Private/NetworkPredictionWorldManager.cpp
@@ -131,7 +131,7 @@ void UNetworkPredictionWorldManager::ReconcileSimulationsPostNetworkUpdate_Inter
 	// -------------------------------------------------------------------------
 	//	Non-rollback reconcile services
 	// -------------------------------------------------------------------------
-	
+
 	// Don't reconcile FixedTick interpolates until we've started interpolation
 	// This makes the service's implementation easier if it can rely on a known
 	// ToFrame while reconciling network updates
@@ -175,7 +175,7 @@ void UNetworkPredictionWorldManager::ReconcileSimulationsPostNetworkUpdate_Inter
 
 		if (RollbackFrame < FixedTickState.PendingFrame)
 		{
-			// Common case: rollback to previously ticked frame and resimulate			
+			// Common case: rollback to previously ticked frame and resimulate
 
 			const int32 EndFrame = FixedTickState.PendingFrame;
 			const int32 NumFrames = EndFrame - RollbackFrame;
@@ -207,7 +207,7 @@ void UNetworkPredictionWorldManager::ReconcileSimulationsPostNetworkUpdate_Inter
 				FixedTickState.PendingFrame = Frame;
 				FNetSimTimeStep Step = FixedTickState.GetNextTimeStep();
 				FServiceTimeStep ServiceStep = FixedTickState.GetNextServiceTimeStep();
-			
+
 				const int32 ServerInputFrame = Frame + FixedTickState.Offset;
 				UE_NP_TRACE_PUSH_TICK(Step.TotalSimulationTime, FixedTickState.FixedStepMS, Step.Frame);
 
@@ -226,7 +226,7 @@ void UNetworkPredictionWorldManager::ReconcileSimulationsPostNetworkUpdate_Inter
 				}
 				EnsurePhysicsGTSync(TEXT("PreResimPhysics"));
 
-				// Advance physics 
+				// Advance physics
 				if (bDoPhysics)
 				{
 					AdvancePhysicsResimFrame(PhysicsFrame);
@@ -253,7 +253,7 @@ void UNetworkPredictionWorldManager::ReconcileSimulationsPostNetworkUpdate_Inter
 			// For now, just do nothing. We are either in a really bad state of PL or are just starting up
 			// As our input frames make the round trip, we'll get some slack and be doing corrections in the above code block
 			// (Setting the correction data now most likely is still wrong and not worth the iteration time)
-			
+
 			UE_LOG(LogNetworkPrediction, Warning, TEXT("RollbackFrame %d EQUAL PendingFrame %d... Offset: %d"), RollbackFrame, FixedTickState.PendingFrame, FixedTickState.Offset);
 		}
 		else if (RollbackFrame > FixedTickState.PendingFrame)
@@ -290,7 +290,7 @@ void UNetworkPredictionWorldManager::BeginNewSimulationFrame_Internal(float Delt
 {
 	QUICK_SCOPE_CYCLE_COUNTER(STAT_NP_TICK);
 	TRACE_CPUPROFILER_EVENT_SCOPE(NetworkPrediction::Tick);
-	
+
 	ActiveInstance = this;
 	bLockServices = true;
 
@@ -305,9 +305,9 @@ void UNetworkPredictionWorldManager::BeginNewSimulationFrame_Internal(float Delt
 		TRACE_CPUPROFILER_EVENT_SCOPE(NetworkPrediction::FixedTick);
 
 		FixedTickState.UnspentTimeMS += fEngineFrameDeltaTimeMS;
-		
+
 		// If we are using physics, we can only do one NP sim tick per frame.
-		// For accumulation fixed tick to work, physics needs to support it too 
+		// For accumulation fixed tick to work, physics needs to support it too
 		// see notes in NetworkPredictiontickState.h
 		const bool bSingleTick = Physics.bUsingPhysics;
 
@@ -321,7 +321,7 @@ void UNetworkPredictionWorldManager::BeginNewSimulationFrame_Internal(float Delt
 
 			FNetSimTimeStep Step = FixedTickState.GetNextTimeStep();
 			FServiceTimeStep ServiceStep = FixedTickState.GetNextServiceTimeStep();
-			
+
 			const int32 ServerInputFrame = FixedTickState.PendingFrame + FixedTickState.Offset;
 
 			UE_NP_TRACE_PUSH_INPUT_FRAME(ServerInputFrame);
@@ -336,7 +336,7 @@ void UNetworkPredictionWorldManager::BeginNewSimulationFrame_Internal(float Delt
 			}
 
 			UE_NP_TRACE_PUSH_TICK(Step.TotalSimulationTime, FixedTickState.FixedStepMS, Step.Frame);
-			
+
 			// Should we increment PendingFrame before or after the tick?
 			// Before: sims that are spawned during Tick (of other sims) will not be ticked this frame.
 			// So we want their seed state/cached pending frame to be set to the next pending frame, not this one.
@@ -346,7 +346,7 @@ void UNetworkPredictionWorldManager::BeginNewSimulationFrame_Internal(float Delt
 			{
 				Ptr->Tick(Step, ServiceStep);
 			}
-			
+
 			if (Settings.bEnableFixedTickSmoothing)
 			{
 				for (TUniquePtr<IFixedSmoothingService>& Ptr : Services.FixedSmoothing.Array)
@@ -369,7 +369,7 @@ void UNetworkPredictionWorldManager::BeginNewSimulationFrame_Internal(float Delt
 	{
 		TRACE_CPUPROFILER_EVENT_SCOPE(NetworkPrediction::IndependentTick);
 
-		// Update VariableTickState	
+		// Update VariableTickState
 		constexpr int32 MinStepMS = 1;
 		constexpr int32 MaxStepMS = 100;
 
@@ -400,7 +400,7 @@ void UNetworkPredictionWorldManager::BeginNewSimulationFrame_Internal(float Delt
 		for (TUniquePtr<ILocalTickService>& Ptr : Services.IndependentLocalTick.Array)
 		{
 			Ptr->Tick(Step, ServiceStep);
-		}	
+		}
 
 		// -------------------------------------------------------------------------
 		//	Remote Independent Tick
@@ -414,7 +414,7 @@ void UNetworkPredictionWorldManager::BeginNewSimulationFrame_Internal(float Delt
 		const int32 EndTotalSimTimeMS = PendingFrameData.TotalMS + PendingFrameData.DeltaMS;
 		VariableTickState.Frames[++VariableTickState.PendingFrame].TotalMS = EndTotalSimTimeMS;
 	}
-	
+
 	// -------------------------------------------------------------------------
 	// Interpolation
 	// -------------------------------------------------------------------------
@@ -453,7 +453,7 @@ void UNetworkPredictionWorldManager::BeginNewSimulationFrame_Internal(float Delt
 						FixedTickState.Interpolation.PCT = 0.f;
 						FixedTickState.Interpolation.AccumulatedTimeMS = 0.f;
 
-						// We need to force a reconcile here since we supress the call until interpolation starts
+						// We need to force a reconcile here since we suppress the call until interpolation starts
 						for (TUniquePtr<IFixedInterpolateService>& Ptr : Services.FixedInterpolate.Array)
 						{
 							Ptr->Reconcile(&FixedTickState);
@@ -509,7 +509,7 @@ void UNetworkPredictionWorldManager::BeginNewSimulationFrame_Internal(float Delt
 		{
 			const int32 DesiredBufferedMS = Settings.IndependentTickInterpolationBufferedMS;
 			const int32 MaxBufferedMS = Settings.IndependentTickInterpolationMaxBufferedMS;
-		
+
 			if (VariableTickState.Interpolation.LatestRecvTimeMS > DesiredBufferedMS)
 			{
 				float InterpolationRate = 1.f;
@@ -579,7 +579,7 @@ void UNetworkPredictionWorldManager::BeginNewSimulationFrame_Internal(float Delt
 		{
 			Ptr->FinalizeFrame(DeltaTimeSeconds, IndependentFrame, IndependentTotalSimTimeMS, 0);
 		}
-	
+
 		for (TUniquePtr<IRemoteFinalizeService>& Ptr : Services.IndependentRemoteFinalize.Array)
 		{
 			Ptr->FinalizeFrame(DeltaTimeSeconds);
diff --git a/Engine/Plugins/Runtime/NetworkPrediction/Source/NetworkPrediction/Public/NetworkPredictionCues.h b/Engine/Plugins/Runtime/NetworkPrediction/Source/NetworkPrediction/Public/NetworkPredictionCues.h
--- a/Engine/Plugins/Runtime/NetworkPrediction/Source/NetworkPrediction/Public/NetworkPredictionCues.h
+++ b/Engine/Plugins/Runtime/NetworkPrediction/Source/NetworkPrediction/Public/NetworkPredictionCues.h
@@ -14,7 +14,7 @@ Cues are events that are outputted by a Networked Simulation. They have the foll
 -They should not affect the simulation. Or rather, if they do affect the simulation, it will be during the actor tick, effectively the same as any "out of band" modifications.
 -They provide automatic replication and invocation settings (traits). They will not "double play" during resimulates.
 -They are time aware. When dispatched, the receiver is given how much time has passed (relative to local "head" time) since the invocation.
--They are rollback aware. When dispatched, the receiver is given a callback that will be invoked if the cue needs to rollback (undo itself). 
+-They are rollback aware. When dispatched, the receiver is given a callback that will be invoked if the cue needs to rollback (undo itself).
 -The callback is not given in contexts when a rollback is impossible: e.g, on the authority.
 
 Notes on reliability:
@@ -111,7 +111,7 @@ struct FNetSimCueSystemParamemters
 	const int32& TimeSinceInvocation;
 
 	// Callback structure if applicable. This will be null on non-rewindable cues as well as execution contexts where rollbacks wont happen (e.g, authority).
-	FNetSimCueCallbacks* Callbacks;	
+	FNetSimCueCallbacks* Callbacks;
 };
 
 // ----------------------------------------------------------------------------------------------------------------------------------------------
@@ -132,18 +132,18 @@ public:
 	{
 		FRegisteredCueTypeInfo& RegisteredTypes = GetRegistedTypeInfo();
 		UE_LOG(LogNetworkPredictionCues, Verbose, TEXT("RegisterCueType %s @ 0x%X. (NumPending: %d) "), *TypeName, &TCue::ID, RegisteredTypes.CueTypes.Num());
-		
+
 		RegisteredTypes.bDirty = true;
 
 		// Create a TypeInfo for this TCue and store it in the list
-		FTypeInfo TypeInfo = 
+		FTypeInfo TypeInfo =
 		{
-			&TCue::ID, 
+			&TCue::ID,
 			[](){ return new TNetSimCueWrapper<TCue>();},
 			TypeName
 		};
 		RegisteredTypes.CueTypes.Emplace(TypeInfo);
-		
+
 		static_assert(THasMemberFunction_NetSerialize<TCue>::Value || !TNetSimCueTypeRequirements<TCue>::RequiresNetSerialize, "TCue must implement void NetSerialize(FArchive&)");
 		static_assert(THasNetIdenticalHelper<TCue>::Value || !TNetSimCueTypeRequirements<TCue>::RequiresNetIdentical, "TCue must implement bool NetIdentical(const TCue&) const");
 	}
@@ -260,7 +260,7 @@ private:
 
 	// The registered type info lives in its own inlined static function to avoid static variable initialization order conflicts
 	NETWORKPREDICTION_API static FRegisteredCueTypeInfo& GetRegistedTypeInfo();
-	
+
 	TMap<FNetSimCueTypeId, FTypeInfo> TypeInfoMap;
 
 	NETWORKPREDICTION_API static FGlobalCueTypeTable Singleton;
@@ -279,13 +279,13 @@ struct FSavedCue
 
 	FSavedCue(FSavedCue&&) = default;
 	FSavedCue& operator=(FSavedCue&&) = default;
-	
+
 	FSavedCue(const FNetSimCueTypeId& InId, const int32& InFrame, const int32& InTime, const bool& bInAllowRollback, const bool& bInNetConfirmed, const bool& bInResimulates, FNetSimCueWrapperBase* Cue)
 		: ID(InId), Frame(InFrame), Time(InTime), CueInstance(Cue), bAllowRollback(bInAllowRollback), bNetConfirmed(bInNetConfirmed), bResimulates(bInResimulates), ReplicationTarget(Cue->GetReplicationTarget())
 	{
-		
+
 	}
-	
+
 	void NetSerialize(FArchive& Ar, const bool bSerializeFrameNumber)
 	{
 		npCheckSlow(FGlobalCueTypeTable::IsRegisterationTypeInfoDirty() == false);
@@ -351,7 +351,7 @@ struct FSavedCue
 	bool bDispatched = false;		// Cue has been dispatched to the local handler. Never dispatch twice.
 	bool bAllowRollback = false;	// Cue supports rolling back. ie., we should pass the user valid FNetSimCueCallbacks rollback callback (note this is "this saved cue" specifically, not "this cue type". E.g, on authority, bAllowRollback is always false)
 	bool bNetConfirmed = false;		// This cue has been net confirmed, meaning we received it directly via replication or we received a replicated cue that matched this one that was locally predicted.
-	
+
 	bool bResimulates = false;				 // Whether this cue supports invocation during resimulation. Needed to set bResimulatePendingRollback
 	bool bPendingResimulateRollback = false; // Rollback is pending due to resimulation (unless CUe is matched with an Invocation during the resimulate)
 
@@ -444,19 +444,19 @@ struct FCueDispatcherTraitsBase
 };
 template<typename T> struct TCueDispatcherTraits : public FCueDispatcherTraitsBase { };
 
-// Non-templated, "networking model independent" base: this is what the pure simulation code gets to invoke cues. 
+// Non-templated, "networking model independent" base: this is what the pure simulation code gets to invoke cues.
 struct FNetSimCueDispatcher
 {
 	virtual ~FNetSimCueDispatcher() = default;
 
 	// Invoke - this is how to invoke a cue from simulation code. This will construct the CueType T emplace in the saved cue record.
-	// 
+	//
 	// Best way to call:
 	//	Invoke<FMyCue>(a, b, c); // a, b, c are constructor parameters
 	//
 	// This works too, but will cause a move (if possible) or copy
 	//	FMyCue MyCue(a,b,c);
-	//	Invoke<FMyCue>(MyCue);	
+	//	Invoke<FMyCue>(MyCue);
 
 	template<typename T, typename... ArgsType>
 	void Invoke(ArgsType&&... Args)
@@ -495,7 +495,7 @@ struct FNetSimCueDispatcher
 				if (Context.TickContext == ESimulationTickContext::Resimulate)
 				{
 					npEnsure(RollbackFrame >= 0 && RollbackFrame <= Context.Frame);
-					
+
 					// Since we haven't constructed the cue yet, we can't test for uniqueness!
 					// So, create one on the stack. If we let it through we can move it to the appropriate buffer
 					// (not as nice as the non resimulate path, but better than allocating a new FSavedCue+TNetSimCueWrapper and then removing them)
@@ -513,16 +513,16 @@ struct FNetSimCueDispatcher
 						}
 					}
 
-					
+
 					auto& SavedCue = GetBuffer(bTransient).Emplace_GetRef(T::ID, Context.Frame, Context.CurrentSimTime, bAllowRollback, bNetConfirmed, bSupportsResimulate, new TNetSimCueWrapper<T>(MoveTempIfPossible(NewCue)));
 					UE_LOG(LogNetworkPredictionCues, Log, TEXT("%s. Invoking Cue %s during resimulate. Context: %d. Transient: %d. bAllowRollback: %d. bNetConfirmed: %d. Time: %d"), *GetDebugName(), *SavedCue.GetDebugName(), int(Context.TickContext), bTransient, bAllowRollback, bNetConfirmed, Context.CurrentSimTime);
 				}
 				else
 				{
-					// Not resimulate case is simple: construct the new cue emplace in the appropriate list	
+					// Not resimulate case is simple: construct the new cue emplace in the appropriate list
 					auto& SavedCue = GetBuffer(bTransient).Emplace_GetRef(T::ID, Context.Frame, Context.CurrentSimTime, bAllowRollback, bNetConfirmed, bSupportsResimulate, new TNetSimCueWrapper<T>(Forward<ArgsType>(Args)...));
 					UE_LOG(LogNetworkPredictionCues, Log, TEXT("%s. Invoking Cue %s. Context: %d. Transient: %d. bAllowRollback: %d. bNetConfirmed: %d. Time: %d"), *GetDebugName(), *SavedCue.GetDebugName(), int(Context.TickContext), bTransient, bAllowRollback, bNetConfirmed, Context.CurrentSimTime);
-				}				
+				}
 			}
 			else
 			{
@@ -532,7 +532,7 @@ struct FNetSimCueDispatcher
 	}
 
 	virtual FString GetDebugName() const = 0;
-	
+
 protected:
 
 	bool EnsureValidContext()
@@ -540,14 +540,14 @@ protected:
 		return npEnsure(Context.CurrentSimTime > 0 && Context.TickContext != ESimulationTickContext::None);
 	}
 
-	// Sim Context: the Sim has to tell the dispatcher what its doing so that it can decide if it should supress Invocations or not
+	// Sim Context: the Sim has to tell the dispatcher what its doing so that it can decide if it should suppress Invocations or not
 	struct FContext
 	{
 		int32 Frame = INDEX_NONE;
 		int32 CurrentSimTime = 0;
 		ESimulationTickContext TickContext;
 	};
-	
+
 	TArray<FSavedCue> SavedCues;		// Cues that must be saved for some period of time, either for replication or for uniqueness testing
 	TArray<FSavedCue> TransientCues;	// Cues that are dispatched on this frame and then forgotten about
 	TArray<FSavedCue>& GetBuffer(const bool& bTransient) { return bTransient ? TransientCues : SavedCues; }
@@ -575,7 +575,7 @@ struct TNetSimCueDispatcher : public FNetSimCueDispatcher
 	void NetSendSavedCues(FArchive& Ar, ENetSimCueReplicationTarget ReplicationMask, bool bSerializeFrameNumber)
 	{
 		npCheckSlow(Ar.IsSaving());
-		
+
 		// FIXME: requires two passes to count how many elements are valid for this replication mask.
 		// We could count this as saved cues are added or possibly modify the bitstream after writing the elements (tricky and would require casting to FNetBitWriter which feels real bad)
 		FNetSimCueTypeId NumCues = 0;
@@ -596,17 +596,17 @@ struct TNetSimCueDispatcher : public FNetSimCueDispatcher
 				SavedCue.NetSerialize(Ar, bSerializeFrameNumber);
 			}
 		}
-		
+
 	}
 
 	void NetRecvSavedCues(FArchive& Ar, const bool bSerializeFrameNumber, const int32 InLastRecvFrame, const int32 InLastRecvTime)
 	{
 		npCheckSlow(Ar.IsLoading());
-		
+
 		FNetSimCueTypeId NumCues;
 		Ar << NumCues;
 
-		// This is quite inefficient right now. 
+		// This is quite inefficient right now.
 		//	-We are replicating cues in the last X seconds/frames (ReplicationWindow) redundantly
 		//	-Client has to deserialize them (+ heap allocation) and check for uniqueness (have they already processed)
 		//	-If already processed (quite common), they are thrown out.
@@ -642,7 +642,7 @@ struct TNetSimCueDispatcher : public FNetSimCueDispatcher
 				UE_LOG(LogNetworkPredictionCues, Log, TEXT("%s. Discarding replicated NetSimCue %s because it is older (%d) than confirmed TimeMS (%d). CueMask: %d. Our Mask: %d"), *GetDebugName(), *SerializedCue.GetDebugName(), SerializedCue.Frame, LastRecvFrame, int(SerializedCue.ReplicationTarget), int(RecvReplicationMask));
 				continue;
 			}
-				
+
 			// Uniqueness: have we already received/predicted it?
 			// Note: we are basically ignoring invocation time when matching right now. This could potentially be a trait of the cue if needed.
 			// This could create issues if a cue is invoked several times in quick succession, but that can be worked around with arbitrary counter parameters on the cue itself (to force NetUniqueness)
@@ -679,7 +679,7 @@ struct TNetSimCueDispatcher : public FNetSimCueDispatcher
 
 		const int32 SavedCuePruneTimeMS = CalcPrune(CurrentSimTime, LastRecvMS, TCueDispatcherTraits<ModelDef>::ReplicationWindowMS());
 		const int32 SavedCuePruneFrame = CalcPrune(SimFrame, LastRecvFrame, TCueDispatcherTraits<ModelDef>::ReplicationWindowFrames());
-		
+
 		// Prune cues prior to this idx
 		int32 SavedCuePruneIdx = INDEX_NONE;
 
@@ -724,11 +724,11 @@ struct TNetSimCueDispatcher : public FNetSimCueDispatcher
 		// ------------------------------------------------------------------------
 		// Dispatch (call ::HandleCue)
 		// ------------------------------------------------------------------------
-		
+
 		for (int32 SavedCueIdx = 0; SavedCueIdx < SavedCues.Num(); ++ SavedCueIdx)
 		{
 			FSavedCue& SavedCue = SavedCues[SavedCueIdx];
-			
+
 			// See if Cue is ready for prune
 			const bool bCueIsReadyForPrune = (SavedCue.Frame == INDEX_NONE) ?  (SavedCue.Time <= SavedCuePruneTimeMS) : (SavedCue.Frame <= SavedCuePruneFrame);
 			if (bCueIsReadyForPrune)
@@ -743,7 +743,7 @@ struct TNetSimCueDispatcher : public FNetSimCueDispatcher
 
 				if (SavedCue.Frame != INDEX_NONE)
 				{
-					// Frame based comparison (prefer if available)					
+					// Frame based comparison (prefer if available)
 					if (SavedCue.Frame <= SimFrame)
 					{
 						const int32 FramesSinceInvoke = SimFrame - SavedCue.Frame;
@@ -784,7 +784,7 @@ struct TNetSimCueDispatcher : public FNetSimCueDispatcher
 		// ------------------------------------------------------------------------
 		// Prune
 		// ------------------------------------------------------------------------
-		
+
 		// Remove Cues we know longer need to keep around
 		if (SavedCuePruneIdx >= 0)
 		{
@@ -850,7 +850,7 @@ struct TNetSimCueDispatcher : public FNetSimCueDispatcher
 private:
 
 	ENetSimCueReplicationTarget RecvReplicationMask;
-	
+
 	int32 LastRecvFrame = INDEX_NONE;
 	int32 LastRecvMS = 0;
 };
diff --git a/Engine/Plugins/VirtualProduction/SceneState/Source/SceneStateGameplayEditor/Private/SceneStateSequencerSchema.cpp b/Engine/Plugins/VirtualProduction/SceneState/Source/SceneStateGameplayEditor/Private/SceneStateSequencerSchema.cpp
--- a/Engine/Plugins/VirtualProduction/SceneState/Source/SceneStateGameplayEditor/Private/SceneStateSequencerSchema.cpp
+++ b/Engine/Plugins/VirtualProduction/SceneState/Source/SceneStateGameplayEditor/Private/SceneStateSequencerSchema.cpp
@@ -145,7 +145,7 @@ void FSequencerSchema::AddSceneStateComponents(TWeakPtr<ISequencer> InSequencerW
 
 	const FScopedTransaction Transaction(LOCTEXT("AddSceneState", "Add Scene State"));
 
-	FSelectionEventSuppressor SupressEvents = Selection->SuppressEvents();
+	FSelectionEventSuppressor SuppressEvents = Selection->SuppressEvents();
 	Selection->Outliner.Empty();
 
 	for (USceneStateObject* SceneStateObject : SceneStateObjects)
