Network Prediction Plugin: routing input cmds so they can be read for Interpolated sim proxies in Independent ticking mode

diff --git a/Engine/Plugins/Runtime/NetworkPrediction/Source/NetworkPrediction/Public/Services/NetworkPredictionService_Interpolate.inl b/Engine/Plugins/Runtime/NetworkPrediction/Source/NetworkPrediction/Public/Services/NetworkPredictionService_Interpolate.inl
--- a/Engine/Plugins/Runtime/NetworkPrediction/Source/NetworkPrediction/Public/Services/NetworkPredictionService_Interpolate.inl
+++ b/Engine/Plugins/Runtime/NetworkPrediction/Source/NetworkPrediction/Public/Services/NetworkPredictionService_Interpolate.inl
@@ -410,6 +410,7 @@ public:
 
 	using ModelDef = InModelDef;
 	using StateTypes = typename ModelDef::StateTypes;
+	using InputType = typename StateTypes::InputType;
 	using SyncType = typename StateTypes::SyncType;
 	using AuxType = typename StateTypes::AuxType;
 	using SyncAuxType = TSyncAuxPair<StateTypes>;
@@ -461,6 +462,7 @@ public:
 			const int32 LocalFrame = ++InterpolationData.LastWrittenFrame;
 			
 			typename FInstance::FFrame& RecvFrame = InterpolationData.ClientRecvFrames[LocalFrame];
+			RecvFrame.InputCmd = ClientRecvData.InputCmd;
 			RecvFrame.SyncState = ClientRecvData.SyncState;
 			RecvFrame.AuxState = ClientRecvData.AuxState;
 			RecvFrame.SimTimeMS = ClientRecvData.SimTimeMS;
@@ -549,7 +551,7 @@ public:
 				
 				FNetworkPredictionDriver<ModelDef>::FinalizeFrame(InstanceData.Info.Driver, Instance.SyncState, Instance.AuxState);
 
-				InstanceData.Info.View->UpdateView(ToFrame, InterpolationTimeMS, nullptr, ToFromData.SyncState, ToFromData.AuxState);
+				InstanceData.Info.View->UpdateView(ToFrame, InterpolationTimeMS, ToFromData.InputCmd, ToFromData.SyncState, ToFromData.AuxState);
 
 				FNetworkPredictionDriver<ModelDef>::DispatchCues(&InstanceData.CueDispatcher.Get(), InstanceData.Info.Driver, FromFrame, InterpolationTimeMS, 0);
 				
@@ -591,7 +593,7 @@ public:
 				typename FInstance::FFrame& FrameData = Instance.ClientRecvFrames[FromFrame];
 				FNetworkPredictionDriver<ModelDef>::FinalizeFrame(InstanceData.Info.Driver, FrameData.SyncState, FrameData.AuxState);
 
-				InstanceData.Info.View->UpdateView(FromFrame, ActualTimeMS, nullptr, FrameData.SyncState, FrameData.AuxState);
+				InstanceData.Info.View->UpdateView(FromFrame, ActualTimeMS, FrameData.InputCmd, FrameData.SyncState, FrameData.AuxState);
 
 				FNetworkPredictionDriver<ModelDef>::DispatchCues(&InstanceData.CueDispatcher.Get(), InstanceData.Info.Driver, FromFrame, ActualTimeMS, 0);
 
@@ -630,10 +632,11 @@ private:
 		// All other services just care about 'latest received' state interpolating would benefit
 		// from storing all received frames. But this adds complications for just this one case.
 		// TClientRecvData also stores more than we would need to here - the InputCmd and ServerFrame aren't
-		// relevant for independent interpolation. This just seems simpler.
+		// always relevant for independent interpolation. This just seems simpler.
 
 		struct FFrame
 		{
+			TConditionalState<InputType> InputCmd;
 			TConditionalState<SyncType> SyncState;
 			TConditionalState<AuxType> AuxState;
 			int32 SimTimeMS;
