Add Iris replication support to the Network Prediction Plugin
Maintains compatibility with non-Iris replication

#Synced-CL 48841082

diff --git a/Engine/Plugins/Runtime/NetworkPrediction/Source/NetworkPrediction/Private/NetworkPredictionComponent.cpp b/Engine/Plugins/Runtime/NetworkPrediction/Source/NetworkPrediction/Private/NetworkPredictionComponent.cpp
--- a/Engine/Plugins/Runtime/NetworkPrediction/Source/NetworkPrediction/Private/NetworkPredictionComponent.cpp
+++ b/Engine/Plugins/Runtime/NetworkPrediction/Source/NetworkPrediction/Private/NetworkPredictionComponent.cpp
@@ -1,6 +1,7 @@
 // Copyright Epic Games, Inc. All Rights Reserved.
 
 #include "NetworkPredictionComponent.h"
+#include "Engine/NetConnection.h"
 #include "Engine/World.h"
 #include "NetworkPredictionWorldManager.h"
 
@@ -9,6 +10,7 @@
 UNetworkPredictionComponent::UNetworkPredictionComponent()
 {
 	SetIsReplicatedByDefault(true);
+	ServerRPCParamBitWriter.SetAllowResize(true);
 }
 
 void UNetworkPredictionComponent::InitializeComponent()
@@ -39,6 +41,7 @@ void UNetworkPredictionComponent::EndPlay(const EEndPlayReason::Type Reason)
 
 void UNetworkPredictionComponent::PreReplication(IRepChangedPropertyTracker & ChangedPropertyTracker)
 {
+	// Only called on the server
 	Super::PreReplication(ChangedPropertyTracker);
 	
 	CheckOwnerRoleChange();
@@ -51,6 +54,7 @@ void UNetworkPredictionComponent::PreReplication(IRepChangedPropertyTracker & Ch
 
 void UNetworkPredictionComponent::PreNetReceive()
 {
+	// Only called on the client
 	Super::PreNetReceive();
 	CheckOwnerRoleChange();
 }
@@ -78,6 +82,10 @@ bool UNetworkPredictionComponent::CheckOwnerRoleChange()
 	
 	if (CurrentRole != NetworkPredictionProxy.GetCachedNetRole() || bHasNetConnection != NetworkPredictionProxy.GetCachedHasNetConnection())
 	{
+		ReplicationProxy_Autonomous.OwnerActor = OwnerActor;
+		ReplicationProxy_Simulated.OwnerActor = OwnerActor;
+		ReplicationProxy_Replay.OwnerActor = OwnerActor;
+
 		InitializeForNetworkRole(CurrentRole, bHasNetConnection);
 		return true;
 	}
@@ -93,16 +101,99 @@ bool UNetworkPredictionComponent::ServerReceiveClientInput_Validate(const FServe
 void UNetworkPredictionComponent::ServerReceiveClientInput_Implementation(const FServerReplicationRPCParameter& ProxyParameter)
 {
 	// The const_cast is unavoidable here because the replication system only allows by value (forces copy, bad) or by const reference. This use case is unique because we are using the RPC parameter as a temp buffer.
-	const_cast<FServerReplicationRPCParameter&>(ProxyParameter).NetSerializeToProxy(ReplicationProxy_ServerRPC);
+	FServerReplicationRPCParameter& MutableProxyParameter = const_cast<FServerReplicationRPCParameter&>(ProxyParameter);
+
+	const int32 NumIncomingBits = MutableProxyParameter.DataBits.Num();
+ 	if (NumIncomingBits > FServerReplicationRPCParameter::MaxNumBits)
+ 	{
+ 		// Protect against bad data that could cause server to allocate way too much memory.
+ 		UE_LOG(LogNetworkPrediction, Error, TEXT("ServerReceiveClientInput (%s): Dropping move due to NumBits (%d) exceeding allowable limit (%d). See FServerReplicationRPCParameter::MaxNumBits."), *GetNameSafe(GetOwner()), NumIncomingBits, FServerReplicationRPCParameter::MaxNumBits);
+ 		return;
+ 	}
+
+	// Reuse bit reader to avoid allocating memory each time.
+	ServerRPCParamBitReader.SetData((uint8*)MutableProxyParameter.DataBits.GetData(), NumIncomingBits);
+
+	AActor* OwnerActor = GetOwner();
+	UNetConnection* NetConnection = OwnerActor->GetNetConnection();
+
+	if (UIrisObjectReferencePackageMap* PackageMap = FServerReplicationRPCParameter::GetIrisPackageMapToReadReferences(NetConnection, MutableProxyParameter.PackageMapExports))
+	{
+		ServerRPCParamBitReader.PackageMap = PackageMap;	// Iris flow
+	}
+	else
+	{
+		ServerRPCParamBitReader.PackageMap = MutableProxyParameter.GetPackageMap();	// Non-Iris / fallback flow
+	}
+
+	if (ServerRPCParamBitReader.PackageMap == nullptr)
+	{
+		UE_LOG(LogNetworkPrediction, Error, TEXT("ServerReceiveClientInput (%s): Failed to find PackageMap for data serialization!"), *GetNameSafe(GetOwner()));
+		return;
+	}
+
+	// Pass the raw data bits to the proxy parameter for serialized loading
+	bool bOutSuccess = true;
+	ReplicationProxy_ServerRPC.NetSerialize(ServerRPCParamBitReader, ServerRPCParamBitReader.PackageMap, bOutSuccess);
+
+	if (!bOutSuccess)
+	{
+		UE_LOG(LogNetworkPrediction, Error, TEXT("ServerReceiveClientInput (%s): NetSerialize failed!"), *GetNameSafe(GetOwner()));
+	}
 }
 
 void UNetworkPredictionComponent::CallServerRPC()
 {
+	AActor* OwnerActor = GetOwner();
+	
 	// Temp hack to make sure the ServerRPC doesn't get suppressed from bandwidth limiting
 	// (system hasn't been optimized and not mature enough yet to handle gaps in input stream)
-	FScopedBandwidthLimitBypass BandwidthBypass(GetOwner());
+	// Note that this has no effect when using Iris replication.
+	FScopedBandwidthLimitBypass BandwidthBypass(OwnerActor);
+
+	// Reset bit writer without affecting allocations
+	FBitWriterMark BitWriterReset;
+	BitWriterReset.Pop(ServerRPCParamBitWriter);
+
+	// 'static' to avoid reallocation each invocation
+	static FServerReplicationRPCParameter ProxyParameter;	
+	ProxyParameter.SetReplicationProxy(ReplicationProxy_ServerRPC);
+
+	UNetConnection* NetConnection = OwnerActor->GetNetConnection();
+
+	if (UPackageMap* PackageMap = FServerReplicationRPCParameter::GetIrisPackageMapToCaptureReferences(NetConnection, ProxyParameter.PackageMapExports))
+	{
+		ServerRPCParamBitWriter.PackageMap = PackageMap;
+	}
+	else
+	{
+		// Extract the net package map used for serializing object references.
+		ServerRPCParamBitWriter.PackageMap = NetConnection ? ToRawPtr(NetConnection->PackageMap) : nullptr;
+	}
+
+	if (ServerRPCParamBitWriter.PackageMap == nullptr)
+	{
+		UE_LOG(LogNetworkPrediction, Error, TEXT("CallServerRPC: Failed to find a NetConnection/PackageMap for data serialization!"));
+		return;
+	}
+
+	// Reset captured exports stored in ProxyParameter
+	ProxyParameter.NetTokensPendingExport.Reset();
+	UE::Net::FNetTokenExportScope NetTokenExportScope(ServerRPCParamBitWriter, NetConnection->GetDriver()->GetNetTokenStore(), ProxyParameter.NetTokensPendingExport, "CallServerRPC");
+
+	bool bSerializeSuccess = true;
+	if (!ProxyParameter.SerializeToProxy(ServerRPCParamBitWriter, ServerRPCParamBitWriter.PackageMap, bSerializeSuccess))
+	{
+		return;
+	}
+
+	// Copy bits to our proxy param struct, which will be NetSerialized to the server
+	ProxyParameter.DataBits.SetNumUninitialized(ServerRPCParamBitWriter.GetNumBits());
+
+	check(ProxyParameter.DataBits.Num() >= ServerRPCParamBitWriter.GetNumBits());
+	FMemory::Memcpy(ProxyParameter.DataBits.GetData(), ServerRPCParamBitWriter.GetData(), ServerRPCParamBitWriter.GetNumBytes());
 
-	FServerReplicationRPCParameter ProxyParameter(ReplicationProxy_ServerRPC);
+	// Trigger RPC to the server
 	ServerReceiveClientInput(ProxyParameter);
 }
 
diff --git a/Engine/Plugins/Runtime/NetworkPrediction/Source/NetworkPrediction/Private/NetworkPredictionNetSerializers.cpp b/Engine/Plugins/Runtime/NetworkPrediction/Source/NetworkPrediction/Private/NetworkPredictionNetSerializers.cpp
new file mode 100644
--- /dev/null
+++ b/Engine/Plugins/Runtime/NetworkPrediction/Source/NetworkPrediction/Private/NetworkPredictionNetSerializers.cpp
@@ -0,0 +1,857 @@
+// Copyright Epic Games, Inc. All Rights Reserved.
+
+#include "NetworkPredictionNetSerializers.h"
+#include "NetworkPredictionConfig.h"
+#include "NetworkPredictionLog.h"
+#include "NetworkPredictionProxy.h"
+#include "NetworkPredictionReplicationProxy.h"
+#include "Engine/NetConnection.h"
+#include "Iris/ReplicationState/PropertyNetSerializerInfoRegistry.h"
+#include "Iris/Serialization/IrisObjectReferencePackageMap.h"
+#include "Iris/Serialization/IrisPackageMapExportUtil.h"
+#include "Iris/Serialization/NetBitStreamUtil.h"
+#include "Iris/Serialization/NetSerializerDelegates.h"
+#include "Net/Iris/ReplicationSystem/ReplicationSystemUtil.h"
+#include "Net/Iris/ReplicationSystem/EngineReplicationBridge.h"
+
+
+#include UE_INLINE_GENERATED_CPP_BY_NAME(NetworkPredictionNetSerializers)
+
+
+
+namespace UE::Net
+{
+
+struct FNetworkPredictionProxyNetSerializer
+{
+	// Version
+	static const uint32 Version = 0;
+
+	// Types
+	struct FQuantizedData
+	{
+		int32 SpawnID;
+		FNetworkPredictionInstanceArchetype InstanceArchetype;
+
+		bool IsEqual(const FQuantizedData& Other) const
+		{
+			return (SpawnID == Other.SpawnID && 
+			       InstanceArchetype.TickingMode == Other.InstanceArchetype.TickingMode);
+		}
+	};
+
+	typedef FNetworkPredictionProxy SourceType;
+	typedef FQuantizedData QuantizedType;
+	typedef FNetworkPredictionProxyNetSerializerConfig ConfigType;
+
+	inline static const ConfigType DefaultConfig;
+
+	static void Serialize(FNetSerializationContext&, const FNetSerializeArgs&);
+	static void Deserialize(FNetSerializationContext&, const FNetDeserializeArgs&);
+
+	static void Quantize(FNetSerializationContext&, const FNetQuantizeArgs&);
+	static void Dequantize(FNetSerializationContext&, const FNetDequantizeArgs&);
+
+	static bool IsEqual(FNetSerializationContext&, const FNetIsEqualArgs&);
+
+	// Avoid overwriting some members in the target FNetworkPredictionProxy instance.
+	static void Apply(FNetSerializationContext&, const FNetApplyArgs&);
+
+private:
+	class FNetSerializerRegistryDelegates final : private UE::Net::FNetSerializerRegistryDelegates
+	{
+	public:
+		virtual ~FNetSerializerRegistryDelegates();
+
+	private:
+		virtual void OnPreFreezeNetSerializerRegistry() override;
+
+		bool QuantizedTypeMeetRequirements() const;
+		void InitNetSerializer();
+
+		inline static const FName NetworkPredictionProxyNetSerializerName = FName("NetworkPredictionProxy");
+		UE_NET_IMPLEMENT_NAMED_STRUCT_NETSERIALIZER_INFO(NetworkPredictionProxyNetSerializerName, FNetworkPredictionProxyNetSerializer);
+	};
+
+	inline static FNetworkPredictionProxyNetSerializer::FNetSerializerRegistryDelegates NetSerializerRegistryDelegates;
+};
+
+UE_NET_IMPLEMENT_SERIALIZER(FNetworkPredictionProxyNetSerializer);
+
+// FNetworkPredictionProxyNetSerializer implementation
+void FNetworkPredictionProxyNetSerializer::Serialize(FNetSerializationContext& Context, const FNetSerializeArgs& Args)
+{
+	// Write to the network bitstream from the quantized representation
+	const QuantizedType& Value = *reinterpret_cast<const QuantizedType*>(Args.Source);
+
+	FNetBitStreamWriter* Writer = Context.GetBitStreamWriter();
+
+	WritePackedInt32(Writer, Value.SpawnID);
+	Writer->WriteBits((uint8)Value.InstanceArchetype.TickingMode, 8u);
+}
+
+void FNetworkPredictionProxyNetSerializer::Deserialize(FNetSerializationContext& Context, const FNetDeserializeArgs& Args)
+{
+	// Read a quantized representation from the network bitstream
+	QuantizedType& Target = *reinterpret_cast<QuantizedType*>(Args.Target);
+
+	FNetBitStreamReader* Reader = Context.GetBitStreamReader();
+
+	Target.SpawnID = ReadPackedInt32(Reader);
+
+	const uint8 TickingModeAsByte = static_cast<uint8>(Reader->ReadBits(8u));
+	Target.InstanceArchetype.TickingMode = static_cast<ENetworkPredictionTickingPolicy>(TickingModeAsByte);
+}
+
+
+void FNetworkPredictionProxyNetSerializer::Quantize(FNetSerializationContext& Context, const FNetQuantizeArgs& Args)
+{
+	// Create a quantized representation from the full instance shadow copy
+	SourceType& Source = *reinterpret_cast<SourceType*>(Args.Source);
+	QuantizedType& Target = *reinterpret_cast<QuantizedType*>(Args.Target);
+
+	Target.SpawnID = Source.GetID();
+	Target.InstanceArchetype.TickingMode = Source.GetCachedArchetype().TickingMode;
+}
+
+void FNetworkPredictionProxyNetSerializer::Dequantize(FNetSerializationContext& Context, const FNetDequantizeArgs& Args)
+{
+	// Write onto a full shadow copy instance from a quantized representation
+	const QuantizedType& Source = *reinterpret_cast<const QuantizedType*>(Args.Source);
+	SourceType& Target = *reinterpret_cast<SourceType*>(Args.Target);
+
+	Target.GetID() = FNetworkPredictionID(Source.SpawnID, Target.GetID().GetTraceID());
+	Target.GetCachedArchetype().TickingMode = Source.InstanceArchetype.TickingMode;
+}
+
+bool FNetworkPredictionProxyNetSerializer::IsEqual(FNetSerializationContext& Context, const FNetIsEqualArgs& Args)
+{
+	if (Args.bStateIsQuantized)
+	{
+		const QuantizedType& QuantizedValue0 = *reinterpret_cast<const QuantizedType*>(Args.Source0);
+		const QuantizedType& QuantizedValue1 = *reinterpret_cast<const QuantizedType*>(Args.Source1);
+		return QuantizedValue0.IsEqual(QuantizedValue1);
+	}
+	else
+	{
+		const SourceType& Value0 = *reinterpret_cast<SourceType*>(Args.Source0);
+		const SourceType& Value1 = *reinterpret_cast<SourceType*>(Args.Source1);
+
+		return Value0.Identical(&Value1, 0);
+	}
+}
+
+
+void FNetworkPredictionProxyNetSerializer::Apply(FNetSerializationContext& Context, const FNetApplyArgs& Args)
+{
+	// Copy selectively from the shadow copy to the real instance that game-level systems see, then let it configure itself
+	SourceType& Source = *reinterpret_cast<SourceType*>(Args.Source);
+	SourceType& Target = *reinterpret_cast<SourceType*>(Args.Target);
+
+	const int32 NewID = Source.GetID();	// ID will be set on Target inside ConfigureAfterLoading
+	Target.GetCachedArchetype().TickingMode = Source.GetCachedArchetype().TickingMode;
+
+	Target.ConfigureAfterLoading(NewID);
+}
+
+// FNetSerializerRegistryDelegates
+FNetworkPredictionProxyNetSerializer::FNetSerializerRegistryDelegates::~FNetSerializerRegistryDelegates()
+{
+	UE_NET_UNREGISTER_NETSERIALIZER_INFO(NetworkPredictionProxyNetSerializerName);
+}
+
+void FNetworkPredictionProxyNetSerializer::FNetSerializerRegistryDelegates::OnPreFreezeNetSerializerRegistry()
+{
+	UE_NET_REGISTER_NETSERIALIZER_INFO(NetworkPredictionProxyNetSerializerName);
+}
+
+
+} // end namespace UE::Net
+
+
+
+
+// FServerReplicationRPCParameterNetSerializer implementation
+//  Largely a copy of FCharacterNetworkSerializationPackedBitsNetSerializer to start
+namespace UE::Net::Private
+{
+	static constexpr inline uint32 CalculateRequiredWordCount(uint32 NumBits) { return (NumBits + NumBitsPerDWORD - 1U) / NumBitsPerDWORD; }
+}
+
+
+namespace UE::Net
+{
+
+	struct FServerReplicationRPCParameterNetSerializerQuantizedType
+	{
+		typedef uint32 WordType;
+		static constexpr uint32 DefaultReserveNumBits = FServerReplicationRPCParameter::NumReservedBits;
+		static constexpr uint32 InlinedWordCount = Private::CalculateRequiredWordCount(DefaultReserveNumBits);
+
+		typedef FNetSerializerArrayStorage<WordType, AllocationPolicies::TInlinedElementAllocationPolicy<InlinedWordCount>> FDataBitsStorage;
+
+		FIrisPackageMapExportsQuantizedType QuantizedExports;
+		FDataBitsStorage DataBitsStorage;
+		uint32 NumDataBits = 0;
+	};
+
+}
+
+template <> struct TIsPODType<UE::Net::FServerReplicationRPCParameterNetSerializerQuantizedType> { enum { Value = true }; };
+
+namespace UE::Net
+{
+
+struct FServerReplicationRPCParameterNetSerializer
+{
+	// Version
+	static const uint32 Version = 0;
+
+	// Traits
+	static constexpr bool bHasDynamicState = true;
+	static constexpr bool bHasCustomNetReference = true;
+
+	typedef uint32 WordType;
+
+	typedef FServerReplicationRPCParameter SourceType;
+	typedef FServerReplicationRPCParameterNetSerializerQuantizedType QuantizedType;
+	typedef struct FServerReplicationRPCParameterNetSerializerConfig ConfigType;
+
+	static const ConfigType DefaultConfig;
+
+	static void Serialize(FNetSerializationContext&, const FNetSerializeArgs&);
+	static void Deserialize(FNetSerializationContext&, const FNetDeserializeArgs&);
+
+	static void Quantize(FNetSerializationContext&, const FNetQuantizeArgs&);
+	static void Dequantize(FNetSerializationContext&, const FNetDequantizeArgs&);
+
+	static bool IsEqual(FNetSerializationContext&, const FNetIsEqualArgs&);
+	static bool Validate(FNetSerializationContext&, const FNetValidateArgs&);
+
+	static void CloneDynamicState(FNetSerializationContext&, const FNetCloneDynamicStateArgs&);
+	static void FreeDynamicState(FNetSerializationContext&, const FNetFreeDynamicStateArgs&);
+
+	static void CollectNetReferences(FNetSerializationContext&, const FNetCollectReferencesArgs&);
+
+private:
+	class FNetSerializerRegistryDelegates final : private UE::Net::FNetSerializerRegistryDelegates
+	{
+	public:
+		virtual ~FNetSerializerRegistryDelegates();
+
+	private:
+		virtual void OnPreFreezeNetSerializerRegistry() override;
+	};
+
+	static void FreeDynamicStateInternal(FNetSerializationContext&, QuantizedType& Value);
+
+	static FServerReplicationRPCParameterNetSerializer::FNetSerializerRegistryDelegates NetSerializerRegistryDelegates;
+};
+
+UE_NET_IMPLEMENT_SERIALIZER(FServerReplicationRPCParameterNetSerializer);
+
+
+const FServerReplicationRPCParameterNetSerializer::ConfigType FServerReplicationRPCParameterNetSerializer::DefaultConfig;
+FServerReplicationRPCParameterNetSerializer::FNetSerializerRegistryDelegates FServerReplicationRPCParameterNetSerializer::NetSerializerRegistryDelegates;
+
+void FServerReplicationRPCParameterNetSerializer::Serialize(FNetSerializationContext& Context, const FNetSerializeArgs& Args)
+{
+	// For now we ignore this in default state hash due to complications with asymmetrically serialized state.
+	if (Context.IsInitializingDefaultState())
+	{
+		return;
+	}
+
+	const QuantizedType& Value = *reinterpret_cast<QuantizedType*>(Args.Source);
+
+	FNetBitStreamWriter* Writer = Context.GetBitStreamWriter();
+
+	// Serialize captured references and exports
+	FIrisPackageMapExportsUtil::Serialize(Context, Value.QuantizedExports);
+
+	// Write data bits
+	const uint32 NumDataBits = Value.NumDataBits;
+	if (Writer->WriteBool(NumDataBits > 0))
+	{
+		UE::Net::WritePackedUint32(Writer, NumDataBits);
+		Writer->WriteBitStream(Value.DataBitsStorage.GetData(), 0, NumDataBits);
+	}
+}
+
+// static
+void FServerReplicationRPCParameterNetSerializer::FreeDynamicStateInternal(FNetSerializationContext& Context, QuantizedType& Value)
+{
+	// Free quantized state for captured reference and exports
+	FIrisPackageMapExportsUtil::FreeDynamicState(Context, Value.QuantizedExports);
+
+	Value.DataBitsStorage.Free(Context);
+	Value.NumDataBits = 0;
+}
+
+void FServerReplicationRPCParameterNetSerializer::Deserialize(FNetSerializationContext& Context, const FNetDeserializeArgs& Args)
+{
+	// For consistency, we should never get here. For now we ignore this in default state hash due to complications with asymmetrically serialized state.
+	if (Context.IsInitializingDefaultState())
+	{
+		return;
+	}
+
+	const ConfigType* Config = static_cast<const ConfigType*>(Args.NetSerializerConfig);
+	QuantizedType& TargetValue = *reinterpret_cast<QuantizedType*>(Args.Target);
+
+	FNetBitStreamReader* Reader = Context.GetBitStreamReader();
+
+	// Deserialize captured references and exports
+	FIrisPackageMapExportsUtil::Deserialize(Context, TargetValue.QuantizedExports);
+
+	const bool bHasDataBits = Reader->ReadBool();
+	if (bHasDataBits)
+	{
+		const uint32 NumDataBits = UE::Net::ReadPackedUint32(Reader);
+
+		const uint32 MaxNumDataBits = Config->MaxQuantizedSizeBits;
+		if (NumDataBits > MaxNumDataBits)
+		{
+			Context.SetError(GNetError_ArraySizeTooLarge);
+			UE_LOG(LogNetworkPrediction, Error, TEXT("FServerReplicationRPCParameterNetSerializer::Deserialize: Invalidating RPC param due to NumBits (%u) exceeding allowable limit (%u)."), NumDataBits, MaxNumDataBits);
+			ensureMsgf(false, TEXT("Invalidating received FServerReplicationRPCParameterRPC due to NumBits exceeding allowable limit"));
+			return;
+		}
+
+		const uint32 RequiredWordCount = Private::CalculateRequiredWordCount(NumDataBits);
+		TargetValue.DataBitsStorage.AdjustSize(Context, RequiredWordCount);
+
+		Reader->ReadBitStream(TargetValue.DataBitsStorage.GetData(), NumDataBits);
+		TargetValue.NumDataBits = NumDataBits;
+	}
+	else
+	{
+		TargetValue.DataBitsStorage.Free(Context);
+		TargetValue.NumDataBits = 0U;
+	}
+}
+
+void FServerReplicationRPCParameterNetSerializer::Quantize(FNetSerializationContext& Context, const FNetQuantizeArgs& Args)
+{
+	// Create a quantized representation from a full shadow copy instance
+	const SourceType& SourceValue = *reinterpret_cast<const SourceType*>(Args.Source);
+	QuantizedType& TargetValue = *reinterpret_cast<QuantizedType*>(Args.Target);
+
+	// Quantize captured references and exports
+	FIrisPackageMapExportsUtil::Quantize(Context, SourceValue.PackageMapExports, MakeArrayView(SourceValue.NetTokensPendingExport), TargetValue.QuantizedExports);
+
+	uint32 NumDataBits = SourceValue.DataBits.Num();
+
+	const ConfigType* Config = static_cast<const ConfigType*>(Args.NetSerializerConfig);
+	const uint32 MaxNumDataBits = Config->MaxQuantizedSizeBits;
+
+	if (NumDataBits > MaxNumDataBits)
+	{
+		// This is just to avoid disconnect and instead warn and invalidate the data on the sending side.
+		UE_LOG(LogNetworkPrediction, Error, TEXT("FServerReplicationRPCParameterNetSerializer::Quantize: Invalidating RPC param due to NumBits (%u) exceeding allowable limit (%u). See FServerReplicationRPCParameterNetSerializerConfig::DefaultMaxQuantizedSizeBits."), NumDataBits, MaxNumDataBits);
+		NumDataBits = 0U;
+		ensureMsgf(false, TEXT("Invalidating RPC param due to NumBits exceeding allowable limit"));
+	}
+
+	TargetValue.DataBitsStorage.AdjustSize(Context, Private::CalculateRequiredWordCount(NumDataBits));
+	if (NumDataBits > 0)
+	{
+		FMemory::Memcpy(TargetValue.DataBitsStorage.GetData(), SourceValue.DataBits.GetData(), (NumDataBits + 7U) / 8U);
+	}
+	TargetValue.NumDataBits = NumDataBits;
+}
+
+void FServerReplicationRPCParameterNetSerializer::Dequantize(FNetSerializationContext& Context, const FNetDequantizeArgs& Args)
+{
+	// Write onto a full shadow copy instance from a quantized representation
+	const QuantizedType& Source = *reinterpret_cast<const QuantizedType*>(Args.Source);
+	SourceType& Target = *reinterpret_cast<SourceType*>(Args.Target);
+
+	// Dequantize captured references and exports and inject into target
+	FIrisPackageMapExportsUtil::Dequantize(Context, Source.QuantizedExports, Target.PackageMapExports);
+
+	// DataBits
+	Target.DataBits.SetNumUninitialized(Source.NumDataBits);
+	Target.DataBits.SetRangeFromRange(0, Source.NumDataBits, Source.DataBitsStorage.GetData());
+}
+
+bool FServerReplicationRPCParameterNetSerializer::IsEqual(FNetSerializationContext& Context, const FNetIsEqualArgs& Args)
+{
+	if (Args.bStateIsQuantized)
+	{
+		const QuantizedType& Value0 = *reinterpret_cast<const QuantizedType*>(Args.Source0);
+		const QuantizedType& Value1 = *reinterpret_cast<const QuantizedType*>(Args.Source1);
+
+		if (Value0.NumDataBits != Value1.NumDataBits)
+		{
+			return false;
+		}
+
+		// Compare references and exports
+		if (!FIrisPackageMapExportsUtil::IsEqual(Context, Value0.QuantizedExports, Value1.QuantizedExports))
+		{
+			return false;
+		}
+
+		const uint32 RequiredWords = Private::CalculateRequiredWordCount(Value0.NumDataBits);
+		if (RequiredWords > 0 && FMemory::Memcmp(Value0.DataBitsStorage.GetData(), Value1.DataBitsStorage.GetData(), sizeof(WordType) * RequiredWords) != 0)
+		{
+			return false;
+		}
+	}
+	else
+	{
+		const SourceType& Value0 = *reinterpret_cast<SourceType*>(Args.Source0);
+		const SourceType& Value1 = *reinterpret_cast<SourceType*>(Args.Source1);
+		return Value0.DataBits == Value1.DataBits;
+	}
+
+	return true;
+}
+
+bool FServerReplicationRPCParameterNetSerializer::Validate(FNetSerializationContext& Context, const FNetValidateArgs& Args)
+{
+	const ConfigType* Config = static_cast<const ConfigType*>(Args.NetSerializerConfig);
+	const QuantizedType& SourceValue = *reinterpret_cast<const QuantizedType*>(Args.Source);
+
+	const uint32 MaxNumDataBits = Config->MaxQuantizedSizeBits;
+
+	if (SourceValue.NumDataBits > MaxNumDataBits)
+	{
+		return false;
+	}
+
+	if (!FIrisPackageMapExportsUtil::Validate(Context, SourceValue.QuantizedExports))
+	{
+		return false;
+	}
+
+	return true;
+}
+
+void FServerReplicationRPCParameterNetSerializer::CloneDynamicState(FNetSerializationContext& Context, const FNetCloneDynamicStateArgs& Args)
+{
+	const QuantizedType& SourceValue = *reinterpret_cast<const QuantizedType*>(Args.Source);
+	QuantizedType& TargetValue = *reinterpret_cast<QuantizedType*>(Args.Target);
+
+	FIrisPackageMapExportsUtil::CloneDynamicState(Context, TargetValue.QuantizedExports, SourceValue.QuantizedExports);
+
+	TargetValue.DataBitsStorage.Clone(Context, SourceValue.DataBitsStorage);
+	TargetValue.NumDataBits = SourceValue.NumDataBits;
+}
+
+void FServerReplicationRPCParameterNetSerializer::FreeDynamicState(FNetSerializationContext& Context, const FNetFreeDynamicStateArgs& Args)
+{
+	QuantizedType& Value = *reinterpret_cast<QuantizedType*>(Args.Source);
+	FreeDynamicStateInternal(Context, Value);
+}
+
+void FServerReplicationRPCParameterNetSerializer::CollectNetReferences(FNetSerializationContext& Context, const FNetCollectReferencesArgs& Args)
+{
+	const QuantizedType& Value = *reinterpret_cast<const QuantizedType*>(Args.Source);
+	FNetReferenceCollector& Collector = *reinterpret_cast<UE::Net::FNetReferenceCollector*>(Args.Collector);
+
+	FIrisPackageMapExportsUtil::CollectNetReferences(Context, Value.QuantizedExports, Args.ChangeMaskInfo, Collector);
+}
+
+static const FName PropertyNetSerializerRegistry_NAME_ServerReplicationRPCParameter("ServerReplicationRPCParameter");
+UE_NET_IMPLEMENT_NAMED_STRUCT_NETSERIALIZER_INFO(PropertyNetSerializerRegistry_NAME_ServerReplicationRPCParameter, FServerReplicationRPCParameterNetSerializer);
+
+
+FServerReplicationRPCParameterNetSerializer::FNetSerializerRegistryDelegates::~FNetSerializerRegistryDelegates()
+{
+	UE_NET_UNREGISTER_NETSERIALIZER_INFO(PropertyNetSerializerRegistry_NAME_ServerReplicationRPCParameter);
+}
+
+void FServerReplicationRPCParameterNetSerializer::FNetSerializerRegistryDelegates::OnPreFreezeNetSerializerRegistry()
+{
+	UE_NET_REGISTER_NETSERIALIZER_INFO(PropertyNetSerializerRegistry_NAME_ServerReplicationRPCParameter);
+}
+
+} // end UE::Net
+
+
+namespace UE::Net
+{
+
+	struct FFReplicationProxyNetSerializerQuantizedType
+	{
+		typedef uint32 WordType;
+		static constexpr uint32 DefaultReserveNumBits = FReplicationProxy::NumReservedBits;
+		static constexpr uint32 InlinedWordCount = Private::CalculateRequiredWordCount(DefaultReserveNumBits);
+
+		typedef FNetSerializerArrayStorage<WordType, AllocationPolicies::TInlinedElementAllocationPolicy<InlinedWordCount>> FDataBitsStorage;
+
+		FIrisPackageMapExportsQuantizedType QuantizedExports;
+		// How many bytes the current allocation can hold.
+		uint32 ByteCapacity = 0;
+		FDataBitsStorage DataBitsStorage;
+		uint32 NumDataBits = 0;
+
+		int32 CachedPendingFrame = INDEX_NONE;
+	};
+
+}
+
+template <> struct TIsPODType<UE::Net::FFReplicationProxyNetSerializerQuantizedType> { enum { Value = true }; };
+
+
+namespace UE::Net
+{
+
+struct FReplicationProxyNetSerializer
+{
+	// Version
+	static const uint32 Version = 0;
+
+	// Traits
+	static constexpr bool bHasDynamicState = true;
+	static constexpr bool bHasCustomNetReference = true;
+
+	// Types
+
+	typedef FReplicationProxy SourceType;
+	typedef FFReplicationProxyNetSerializerQuantizedType QuantizedType;
+	typedef FReplicationProxyNetSerializerConfig ConfigType;
+
+	inline static const ConfigType DefaultConfig;
+
+	static void Serialize(FNetSerializationContext&, const FNetSerializeArgs& Args);
+	static void Deserialize(FNetSerializationContext&, const FNetDeserializeArgs& Args);
+
+	static void Quantize(FNetSerializationContext&, const FNetQuantizeArgs& Args);
+	static void Dequantize(FNetSerializationContext&, const FNetDequantizeArgs& Args);
+
+	static bool IsEqual(FNetSerializationContext&, const FNetIsEqualArgs& Args);
+	
+	static void Apply(FNetSerializationContext&, const FNetApplyArgs& Args);
+
+	static void CloneDynamicState(FNetSerializationContext&, const FNetCloneDynamicStateArgs&);
+	static void FreeDynamicState(FNetSerializationContext&, const FNetFreeDynamicStateArgs&);
+
+	static void CollectNetReferences(FNetSerializationContext&, const FNetCollectReferencesArgs&);
+
+private:
+	static constexpr uint32 AllocationAlignment = 4U;
+
+	static void FreeDynamicStateInternal(FNetSerializationContext&, QuantizedType& Value);
+	static void GrowDynamicStateInternal(FNetSerializationContext&, QuantizedType& Value, uint32 NewBitCount);
+	static void AdjustStorageSize(FNetSerializationContext&, QuantizedType& Value, uint32 NewBitCount);
+
+private:
+	class FNetSerializerRegistryDelegates final : private UE::Net::FNetSerializerRegistryDelegates
+	{
+	public:
+		virtual ~FNetSerializerRegistryDelegates();
+
+	private:
+		virtual void OnPreFreezeNetSerializerRegistry() override;
+	};
+
+	static FReplicationProxyNetSerializer::FNetSerializerRegistryDelegates NetSerializerRegistryDelegates;
+};
+
+
+UE_NET_IMPLEMENT_SERIALIZER(FReplicationProxyNetSerializer);
+
+
+FReplicationProxyNetSerializer::FNetSerializerRegistryDelegates FReplicationProxyNetSerializer::NetSerializerRegistryDelegates;
+
+
+void FReplicationProxyNetSerializer::Serialize(FNetSerializationContext& Context, const FNetSerializeArgs& Args)
+{
+	if (Context.IsInitializingDefaultState())
+	{
+		return;
+	}
+
+	const ConfigType* Config = static_cast<const ConfigType*>(Args.NetSerializerConfig);
+	const QuantizedType& Value = *reinterpret_cast<const QuantizedType*>(Args.Source);
+	FNetBitStreamWriter* Writer = Context.GetBitStreamWriter();
+
+	// If we have any captured exports, serialize them.
+	FIrisPackageMapExportsUtil::Serialize(Context, Value.QuantizedExports);
+
+	// Write data.
+	const uint32 NumDataBits = Value.NumDataBits;
+	if (Writer->WriteBool(NumDataBits > 0))
+	{
+		UE::Net::WritePackedUint32(Writer, NumDataBits);
+		Writer->WriteBitStream(Value.DataBitsStorage.GetData(), 0, NumDataBits);
+	}
+}
+
+void FReplicationProxyNetSerializer::Deserialize(FNetSerializationContext& Context, const FNetDeserializeArgs& Args)
+{
+	// For consistency, we should never get here. For now we ignore this in default state hash due to complications with asymmetrically serialized state.
+	if (Context.IsInitializingDefaultState())
+	{
+		return;
+	}
+
+	const ConfigType* Config = static_cast<const ConfigType*>(Args.NetSerializerConfig);
+	QuantizedType& Value = *reinterpret_cast<QuantizedType*>(Args.Target);
+	const uint32 CurrentBitCount = Value.NumDataBits;
+
+	FNetBitStreamReader* Reader = Context.GetBitStreamReader();
+
+	// Read exports for packagemap.
+	FIrisPackageMapExportsUtil::Deserialize(Context, Value.QuantizedExports);
+
+	const bool bHasDataBits = Reader->ReadBool();
+
+	if (bHasDataBits)
+	{
+		// Read the data
+		const uint32 NewBitCount = ReadPackedUint32(Reader);
+
+		if (!ensureMsgf(NewBitCount <= Config->MaxQuantizedSizeBits,
+			TEXT("FReplicationProxyNetSerializer::Deserialize data size of %u bits exceeds maximum of %u bits."),
+			NewBitCount, Config->MaxQuantizedSizeBits))
+		{
+			Context.SetError(GNetError_ArraySizeTooLarge);
+			return;
+		}
+
+		AdjustStorageSize(Context, Value, NewBitCount);
+
+		Reader->ReadBitStream(static_cast<uint32*>(Value.DataBitsStorage.GetData()), NewBitCount);
+	}
+	else
+	{
+		Value.DataBitsStorage.Free(Context);
+		Value.NumDataBits = 0U;
+	}
+}
+
+
+void FReplicationProxyNetSerializer::Quantize(FNetSerializationContext& Context, const FNetQuantizeArgs& Args)
+{
+	const ConfigType* Config = static_cast<const ConfigType*>(Args.NetSerializerConfig);
+	SourceType& Source = *reinterpret_cast<SourceType*>(Args.Source);
+	QuantizedType& Value = *reinterpret_cast<QuantizedType*>(Args.Target);
+
+	Value.CachedPendingFrame = Source.GetCachedPendingFrame();
+
+	UIrisObjectReferencePackageMap* PackageMap = nullptr; 
+
+	if (const AActor* ReplicatingActor = Source.OwnerActor.Get())
+	{
+		if (UEngineReplicationBridge* ReplicationBridge = UE::Net::FReplicationSystemUtil::GetActorReplicationBridge(ReplicatingActor))
+		{
+			PackageMap = ReplicationBridge->GetObjectReferencePackageMap();
+		}
+	}
+
+	// Since this struct uses custom serialization path we need to explicitly capture exports in order to forward them to iris
+	UE::Net::FIrisPackageMapExports PackageMapExports;
+	UE::Net::FNetTokenExportContext::FNetTokenExports NetTokensPendingExport;
+
+	if (PackageMap)
+	{
+		PackageMap->InitForWrite(&PackageMapExports);
+	}
+
+	// Use the Property serialization and store as binary blob.
+	FNetBitWriter Archive(PackageMap, FReplicationProxy::NumReservedBits);
+	FNetTokenExportScope NetTokenExportScope(Archive, Context.GetNetTokenStore(), NetTokensPendingExport);
+
+	if (!Context.IsInitializingDefaultState())
+	{
+		bool bOutSuccess = true;
+		Source.NetSerialize(Archive, PackageMap, bOutSuccess);
+	}
+
+	const uint64 BitCount = Archive.GetNumBits();
+
+	if (!ensureMsgf(BitCount <= Config->MaxQuantizedSizeBits,
+		TEXT("FReplicationProxyNetSerializer::Quantize: data size of %llu bits exceeds maximum of %u bits."),
+		BitCount, Config->MaxQuantizedSizeBits))
+	{
+		Context.SetError(GNetError_ArraySizeTooLarge);
+		return;
+	}
+
+	if (!ensureMsgf(!Archive.IsError(), TEXT("FReplicationProxyNetSerializer::Quantize: NetBitWriter archive error in NetSerialize")))
+	{
+		Context.SetError(GNetError_InternalError);
+		return;
+	}
+
+	// Quantize captured exports
+	FIrisPackageMapExportsUtil::Quantize(Context, PackageMapExports, NetTokensPendingExport, Value.QuantizedExports);
+
+	// Deal with serialized data
+	AdjustStorageSize(Context, Value, static_cast<uint32>(BitCount));
+	if (BitCount > 0)
+	{
+		FMemory::Memcpy(Value.DataBitsStorage.GetData(), Archive.GetData(), (BitCount + 7U) / 8U);
+	}
+}
+
+void FReplicationProxyNetSerializer::Dequantize(FNetSerializationContext& Context, const FNetDequantizeArgs& Args)
+{
+	const QuantizedType& Source = *reinterpret_cast<const QuantizedType*>(Args.Source);
+	SourceType& Target = *reinterpret_cast<SourceType*>(Args.Target);
+
+	// Dequantize captured references and exports and inject into target
+	FIrisPackageMapExportsUtil::Dequantize(Context, Source.QuantizedExports, Target.PackageMapExports);
+
+	// DataBits
+	Target.DataBits.SetNumUninitialized(Source.NumDataBits);
+	Target.DataBits.SetRangeFromRange(0, Source.NumDataBits, Source.DataBitsStorage.GetData());
+}
+
+bool FReplicationProxyNetSerializer::IsEqual(FNetSerializationContext& Context, const FNetIsEqualArgs& Args)
+{
+	if (Args.bStateIsQuantized)
+	{
+		const QuantizedType& Value0 = *reinterpret_cast<const QuantizedType*>(Args.Source0);
+		const QuantizedType& Value1 = *reinterpret_cast<const QuantizedType*>(Args.Source1);
+
+		return (Value0.CachedPendingFrame == Value1.CachedPendingFrame);
+	}
+	else
+	{
+		const SourceType& Value0 = *reinterpret_cast<SourceType*>(Args.Source0);
+		const SourceType& Value1 = *reinterpret_cast<SourceType*>(Args.Source1);
+		return Value0.Identical(&Value1, 0);
+	}
+}
+
+void FReplicationProxyNetSerializer::Apply(FNetSerializationContext& Context, const FNetApplyArgs& Args)
+{
+	const SourceType& Source = *reinterpret_cast<const SourceType*>(Args.Source);
+	SourceType& Target = *reinterpret_cast<SourceType*>(Args.Target);
+
+	UIrisObjectReferencePackageMap* PackageMap = nullptr;
+
+	if (const AActor* ReplicatingActor = Target.OwnerActor.Get())
+	{
+		if (UEngineReplicationBridge* ReplicationBridge = UE::Net::FReplicationSystemUtil::GetActorReplicationBridge(ReplicatingActor))
+		{
+			PackageMap = ReplicationBridge->GetObjectReferencePackageMap();
+		}
+	}
+
+	// Setup resolve context for call into NetSerialize
+	UE::Net::FNetTokenResolveContext ResolveContext;
+	ResolveContext.RemoteNetTokenStoreState = Context.GetRemoteNetTokenStoreState();
+	ResolveContext.NetTokenStore = Context.GetNetTokenStore();
+
+	if (PackageMap)
+	{
+		PackageMap->InitForRead(&Source.PackageMapExports, ResolveContext);
+	}
+
+	uint32 NumDataBits = Source.DataBits.Num();
+
+	if (NumDataBits > 0)
+	{
+		// Note that we are reading from Source's data bits and net serializing directly from them, 
+		// without copying the bits to Target. This avoids an unnecessary memcpy.
+		FNetBitReader Archive(PackageMap, static_cast<uint8*>((void*)Source.DataBits.GetData()), NumDataBits);
+		bool bOutSuccess = true;
+		Target.NetSerialize(Archive, PackageMap, bOutSuccess);
+	}
+}
+
+
+void FReplicationProxyNetSerializer::CloneDynamicState(FNetSerializationContext& Context, const FNetCloneDynamicStateArgs& Args)
+{
+	const QuantizedType& SourceValue = *reinterpret_cast<const QuantizedType*>(Args.Source);
+	QuantizedType& TargetValue = *reinterpret_cast<QuantizedType*>(Args.Target);
+
+	FIrisPackageMapExportsUtil::CloneDynamicState(Context, TargetValue.QuantizedExports, SourceValue.QuantizedExports);
+	TargetValue.DataBitsStorage.Clone(Context, SourceValue.DataBitsStorage);
+	TargetValue.NumDataBits = SourceValue.NumDataBits;
+	TargetValue.ByteCapacity = SourceValue.ByteCapacity;
+	TargetValue.CachedPendingFrame = SourceValue.CachedPendingFrame;
+}
+
+void FReplicationProxyNetSerializer::FreeDynamicState(FNetSerializationContext& Context, const FNetFreeDynamicStateArgs& Args)
+{
+	FreeDynamicStateInternal(Context, *reinterpret_cast<QuantizedType*>(Args.Source));
+}
+
+void FReplicationProxyNetSerializer::FreeDynamicStateInternal(FNetSerializationContext& Context, QuantizedType& Value)
+{
+	// Free quantized state for captured reference and exports
+	FIrisPackageMapExportsUtil::FreeDynamicState(Context, Value.QuantizedExports);
+
+	Value.DataBitsStorage.Free(Context);
+	Value.NumDataBits = 0;
+}
+
+void FReplicationProxyNetSerializer::GrowDynamicStateInternal(FNetSerializationContext& Context, QuantizedType& Value, uint32 NewBitCount)
+{
+	checkSlow(NewBitCount > Value.NumDataBits);
+
+	const uint32 ByteCount = Align((NewBitCount + 7U) / 8U, AllocationAlignment);
+
+	const uint32 RequiredWordCount = Private::CalculateRequiredWordCount(NewBitCount);
+	Value.DataBitsStorage.AdjustSize(Context, RequiredWordCount);
+
+	// Clear the last word to support IsEqual Memcmp optimization.
+	const uint32 LastWordIndex = ByteCount / 4U - 1U;
+	static_cast<uint32*>(Value.DataBitsStorage.GetData())[LastWordIndex] = 0U;
+
+	Value.ByteCapacity = ByteCount;
+	Value.NumDataBits = NewBitCount;
+}
+
+
+void FReplicationProxyNetSerializer::AdjustStorageSize(FNetSerializationContext& Context, QuantizedType& Value, uint32 NewBitCount)
+{
+	const uint32 NewByteCapacity = Align((NewBitCount + 7U) / 8U, AllocationAlignment);
+	if (NewByteCapacity == 0)
+	{
+		// Free everything
+		FreeDynamicStateInternal(Context, Value);
+	}
+	else if (NewByteCapacity > Value.ByteCapacity)
+	{
+		GrowDynamicStateInternal(Context, Value, NewBitCount);
+	}
+	// If byte capacity is within the allocated capacity we just update the bit count and clear the last word
+	else
+	{
+		Value.NumDataBits = NewBitCount;
+
+		// Clear the last word to support IsEqual Memcmp optimization.
+		const uint32 LastWordIndex = NewByteCapacity / 4U - 1U;
+		static_cast<uint32*>(Value.DataBitsStorage.GetData())[LastWordIndex] = 0U;
+	}
+}
+
+void FReplicationProxyNetSerializer::CollectNetReferences(FNetSerializationContext& Context, const FNetCollectReferencesArgs& Args)
+{
+	const ConfigType& Config = *static_cast<const ConfigType*>(Args.NetSerializerConfig);
+	const QuantizedType& Value = *reinterpret_cast<const QuantizedType*>(Args.Source);
+	FNetReferenceCollector& Collector = *reinterpret_cast<UE::Net::FNetReferenceCollector*>(Args.Collector);
+
+	FIrisPackageMapExportsUtil::CollectNetReferences(Context, Value.QuantizedExports, Args.ChangeMaskInfo, Collector);
+}
+
+static const FName PropertyNetSerializerRegistry_NAME_ReplicationProxy("ReplicationProxy");
+UE_NET_IMPLEMENT_NAMED_STRUCT_NETSERIALIZER_INFO(PropertyNetSerializerRegistry_NAME_ReplicationProxy, FReplicationProxyNetSerializer);
+
+
+FReplicationProxyNetSerializer::FNetSerializerRegistryDelegates::~FNetSerializerRegistryDelegates()
+{
+	UE_NET_UNREGISTER_NETSERIALIZER_INFO(PropertyNetSerializerRegistry_NAME_ReplicationProxy);
+}
+
+void FReplicationProxyNetSerializer::FNetSerializerRegistryDelegates::OnPreFreezeNetSerializerRegistry()
+{
+	UE_NET_REGISTER_NETSERIALIZER_INFO(PropertyNetSerializerRegistry_NAME_ReplicationProxy);
+}
+
+
+
+} // end of UE::Net
+
+
+
+
+
diff --git a/Engine/Plugins/Runtime/NetworkPrediction/Source/NetworkPrediction/Private/NetworkPredictionNetSerializers.h b/Engine/Plugins/Runtime/NetworkPrediction/Source/NetworkPrediction/Private/NetworkPredictionNetSerializers.h
new file mode 100644
--- /dev/null
+++ b/Engine/Plugins/Runtime/NetworkPrediction/Source/NetworkPrediction/Private/NetworkPredictionNetSerializers.h
@@ -0,0 +1,68 @@
+// Copyright Epic Games, Inc. All Rights Reserved.
+
+#pragma once
+
+#include "Iris/Serialization/NetSerializer.h"
+#include "UObject/ObjectMacros.h"
+#include "NetworkPredictionReplicationProxy.h"
+#include "NetworkPredictionNetSerializers.generated.h"
+
+
+// Iris net serialization for FNetworkPredictionProxy. 
+// This is required due to the struct mixing locally-set and replicated member variables.
+USTRUCT()
+struct FNetworkPredictionProxyNetSerializerConfig : public FNetSerializerConfig
+{
+	GENERATED_BODY()
+};
+
+namespace UE::Net
+{
+	UE_NET_DECLARE_SERIALIZER(FNetworkPredictionProxyNetSerializer, NETWORKPREDICTION_API);
+}
+
+
+// Iris net serialization for FServerReplicationRPCParameter. 
+// This is required due to the struct using a custom NetSerialize function that is routed to other serialization method(s) that may do polymorphic and/or conditional serialization.
+// Implementation is based on FCharacterNetworkSerializationPackedBitsNetSerializer
+USTRUCT()
+struct FServerReplicationRPCParameterNetSerializerConfig : public FNetSerializerConfig
+{
+	GENERATED_BODY()
+
+public:
+	static constexpr uint32 DefaultMaxQuantizedSizeBits = FServerReplicationRPCParameter::MaxNumBits;
+
+	// Value used to sanity check incoming data so that we do not over-allocate dynamic memory
+	UPROPERTY()
+	uint32 MaxQuantizedSizeBits = DefaultMaxQuantizedSizeBits;
+};
+
+namespace UE::Net
+{
+
+UE_NET_DECLARE_SERIALIZER(FServerReplicationRPCParameterNetSerializer, NETWORKPREDICTION_API);
+
+}
+
+
+// Iris net serialization for FReplicationProxy. 
+// This is required due to the struct using a custom NetSerialize function that is routed to other serialization method(s) that may do polymorphic and/or conditional serialization.
+// Implementation is partly based on the last resort net serializer.
+USTRUCT()
+struct FReplicationProxyNetSerializerConfig : public FNetSerializerConfig
+{
+	GENERATED_BODY()
+
+public:
+	static constexpr uint32 DefaultMaxQuantizedSizeBits = FReplicationProxy::MaxNumBits;
+
+	// Value used to sanity check incoming data so that we do not over-allocate dynamic memory
+	UPROPERTY()
+	uint32 MaxQuantizedSizeBits = DefaultMaxQuantizedSizeBits;
+};
+
+namespace UE::Net
+{
+	UE_NET_DECLARE_SERIALIZER(FReplicationProxyNetSerializer, NETWORKPREDICTION_API);
+}
\ No newline at end of file
diff --git a/Engine/Plugins/Runtime/NetworkPrediction/Source/NetworkPrediction/Private/NetworkPredictionReplicationProxy.cpp b/Engine/Plugins/Runtime/NetworkPrediction/Source/NetworkPrediction/Private/NetworkPredictionReplicationProxy.cpp
--- a/Engine/Plugins/Runtime/NetworkPrediction/Source/NetworkPrediction/Private/NetworkPredictionReplicationProxy.cpp
+++ b/Engine/Plugins/Runtime/NetworkPrediction/Source/NetworkPrediction/Private/NetworkPredictionReplicationProxy.cpp
@@ -1,7 +1,12 @@
 // Copyright Epic Games, Inc. All Rights Reserved
 
 #include "NetworkPredictionReplicationProxy.h"
+#include "NetworkPredictionLog.h"
 #include "NetworkPredictionProxy.h"
+#include "Iris/ReplicationState/PropertyNetSerializerInfoRegistry.h"
+#include "Iris/ReplicationSystem/ReplicationSystem.h"
+#include "Net/Iris/ReplicationSystem/EngineReplicationBridge.h"
+#include "Net/Iris/ReplicationSystem/ReplicationSystemUtil.h"
 #include "Engine/NetConnection.h"
 
 #include UE_INLINE_GENERATED_CPP_BY_NAME(NetworkPredictionReplicationProxy)
@@ -18,13 +23,18 @@ void FReplicationProxy::Init(FNetworkPredictionProxy* InNetSimProxy, EReplicatio
 
 bool FReplicationProxy::NetSerialize(FArchive& Ar, class UPackageMap* Map, bool& bOutSuccess)
 {
-	if (npEnsureMsgf(NetSerializeFunc, TEXT("NetSerializeFunc not set for FReplicationProxy %d"), ReplicationTarget))
+	if (NetSerializeFunc)
 	{
 		FNetSerializeParams P(Ar, Map);
 		NetSerializeFunc(P);
 		bOutSuccess = !Ar.IsError();
 		return true;
 	}
+	else
+	{
+		UE_LOG(LogNetworkPrediction, Warning, TEXT("NetSerializeFunc not set for FReplicationProxy %d. This may occur as replication begins if using Iris, but should subside once the object's net role is established."), ReplicationTarget);
+	}
+
 	bOutSuccess = false;
 	return true;
 }
@@ -43,50 +53,92 @@ bool FReplicationProxy::Identical(const FReplicationProxy* Other, uint32 PortFla
 }
 
 // -------------------------------------------------------------------------------------------------------------------------------
-//	FServerRPCProxyParameter
+//	FServerReplicationRPCParameter
 // -------------------------------------------------------------------------------------------------------------------------------
-
-TArray<uint8> FServerReplicationRPCParameter::TempStorage;
-
 bool FServerReplicationRPCParameter::NetSerialize(FArchive& Ar, class UPackageMap* Map, bool& bOutSuccess)
 {
+	CachedPackageMap = Map;
+
+	// Array size in bits, using minimal number of bytes to write it out.
+	uint32 NumBits = DataBits.Num();
+	Ar.SerializeIntPacked(NumBits);
+
+	if (NumBits > FServerReplicationRPCParameter::MaxNumBits)
+	{
+		// Protect against bad data that could cause server to allocate way too much memory.
+		UE_LOG(LogNetworkPrediction, Error, TEXT("FServerReplicationRPCParameter::NetSerialize: Dropping incoming RPC param due to NumBits (%d) exceeding allowable limit (%d)."), NumBits, FServerReplicationRPCParameter::MaxNumBits);
+		bOutSuccess = false;
+		return true;
+	}
+
 	if (Ar.IsLoading())
 	{
-		// Loading: serialize to temp storage. We'll do the real deserialize in a manual call to ::NetSerializeToProxy
-		FNetBitReader& BitReader = (FNetBitReader&)Ar;
-		CachedNumBits = BitReader.GetBitsLeft();
-		CachedPackageMap = Map;
-
-		const int64 BytesLeft = BitReader.GetBytesLeft();
-		check(BytesLeft > 0); // Should not possibly be able to get here with an empty archive
-		TempStorage.Reset(BytesLeft);
-		TempStorage.SetNumUninitialized(BytesLeft);
-		TempStorage.Last() = 0;
-
-		BitReader.SerializeBits(TempStorage.GetData(), CachedNumBits);
+		DataBits.Init(0, NumBits);
 	}
-	else
+	else if (Ar.IsSaving() && NetTokensPendingExport.Num())
 	{
-		// Saving: directly call into the proxy's NetSerialize. No need for temp storage.
-		check(Proxy); // Must have been set before, via ctor.
-		return Proxy->NetSerialize(Ar, Map, bOutSuccess);
+		// As we now support exporting NetTokens from shared serialization and FServerReplicationRPCParameter serializes data outside of the normal flow
+		// we explicitly capture exports which we needs to be inject during actual serialization.
+		if (UE::Net::FNetTokenExportContext* ExportContext = UE::Net::FNetTokenExportContext::GetNetTokenExportContext(Ar))
+		{
+			ExportContext->AppendNetTokensPendingExport(NetTokensPendingExport);
+		}
 	}
 
-	return true;
+	// Array data
+	Ar.SerializeBits(DataBits.GetData(), NumBits);
+
+	bOutSuccess = !Ar.IsError();
+	return true; 
+
+
 }
 
-void FServerReplicationRPCParameter::NetSerializeToProxy(FReplicationProxy& InProxy)
+bool FServerReplicationRPCParameter::SerializeToProxy(FArchive& Ar, class UPackageMap* Map, bool& bOutSuccess)
 {
-	check(CachedPackageMap != nullptr);
-	check(CachedNumBits != -1);
+	check (Ar.IsSaving()); // This function is only for writing to the archive from the client's proxy
+	check (Proxy);
 
-	FNetBitReader BitReader(CachedPackageMap, TempStorage.GetData(), CachedNumBits);
+	Proxy->NetSerialize(Ar, Map, bOutSuccess);
 
-	bool bOutSuccess = true;
-	InProxy.NetSerialize(BitReader, CachedPackageMap, bOutSuccess);
+	return !Ar.IsError();
+}
+
+//static 
+UIrisObjectReferencePackageMap* FServerReplicationRPCParameter::GetIrisPackageMapToCaptureReferences(UNetConnection* NetConnection, UE::Net::FIrisPackageMapExports& TargetPkgMapExports)
+{
+	using namespace UE::Net;
+
+	if (const UEngineReplicationBridge* Bridge = FReplicationSystemUtil::GetActorReplicationBridge(NetConnection))
+	{
+		if (UIrisObjectReferencePackageMap* ObjectReferencePackageMap = Bridge->GetObjectReferencePackageMap())
+		{
+			ObjectReferencePackageMap->InitForWrite(&TargetPkgMapExports);
+
+			return ObjectReferencePackageMap;
+		}
+	}
+
+	return nullptr;
+}
+
+//static
+UIrisObjectReferencePackageMap* FServerReplicationRPCParameter::GetIrisPackageMapToReadReferences(const UNetConnection* NetConnection, UE::Net::FIrisPackageMapExports& TargetPkgMapExports)
+{
+	using namespace UE::Net;
+
+	if (const UEngineReplicationBridge* Bridge = FReplicationSystemUtil::GetActorReplicationBridge(NetConnection))
+	{
+		if (UIrisObjectReferencePackageMap* ObjectReferencePackageMap = Bridge->GetObjectReferencePackageMap())
+		{
+			const UReplicationSystem* ReplicationSystem = UE::Net::FReplicationSystemUtil::GetReplicationSystem(NetConnection->GetDriver());
+			ObjectReferencePackageMap->InitForRead(&TargetPkgMapExports, ReplicationSystem->GetNetTokenResolveContext(NetConnection->GetConnectionHandle().GetParentConnectionId()));
+
+			return ObjectReferencePackageMap;
+		}
+	}
 
-	CachedNumBits = -1;
-	CachedPackageMap = nullptr;
+	return nullptr;
 }
 
 // -------------------------------------------------------------------------------------------------------------------------------
diff --git a/Engine/Plugins/Runtime/NetworkPrediction/Source/NetworkPrediction/Public/NetworkPredictionComponent.h b/Engine/Plugins/Runtime/NetworkPrediction/Source/NetworkPrediction/Public/NetworkPredictionComponent.h
--- a/Engine/Plugins/Runtime/NetworkPrediction/Source/NetworkPrediction/Public/NetworkPredictionComponent.h
+++ b/Engine/Plugins/Runtime/NetworkPrediction/Source/NetworkPrediction/Public/NetworkPredictionComponent.h
@@ -74,6 +74,13 @@ private:
 	UPROPERTY(Replicated, transient)
 	FReplicationProxy ReplicationProxy_Replay;
 
+	/** Iris support: used for writing server RPC parameter bits. */
+	FNetBitWriter ServerRPCParamBitWriter;
+
+	/** Iris support: used for reading server RPC parameter bits. */
+	FNetBitReader ServerRPCParamBitReader;
+
+
 protected:
 
 	FReplicationProxySet GetReplicationProxies()
diff --git a/Engine/Plugins/Runtime/NetworkPrediction/Source/NetworkPrediction/Public/NetworkPredictionProxy.h b/Engine/Plugins/Runtime/NetworkPrediction/Source/NetworkPrediction/Public/NetworkPredictionProxy.h
--- a/Engine/Plugins/Runtime/NetworkPrediction/Source/NetworkPrediction/Public/NetworkPredictionProxy.h
+++ b/Engine/Plugins/Runtime/NetworkPrediction/Source/NetworkPrediction/Public/NetworkPredictionProxy.h
@@ -5,6 +5,7 @@
 #include "NetworkPredictionID.h"
 #include "NetworkPredictionStateView.h"
 #include "NetworkPredictionConfig.h"
+#include "NetworkPredictionCheck.h"
 #include "Engine/EngineBaseTypes.h"
 #include "Engine/EngineTypes.h"
 #include "NetworkPredictionProxy.generated.h"
@@ -191,27 +192,37 @@ struct FNetworkPredictionProxy
 
 		if (Ar.IsLoading())
 		{
-			if ((int32)ID != RawID)
+			ConfigureAfterLoading(RawID);
+		}
+
+		bOutSuccess = !Ar.IsError();
+		return true;
+	}
+
+	void ConfigureAfterLoading(uint32 LoadedID)
+	{
+		if ((int32)ID != LoadedID)
+		{
+			if (ConfigFunc)
 			{
-				if (ConfigFunc)
-				{
-					// We've already been init, so have to go through ConfigFunc to remap new id
-					FNetworkPredictionID NewID(RawID, ID.GetTraceID());
-					ConfigFunc(this, NewID, EConfigAction::UpdateConfigWithDefault);
-				}
-				else
-				{
-					// haven't been init yet so just set the replicated ID so we don't create a client side one
-					ID = FNetworkPredictionID(RawID);
-				}
+				// We've already been init, so have to go through ConfigFunc to remap new id
+				FNetworkPredictionID NewID(LoadedID, ID.GetTraceID());
+				ConfigFunc(this, NewID, EConfigAction::UpdateConfigWithDefault);
 			}
 			else
+			{
+				// haven't been init yet so just set the replicated ID so we don't create a client side one
+				ID = FNetworkPredictionID(LoadedID);
+			}
+		}
+		else
+		{
+			if (npEnsureMsgf((ConfigFunc != nullptr), TEXT("Attempting to configure with a changed archetype but no ConfigFunc has been set yet!")))
 			{
 				// Archetype change, call ConfigFunc but don't change ID
 				ConfigFunc(this, FNetworkPredictionID(), EConfigAction::UpdateConfigWithDefault);
 			}
 		}
-		return true;
 	}
 
 	bool Identical(const FNetworkPredictionProxy* Other, uint32 PortFlags) const
@@ -226,6 +237,10 @@ struct FNetworkPredictionProxy
 	ENetRole GetCachedNetRole() const { return CachedNetRole; }
 	bool GetCachedHasNetConnection() const { return bCachedHasNetConnection; }
 
+public:	// Direct accessors for Iris replication support only
+	FNetworkPredictionID& GetID() { return ID; }
+	FNetworkPredictionInstanceArchetype& GetCachedArchetype() { return CachedArchetype; }
+
 private:
 
 	// Allows ConfigFunc to be invoked to "do a thing" instead of set a new config/id.
@@ -246,7 +261,7 @@ private:
 	FNetworkPredictionStateView View;
 
 	ENetRole CachedNetRole = ROLE_None;
-	bool bCachedHasNetConnection;
+	bool bCachedHasNetConnection = false;
 	FNetworkPredictionInstanceConfig CachedConfig;
 	FNetworkPredictionInstanceArchetype CachedArchetype;
 	uint8 ArchetypeDirtyCount = 0;
diff --git a/Engine/Plugins/Runtime/NetworkPrediction/Source/NetworkPrediction/Public/NetworkPredictionReplicationProxy.h b/Engine/Plugins/Runtime/NetworkPrediction/Source/NetworkPrediction/Public/NetworkPredictionReplicationProxy.h
--- a/Engine/Plugins/Runtime/NetworkPrediction/Source/NetworkPrediction/Public/NetworkPredictionReplicationProxy.h
+++ b/Engine/Plugins/Runtime/NetworkPrediction/Source/NetworkPrediction/Public/NetworkPredictionReplicationProxy.h
@@ -2,6 +2,10 @@
 
 #pragma once
 #include "NetworkPredictionCheck.h"
+#include "Net/Core/NetToken/NetTokenExportContext.h"
+#include "Net/Iris/ReplicationSystem/ReplicationSystemUtil.h"
+#include "Iris/Serialization/IrisObjectReferencePackageMap.h"
+
 
 #include "NetworkPredictionReplicationProxy.generated.h"
 
@@ -33,6 +37,7 @@ struct FNetSerializeParams
 	UPackageMap* Map;
 };
 
+
 // Redirects NetSerialize to a dynamically set NetSerializeFunc.
 // This is how we hook into the replication systems role-based serialization
 USTRUCT()
@@ -45,9 +50,35 @@ struct FReplicationProxy
 	UE_API void OnPreReplication();	
 	UE_API bool Identical(const FReplicationProxy* Other, uint32 PortFlags) const;
 
+	UE_API int32 GetCachedPendingFrame() const 
+	{ 
+		return CachedPendingFrame; 
+	}
+
 	TFunction<void(const FNetSerializeParams& P)> NetSerializeFunc;
+
 	FNetworkPredictionProxy* NetSimProxy = nullptr;
 
+
+
+public:
+	// Number of bits to reserve in serialization container. Make this large enough to try to avoid re-allocation during the worst case calls.
+	static constexpr uint32 NumReservedBits = 4096;
+	static constexpr uint32 MaxNumBits = 16384;
+
+	// Iris support: holding ptr to associated actor for replication bridge / package map retrieval
+	TWeakObjectPtr<AActor> OwnerActor = nullptr;
+
+	// Iris support: TInlineAllocator used with TBitArray takes the number of 32-bit dwords, but the define is in number of bits, so convert here by dividing by 32.
+	TBitArray<TInlineAllocator<NumReservedBits / NumBitsPerDWORD>> DataBits;
+
+	// Iris support: Since this struct uses custom serialization path we need to explicitly capture exports in order to forward them to iris
+	// This is managed by the use of a custom packagemap.
+	UE::Net::FIrisPackageMapExports PackageMapExports;
+
+	// Iris support: Since we capturing data outside of the normal serialization path we also need to store exports to inject when actually sending the data.
+	UE::Net::FNetTokenExportContext::FNetTokenExports NetTokensPendingExport;
+
 private:
 
 	EReplicationProxyTarget ReplicationTarget;
@@ -82,32 +113,55 @@ struct FReplicationProxySet
 	}
 };
 
+
+
 // -------------------------------------------------------------------------------------------------------------------------------
 //	FServerRPCProxyParameter
 //	Used for the client->Server RPC. Since this is instantiated on the stack by the replication system prior to net serializing,
 //	we have no opportunity to point the RPC parameter to the member variables we want. So we serialize into a generic temp byte buffer
-//	and move them into the real buffers on the component in the RPC body (via ::NetSerialzeToProxy).
+//	and move them into the real buffers on the component in the RPC body (via ::SerializeToProxy).
 // -------------------------------------------------------------------------------------------------------------------------------
 USTRUCT()
 struct FServerReplicationRPCParameter
 {
 	GENERATED_BODY()
 
-	// Receive flow: ctor() -> NetSerializetoProxy
+	// Receive flow: ctor() -> NetSerialize
 	FServerReplicationRPCParameter() : Proxy(nullptr)	{ }
-	UE_API void NetSerializeToProxy(FReplicationProxy& InProxy);
 
-	// Send flow: ctor(Proxy) -> NetSerialize
+	// Send flow: ctor(Proxy) -> SerializeToProxy
 	FServerReplicationRPCParameter(FReplicationProxy& InProxy) : Proxy(&InProxy) { }
+	UE_API void SetReplicationProxy(FReplicationProxy& InProxy) { Proxy = &InProxy; }
+	UE_API bool SerializeToProxy(FArchive& Ar, class UPackageMap* Map, bool& bOutSuccess);
 	UE_API bool NetSerialize(FArchive& Ar, class UPackageMap* Map, bool& bOutSuccess);
 
-private:
+	UE_API UPackageMap* GetPackageMap() const
+	{ 
+		return CachedPackageMap; 
+	}
 
-	static UE_API TArray<uint8> TempStorage;
+public:
+	// Iris support
+	static UIrisObjectReferencePackageMap* GetIrisPackageMapToCaptureReferences(UNetConnection* NetConnection, UE::Net::FIrisPackageMapExports& PackageMapExports);
+	static UIrisObjectReferencePackageMap* GetIrisPackageMapToReadReferences(const UNetConnection* NetConnection, UE::Net::FIrisPackageMapExports& PackageMapExports);
 
-	FReplicationProxy* Proxy;
-	int64 CachedNumBits = -1;
-	class UPackageMap* CachedPackageMap = nullptr;
+	// Number of bits to reserve in serialization container. Make this large enough to try to avoid re-allocation during worst-case RPC calls.
+	static constexpr uint32 NumReservedBits = 4096;
+	static constexpr uint32 MaxNumBits = 16384;
+
+	// Iris support: TInlineAllocator used with TBitArray takes the number of 32-bit dwords, but the define is in number of bits, so convert here by dividing by 32.
+	TBitArray<TInlineAllocator<NumReservedBits / NumBitsPerDWORD>> DataBits;
+
+	// Iris support: Since this struct uses custom serialization path we need to explicitly capture exports in order to forward them to iris
+	// This is managed by the use of a custom packagemap.
+	UE::Net::FIrisPackageMapExports PackageMapExports;
+
+	// Iris support: Since we capturing data outside of the normal serialization path we also need to store exports to inject when actually sending the data.
+	UE::Net::FNetTokenExportContext::FNetTokenExports NetTokensPendingExport;
+
+private:
+	FReplicationProxy* Proxy = nullptr;
+	UPackageMap* CachedPackageMap = nullptr;	// Used for non-Iris replication
 };
 
 template<>
