Network Prediction & Mover: expanding ShouldReconcile to allow it to be deferred to the simulation driver, and opening the door to custom behaviors at either Mover's driver or component level, such as altering reconciliation based on settings outside of the Synx/Aux states.

#Synced-CL 49404182

diff --git a/Engine/Plugins/Experimental/Mover/Source/Mover/Private/Backends/MoverNetworkPredictionLiaison.cpp b/Engine/Plugins/Experimental/Mover/Source/Mover/Private/Backends/MoverNetworkPredictionLiaison.cpp
--- a/Engine/Plugins/Experimental/Mover/Source/Mover/Private/Backends/MoverNetworkPredictionLiaison.cpp
+++ b/Engine/Plugins/Experimental/Mover/Source/Mover/Private/Backends/MoverNetworkPredictionLiaison.cpp
@@ -54,6 +54,16 @@ void UMoverNetworkPredictionLiaisonComponent::ProduceInput(const int32 DeltaTime
 	MoverComp->ProduceInput(DeltaTimeMS, Cmd);
 }
 
+bool UMoverNetworkPredictionLiaisonComponent::ShouldReconcile(const FMoverSyncState* PredictedState, const FMoverSyncState* AuthorityState) const
+{
+	return PredictedState->ShouldReconcile(*AuthorityState);
+}
+
+bool UMoverNetworkPredictionLiaisonComponent::ShouldReconcile(const FMoverAuxStateContext* PredictedState, const FMoverAuxStateContext* AuthorityState) const
+{
+	return PredictedState->ShouldReconcile(*AuthorityState);
+}
+
 void UMoverNetworkPredictionLiaisonComponent::RestoreFrame(const FMoverSyncState* SyncState, const FMoverAuxStateContext* AuxState)
 {
 	check(MoverComp);
diff --git a/Engine/Plugins/Experimental/Mover/Source/Mover/Public/Backends/MoverNetworkPredictionLiaison.h b/Engine/Plugins/Experimental/Mover/Source/Mover/Public/Backends/MoverNetworkPredictionLiaison.h
--- a/Engine/Plugins/Experimental/Mover/Source/Mover/Public/Backends/MoverNetworkPredictionLiaison.h
+++ b/Engine/Plugins/Experimental/Mover/Source/Mover/Public/Backends/MoverNetworkPredictionLiaison.h
@@ -34,6 +34,12 @@ public:
 	// Get latest local input prior to simulation step. Called by Network Prediction system on owner's instance (autonomous or authority).
 	UE_API void ProduceInput(const int32 DeltaTimeMS, FMoverInputCmdContext* Cmd);
 
+	// Query whether a rollback / resim should occur based on how different the predicted and authority sync states were. Called by Network Prediction system.
+	UE_API bool ShouldReconcile(const FMoverSyncState* PredictedSyncState, const FMoverSyncState* AuthoritySyncState) const;
+
+	// Query whether a rollback / resim should occur based on how different the predicted and authority aux states were. Called by Network Prediction system.
+	UE_API bool ShouldReconcile(const FMoverAuxStateContext* PredictedAuxState, const FMoverAuxStateContext* AuthorityAuxState) const;
+
 	// Restore a previous frame prior to resimulating. Called by Network Prediction system.
 	UE_API void RestoreFrame(const FMoverSyncState* SyncState, const FMoverAuxStateContext* AuxState);
 
diff --git a/Engine/Plugins/Runtime/NetworkPrediction/Source/NetworkPrediction/Public/NetworkPredictionDriver.h b/Engine/Plugins/Runtime/NetworkPrediction/Source/NetworkPrediction/Public/NetworkPredictionDriver.h
--- a/Engine/Plugins/Runtime/NetworkPrediction/Source/NetworkPrediction/Public/NetworkPredictionDriver.h
+++ b/Engine/Plugins/Runtime/NetworkPrediction/Source/NetworkPrediction/Public/NetworkPredictionDriver.h
@@ -484,19 +484,73 @@ struct FNetworkPredictionDriverBase
 	//
 	//	Determines if Sync/Aux state have diverged enough to force a correction.
 	//	The default implementation will call ShouldReconcile on the state itself: bool FMySyncState::ShouldReconcile(const FMySyncState& Authority) const.
+	//  May also be defined as Driver member functions (separately for Synx/Aux):  bool MyDriver::ShouldReconcile(const FMySyncState* PredictedState, const FMySyncState* AuthorityState).
 	// -----------------------------------------------------------------------------------------------------------------------------------
-	static bool ShouldReconcile(const TSyncAuxPair<StateTypes>& Predicted, const TSyncAuxPair<StateTypes>& Authority)
+	static bool ShouldReconcile(DriverType* Driver, const TSyncAuxPair<StateTypes>& Predicted, const TSyncAuxPair<StateTypes>& Authority)
 	{
-		return ShouldReconcile(Predicted.Sync, Authority.Sync) || ShouldReconcile(Predicted.Aux, Authority.Aux);
+		return CallShouldReconcileSyncStateMemberFunc(Driver, Predicted.Sync, Authority.Sync) ||
+			   CallShouldReconcileAuxStateMemberFunc(Driver, Predicted.Aux, Authority.Aux);
 	}
 
-	template<typename StateType>
-	static bool ShouldReconcile(const StateType* Predicted, const StateType* Authority)
+	struct CShouldReconcileSyncStateMemberFuncable
+	{
+		template <typename InDriverType, typename...>
+		auto Requires(InDriverType* Driver, const SyncType* Predicted, const SyncType* Authority) -> decltype(
+			Driver->ShouldReconcile(Predicted, Authority)
+			);
+	};
+
+	struct CShouldReconcileAuxStateMemberFuncable
 	{
+		template <typename InDriverType, typename...>
+		auto Requires(InDriverType* Driver, const AuxType* Predicted, const AuxType* Authority) -> decltype(
+			Driver->ShouldReconcile(Predicted, Authority)
+			);
+	};
+
+	
+	static constexpr bool HasShouldReconcileSyncState = TModels_V<CShouldReconcileSyncStateMemberFuncable, DriverType, SyncType, SyncType>;
+	static constexpr bool HasShouldReconcileAuxState  = TModels_V<CShouldReconcileAuxStateMemberFuncable,  DriverType, AuxType,  AuxType>;
+
+
+	template<bool HasFunc = HasShouldReconcileSyncState>
+	static typename TEnableIf<HasFunc, bool>::Type CallShouldReconcileSyncStateMemberFunc(DriverType* Driver, const SyncType* Predicted, const SyncType* Authority)
+	{
+		npCheckSlow(Driver);
+		return Driver->ShouldReconcile(Predicted, Authority);
+	}
+
+	template<bool HasFunc = HasShouldReconcileSyncState>
+	static typename TEnableIf<!HasFunc, bool>::Type CallShouldReconcileSyncStateMemberFunc(DriverType* Driver, const SyncType* Predicted, const SyncType* Authority)
+	{
+		// ShouldReconcile member function isn't required on the driver, as we'll fall back to querying the StateTypes directly
 		return Predicted->ShouldReconcile(*Authority);
 	}
-	
-	static bool ShouldReconcile(const void* Predicted, const void* Authority) { return false; }
+
+	static bool CallShouldReconcileSyncStateMemberFunc(DriverType* Driver, const void* Predicted, const void* Authority)
+	{
+		return false;
+	}
+
+
+	template<bool HasFunc = HasShouldReconcileAuxState>
+	static typename TEnableIf<HasFunc, bool>::Type CallShouldReconcileAuxStateMemberFunc(DriverType* Driver, const AuxType* Predicted, const AuxType* Authority)
+	{
+		npCheckSlow(Driver);
+		return Driver->ShouldReconcile(Predicted, Authority);
+	}
+
+	template<bool HasFunc = HasShouldReconcileAuxState>
+	static typename TEnableIf<!HasFunc, bool>::Type CallShouldReconcileAuxStateMemberFunc(DriverType* Driver, const AuxType* Predicted, const AuxType* Authority)
+	{
+		// ShouldReconcile member function isn't required on the driver, as we'll fall back to querying the StateTypes directly
+		return Predicted->ShouldReconcile(*Authority);
+	}
+
+	static bool CallShouldReconcileAuxStateMemberFunc(DriverType* Driver, const void* Predicted, const void* Authority)
+	{
+		return false; 
+	}
 
 	// -----------------------------------------------------------------------------------------------------------------------------------
 	//	Interpolate
diff --git a/Engine/Plugins/Runtime/NetworkPrediction/Source/NetworkPrediction/Public/Services/NetworkPredictionService_Rollback.inl b/Engine/Plugins/Runtime/NetworkPrediction/Source/NetworkPrediction/Public/Services/NetworkPredictionService_Rollback.inl
--- a/Engine/Plugins/Runtime/NetworkPrediction/Source/NetworkPrediction/Public/Services/NetworkPredictionService_Rollback.inl
+++ b/Engine/Plugins/Runtime/NetworkPrediction/Source/NetworkPrediction/Public/Services/NetworkPredictionService_Rollback.inl
@@ -85,6 +85,8 @@ public:
 			const int32 LocalFrame = ClientRecvData.ServerFrame - Offset;
 			typename TInstanceFrameState<ModelDef>::FFrame& LocalFrameData = Frames.Buffer[LocalFrame];
 
+			TInstanceData<ModelDef>& InstanceData = DataStore->Instances.GetByIndexChecked(ClientRecvData.InstanceIdx);
+
 			bool bDoRollback = false;
 
 			if (NetworkPredictionCVars::ForceReconcile() > 0)
@@ -100,7 +102,7 @@ public:
 
 				NetworkPredictionCVars::SetForceReconcile(0); // reset
 			}
-			else if (FNetworkPredictionDriver<ModelDef>::ShouldReconcile( SyncAuxType(LocalFrameData.SyncState, LocalFrameData.AuxState), SyncAuxType(ClientRecvData.SyncState, ClientRecvData.AuxState) ))
+			else if (FNetworkPredictionDriver<ModelDef>::ShouldReconcile(InstanceData.Info.Driver, SyncAuxType(LocalFrameData.SyncState, LocalFrameData.AuxState), SyncAuxType(ClientRecvData.SyncState, ClientRecvData.AuxState)))
 			{
 				UE_NP_TRACE_SHOULD_RECONCILE(ClientRecvData.TraceID);
 				bDoRollback = true;
@@ -122,7 +124,6 @@ public:
 				const int32 PhysicsFrame = LocalFrame + TickState->PhysicsOffset;
 
 				// Maybe consider putting a copy of PhysicsActor on TClientRecvData if this lookup ever shows up in profiler
-				TInstanceData<ModelDef>& InstanceData = DataStore->Instances.GetByIndexChecked(ClientRecvData.InstanceIdx);
 
 				if (FNetworkPredictionDriver<ModelDef>::ShouldReconcilePhysics(PhysicsFrame, TickState->PhysicsRewindData, InstanceData.Info.Driver, ClientRecvData.Physics))
 				{
@@ -315,7 +316,9 @@ public:
 			const int32 LocalFrame = ClientRecvData.ServerFrame;
 			typename TInstanceFrameState<ModelDef>::FFrame& LocalFrameData = Frames.Buffer[LocalFrame];
 
-			if (FNetworkPredictionDriver<ModelDef>::ShouldReconcile( SyncAuxType(LocalFrameData.SyncState, LocalFrameData.AuxState), SyncAuxType(ClientRecvData.SyncState, ClientRecvData.AuxState) ))
+			TInstanceData<ModelDef>& Instance = DataStore->Instances.GetByIndexChecked(ClientRecvData.InstanceIdx);
+
+			if (FNetworkPredictionDriver<ModelDef>::ShouldReconcile( Instance.Info.Driver, SyncAuxType(LocalFrameData.SyncState, LocalFrameData.AuxState), SyncAuxType(ClientRecvData.SyncState, ClientRecvData.AuxState) ))
 			{
 				UE_NP_TRACE_SHOULD_RECONCILE(ClientRecvData.TraceID);
 				if (NetworkPredictionCVars::PrintReconciles())
@@ -332,8 +335,6 @@ public:
 				LocalFrameData.SyncState = ClientRecvData.SyncState;
 				LocalFrameData.AuxState = ClientRecvData.AuxState;
 
-				TInstanceData<ModelDef>& Instance = DataStore->Instances.GetByIndexChecked(ClientRecvData.InstanceIdx);
-
 				FNetworkPredictionDriver<ModelDef>::RestoreFrame(Instance.Info.Driver, LocalFrameData.SyncState.Get(), LocalFrameData.AuxState.Get());
 
 				// Do rollback
