Network Prediction: removes code related to physics integration that's been disused for years

#Synced-CL 50857662

diff --git a/Engine/Plugins/Runtime/NetworkPrediction/Source/NetworkPrediction/Private/NetworkPredictionPhysics.cpp b/Engine/Plugins/Runtime/NetworkPrediction/Source/NetworkPrediction/Private/NetworkPredictionPhysics.cpp
deleted file mode 100644
diff --git a/Engine/Plugins/Runtime/NetworkPrediction/Source/NetworkPrediction/Private/NetworkPredictionPhysicsComponent.cpp b/Engine/Plugins/Runtime/NetworkPrediction/Source/NetworkPrediction/Private/NetworkPredictionPhysicsComponent.cpp
deleted file mode 100644
diff --git a/Engine/Plugins/Runtime/NetworkPrediction/Source/NetworkPrediction/Private/NetworkPredictionTrace.cpp b/Engine/Plugins/Runtime/NetworkPrediction/Source/NetworkPrediction/Private/NetworkPredictionTrace.cpp
--- a/Engine/Plugins/Runtime/NetworkPrediction/Source/NetworkPrediction/Private/NetworkPredictionTrace.cpp
+++ b/Engine/Plugins/Runtime/NetworkPrediction/Source/NetworkPrediction/Private/NetworkPredictionTrace.cpp
@@ -141,10 +141,6 @@ UE_TRACE_EVENT_BEGIN(NetworkPrediction, AuxState)
 	UE_TRACE_EVENT_FIELD(UE::Trace::AnsiString, Value)
 UE_TRACE_EVENT_END()
 
-UE_TRACE_EVENT_BEGIN(NetworkPrediction, PhysicsState)
-	UE_TRACE_EVENT_FIELD(UE::Trace::AnsiString, Value)
-UE_TRACE_EVENT_END()
-
 // ---------------------------------------------------------------------------
 
 void FNetworkPredictionTrace::TraceSimulationCreated_Internal(FNetworkPredictionID ID, FStringBuilderBase& Builder)
@@ -230,12 +226,6 @@ void FNetworkPredictionTrace::TraceUserState_Internal(ETraceUserState StateType,
 				<< AuxState.Value(Builder.GetData(), Builder.Len());
 			break;
 		}
-		case ETraceUserState::Physics:
-		{
-			UE_TRACE_LOG(NetworkPrediction, PhysicsState, NetworkPredictionChannel)
-				<< PhysicsState.Value(Builder.GetData(), Builder.Len());
-			break;
-		}
 	}
 }
 
diff --git a/Engine/Plugins/Runtime/NetworkPrediction/Source/NetworkPrediction/Private/NetworkPredictionUtil.cpp b/Engine/Plugins/Runtime/NetworkPrediction/Source/NetworkPrediction/Private/NetworkPredictionUtil.cpp
deleted file mode 100644
diff --git a/Engine/Plugins/Runtime/NetworkPrediction/Source/NetworkPrediction/Private/NetworkPredictionWorldManager.cpp b/Engine/Plugins/Runtime/NetworkPrediction/Source/NetworkPrediction/Private/NetworkPredictionWorldManager.cpp
--- a/Engine/Plugins/Runtime/NetworkPrediction/Source/NetworkPrediction/Private/NetworkPredictionWorldManager.cpp
+++ b/Engine/Plugins/Runtime/NetworkPrediction/Source/NetworkPrediction/Private/NetworkPredictionWorldManager.cpp
@@ -2,22 +2,17 @@
 
 #include "NetworkPredictionWorldManager.h"
 #include "Engine/Engine.h"
-#include "Physics/Experimental/PhysScene_Chaos.h"
 #include "PBDRigidsSolver.h"
 #include "ChaosSolversModule.h"
 #include "NetworkPredictionReplicatedManager.h"
 #include "Services/NetworkPredictionService_Finalize.inl"
 #include "Services/NetworkPredictionService_Input.inl"
 #include "Services/NetworkPredictionService_Interpolate.inl"
-#include "Services/NetworkPredictionService_Physics.inl"
 #include "Services/NetworkPredictionService_Rollback.inl"
 #include "Services/NetworkPredictionService_ServerRPC.inl"
 
 #include UE_INLINE_GENERATED_CPP_BY_NAME(NetworkPredictionWorldManager)
 
-// Do extra checks to make sure Physics and GameThread (PrimitiveComponent) are in sync at verious points in the rollback process
-#define NP_ENSURE_PHYSICS_GT_SYNC !(UE_BUILD_SHIPPING || UE_BUILD_TEST)
-
 UNetworkPredictionWorldManager* UNetworkPredictionWorldManager::ActiveInstance=nullptr;
 
 // -----------------------------------------------------------------------------------------------
@@ -151,11 +146,6 @@ void UNetworkPredictionWorldManager::ReconcileSimulationsPostNetworkUpdate_Inter
 	// -------------------------------------------------------------------------
 	//	Fixed Tick rollback
 	// -------------------------------------------------------------------------
-	if (FixedTickState.PhysicsRewindData)
-	{
-		// Local PendingFrame = PhysicsFrame - PhysicsOffset
-		FixedTickState.PhysicsOffset = FixedTickState.PhysicsRewindData->CurrentFrame() - FixedTickState.PendingFrame;
-	}
 
 	// Does anyone need to rollback?
 	int32 RollbackFrame = INDEX_NONE;
@@ -181,24 +171,6 @@ void UNetworkPredictionWorldManager::ReconcileSimulationsPostNetworkUpdate_Inter
 			const int32 NumFrames = EndFrame - RollbackFrame;
 			npEnsureSlow(NumFrames > 0);
 
-			int32 PhysicsFrame = INDEX_NONE;
-			const bool bDoPhysics = (Physics.Solver != nullptr);
-
-			Chaos::EThreadingModeTemp PreThreading = Chaos::EThreadingModeTemp::SingleThread;
-
-			// ---------------------------------------------------
-			//	Rollback physics to local historic state (corrections not injected yet)
-			// ---------------------------------------------------
-			if (bDoPhysics)
-			{
-				PhysicsFrame = RollbackFrame + FixedTickState.PhysicsOffset;
-				check(false);
-				//npEnsure(FixedTickState.PhysicsRewindData->RewindToFrame(PhysicsFrame));
-
-				PreThreading = Physics.Solver->GetThreadingMode();
-				Physics.Solver->SetThreadingMode_External(Chaos::EThreadingModeTemp::SingleThread);
-			}
-
 			bool bFirstStep = true;
 
 			// Do rollback as necessary
@@ -217,35 +189,18 @@ void UNetworkPredictionWorldManager::ReconcileSimulationsPostNetworkUpdate_Inter
 				{
 					Ptr->PreStepRollback(Step, ServiceStep, FixedTickState.Offset, bFirstStep);
 				}
-				EnsurePhysicsGTSync(TEXT("PreStepRollback"));
 
 				// Run Sim ticks
 				for (TUniquePtr<IFixedRollbackService>& Ptr : Services.FixedRollback.Array)
 				{
 					Ptr->StepRollback(Step, ServiceStep);
 				}
-				EnsurePhysicsGTSync(TEXT("PreResimPhysics"));
-
-				// Advance physics
-				if (bDoPhysics)
-				{
-					AdvancePhysicsResimFrame(PhysicsFrame);
-					for (TUniquePtr<IPhysicsService>& Ptr : Services.FixedPhysics.Array)
-					{
-						Ptr->PostResimulate(&FixedTickState);
-					}
-				}
 
-				EnsurePhysicsGTSync(TEXT("PostResimulate"));
 				bFirstStep = false;
 			}
 
 			FixedTickState.PendingFrame = EndFrame;
 
-			if (bDoPhysics)
-			{
-				Physics.Solver->SetThreadingMode_External(PreThreading);
-			}
 		}
 		else if (RollbackFrame == FixedTickState.PendingFrame)
 		{
@@ -299,18 +254,13 @@ void UNetworkPredictionWorldManager::BeginNewSimulationFrame_Internal(float Delt
 	// -------------------------------------------------------------------------
 	//	Fixed Tick
 	// -------------------------------------------------------------------------
-	if (Physics.bUsingPhysics || Services.FixedTick.Array.Num() > 0)
+	if (Services.FixedTick.Array.Num() > 0)
 	{
 		QUICK_SCOPE_CYCLE_COUNTER(STAT_NP_TICK_FIXED);
 		TRACE_CPUPROFILER_EVENT_SCOPE(NetworkPrediction::FixedTick);
 
 		FixedTickState.UnspentTimeMS += fEngineFrameDeltaTimeMS;
 
-		// If we are using physics, we can only do one NP sim tick per frame.
-		// For accumulation fixed tick to work, physics needs to support it too
-		// see notes in NetworkPredictiontickState.h
-		const bool bSingleTick = Physics.bUsingPhysics;
-
 		while ((FixedTickState.UnspentTimeMS + KINDA_SMALL_NUMBER) >= FixedTickState.FixedStepRealTimeMS)
 		{
 			FixedTickState.UnspentTimeMS -= FixedTickState.FixedStepRealTimeMS;
@@ -354,12 +304,6 @@ void UNetworkPredictionWorldManager::BeginNewSimulationFrame_Internal(float Delt
 					Ptr->UpdateSmoothing(ServiceStep, &FixedTickState);
 				}
 			}
-
-			if (bSingleTick)
-			{
-				FixedTickState.UnspentTimeMS = 0.f;
-				break;
-			}
 		}
 	}
 
@@ -586,14 +530,6 @@ void UNetworkPredictionWorldManager::BeginNewSimulationFrame_Internal(float Delt
 		}
 	}
 
-	// -------------------------------------------------------------------------
-	// Fixed Physics (just does bookkeeping/tracing, not actual physics sim)
-	// -------------------------------------------------------------------------
-	for (TUniquePtr<IPhysicsService>& Ptr : Services.FixedPhysics.Array)
-	{
-		Ptr->PostNetworkPredictionFrame(&FixedTickState);
-	}
-
 	// -------------------------------------------------------------------------
 	// Call server RPC (common)
 	// -------------------------------------------------------------------------
@@ -606,72 +542,6 @@ void UNetworkPredictionWorldManager::BeginNewSimulationFrame_Internal(float Delt
 	}
 }
 
-void UNetworkPredictionWorldManager::InitPhysicsCapture()
-{
-	constexpr int32 PhysicsBufferSize = 64; // fixme
-
-	npCheckSlow(Physics.Solver != nullptr);
-	npCheckSlow(!Physics.bRecordingEnabled);
-
-	Physics.bRecordingEnabled = true;
-	Physics.Solver->EnableRewindCapture(PhysicsBufferSize, false);	//TODO: turn optimizations back on
-	FixedTickState.PhysicsRewindData = Physics.Solver->GetRewindData();
-}
-
-void UNetworkPredictionWorldManager::AdvancePhysicsResimFrame(int32& PhysicsFrame)
-{
-	Physics.Solver->AdvanceAndDispatch_External(FixedTickState.PhysicsRewindData->GetDeltaTimeForFrame(PhysicsFrame));
-	Physics.Solver->UpdateGameThreadStructures();
-	PhysicsFrame++;
-}
-
-bool UNetworkPredictionWorldManager::CanPhysicsFixTick() const
-{
-	// GEngine may be null in some contexts, like unit tests.
-	return (GEngine && GEngine->bUseFixedFrameRate) || Settings.bForceEngineFixTickForcePhysics;
-}
-
-void UNetworkPredictionWorldManager::SetUsingPhysics()
-{
-	if (!Physics.bUsingPhysics)
-	{
-		npCheckSlow(Physics.Solver == nullptr);
-		Physics.bUsingPhysics = true;
-
-		Physics.Module = FChaosSolversModule::GetModule();
-		Physics.Solver = GetWorld()->GetPhysicsScene()->GetSolver();
-
-		const UObject* ExistingOwner = Physics.Solver->GetOwner();
-		if (!ExistingOwner)
-		{
-			Physics.Solver->SetOwner(GetWorld());
-		}
-		else
-		{
-			// We don't care if the solver's owner is the actual world or not, just as long as the owner's GetWorld() matches our GetWorld().
-			npEnsureMsgf(ExistingOwner->GetWorld() == GetWorld(), TEXT("Mismatch worlds betwen Solver %s and NP %s"), *GetNameSafe(ExistingOwner->GetWorld()), *GetNameSafe(GetWorld()));
-		}
-
-		npCheckSlow(GEngine);
-		if (!GEngine->bUseFixedFrameRate)
-		{
-			npEnsure(Settings.bForceEngineFixTickForcePhysics);
-			GEngine->bUseFixedFrameRate = true;
-			GEngine->FixedFrameRate = Settings.FixedTickFrameRate;
-		}
-	}
-}
-
-void UNetworkPredictionWorldManager::EnsurePhysicsGTSync(const TCHAR* Context) const
-{
-#if NP_ENSURE_PHYSICS_GT_SYNC
-	for (const TUniquePtr<IPhysicsService>& Ptr : Services.FixedPhysics.Array)
-	{
-		Ptr->EnsureDataInSync(Context);
-	}
-#endif
-}
-
 ENetworkPredictionTickingPolicy UNetworkPredictionWorldManager::PreferredDefaultTickingPolicy() const
 {
 	return Settings.PreferredTickingPolicy;
diff --git a/Engine/Plugins/Runtime/NetworkPrediction/Source/NetworkPrediction/Public/NetworkPredictionCVars.h b/Engine/Plugins/Runtime/NetworkPrediction/Source/NetworkPrediction/Public/NetworkPredictionCVars.h
--- a/Engine/Plugins/Runtime/NetworkPrediction/Source/NetworkPrediction/Public/NetworkPredictionCVars.h
+++ b/Engine/Plugins/Runtime/NetworkPrediction/Source/NetworkPrediction/Public/NetworkPredictionCVars.h
@@ -97,34 +97,3 @@ struct FConditionalAutoConsoleRegister
 #define NETSIM_DEVCVAR_SHIPCONST_FLOAT(Var,Value,VarName,Help) NETSIM_DEVCVAR_FLOAT(Var,Value,VarName,Help)
 #endif
 
-
-
-// Temp
-
-namespace UE_NETWORK_PHYSICS
-{
-	NETSIM_DEVCVAR_SHIPCONST_INT(MockDebug, 0, "np2.Mock.MockDebug", "Enabled spammy log debugging of mock physics object state");
-
-	NETSIM_DEVCVAR_SHIPCONST_INT(bFutureInputs, 0, "np2.FutureInputs", "Enable FutureInputs feature");		 // Not reimplemented yet
-		
-	NETSIM_DEVCVAR_SHIPCONST_INT(bInputDecay, 0, "np2.InputDecay", "Enable Input Decay Feature");		
-	NETSIM_DEVCVAR_SHIPCONST_FLOAT(InputDecayRate, 0.99f, "np2.InputDecayRate", "Rate of input decay");
-
-	NETSIM_DEVCVAR_SHIPCONST_FLOAT(DragK, 1.12f, "np2.Mock.DragK", "Drag Coefficient (higher=more drag)");
-	NETSIM_DEVCVAR_SHIPCONST_FLOAT(MaxAngularVelocity, 30.f, "np2.Mock.MaxAngularVelocity", "Limits how fast character can possibly rotate.")
-	NETSIM_DEVCVAR_SHIPCONST_FLOAT(MovementK, 1.25, "np2.Mock.MovementK", "Movement Coefficient (higher=faster movement)");
-	NETSIM_DEVCVAR_SHIPCONST_FLOAT(RotationK, 1.25, "np2.Mock.RotationK", "Rotation Coefficient (higher=faster movement)");
-	NETSIM_DEVCVAR_SHIPCONST_FLOAT(TurnK, 100000000.f, "np2.Mock.TurnK", "Coefficient for automatic turning (higher=quicker turning)");
-	NETSIM_DEVCVAR_SHIPCONST_FLOAT(TurnDampK, 100.f, "np2.Mock.TurnDampK", "Coefficient for damping portion of turn. Higher=more damping but too higher will lead to instability.");
-	NETSIM_DEVCVAR_SHIPCONST_FLOAT(DampYawVelocityK, 2.f, "np2.Mock.DampYawVelocityK", "Coefficient for damping angular velocity in yaw direction only. This is only enabled when auto target yaw is disabled.");
-
-	NETSIM_DEVCVAR_SHIPCONST_FLOAT(JumpForce, 70, "np2.Mock.JumpForce", "Per-Frame force to apply while jumping.");
-	NETSIM_DEVCVAR_SHIPCONST_INT(JumpFrameDuration, 4, "np2.Mock.JumpFrameDuration", "How many frames to apply jump force for");
-	NETSIM_DEVCVAR_SHIPCONST_INT(JumpFudgeFrames, 10, "np2.Mock.JumpFudgeFrames", "How many frames after being in air do we still allow a jump to begin");
-	NETSIM_DEVCVAR_SHIPCONST_INT(JumpHack, 0, "np2.Mock.JumpHack", "Make jump not rely on trace which currently causes non determinism");
-	NETSIM_DEVCVAR_SHIPCONST_INT(JumpMisPredict, 0, "np2.Mock.JumpMisPredict", "Make jump do random impulse which will cause a misprediction");
-	
-	NETSIM_DEVCVAR_SHIPCONST_INT(MockImpulse, 1, "np2.Mock.BallImpulse", "Make jump not rely on trace which currently causes non determinism");
-	NETSIM_DEVCVAR_SHIPCONST_FLOAT(MockImpulseX, 500.0f, "np2.Mock.BallImpulse.X", "X magnitude");
-	NETSIM_DEVCVAR_SHIPCONST_FLOAT(MockImpulseZ, 300.0f, "np2.Mock.BallImpulse.Z", "Z magnitude");
-}
diff --git a/Engine/Plugins/Runtime/NetworkPrediction/Source/NetworkPrediction/Public/NetworkPredictionConfig.h b/Engine/Plugins/Runtime/NetworkPrediction/Source/NetworkPrediction/Public/NetworkPredictionConfig.h
--- a/Engine/Plugins/Runtime/NetworkPrediction/Source/NetworkPrediction/Public/NetworkPredictionConfig.h
+++ b/Engine/Plugins/Runtime/NetworkPrediction/Source/NetworkPrediction/Public/NetworkPredictionConfig.h
@@ -10,7 +10,7 @@ enum class ENetworkPredictionTickingPolicy : uint8
 {
 	// Client ticks at local frame rate. Server ticks clients independently at client input cmd rate.
 	Independent	= 1 << 0,
-	// Everyone ticks at same fixed rate. Supports group rollback and physics.
+	// Everyone ticks at same fixed rate. Supports group rollback.
 	Fixed		= 1 << 1,
 
 	All = Independent | Fixed UMETA(Hidden),
diff --git a/Engine/Plugins/Runtime/NetworkPrediction/Source/NetworkPrediction/Public/NetworkPredictionDriver.h b/Engine/Plugins/Runtime/NetworkPrediction/Source/NetworkPrediction/Public/NetworkPredictionDriver.h
--- a/Engine/Plugins/Runtime/NetworkPrediction/Source/NetworkPrediction/Public/NetworkPredictionDriver.h
+++ b/Engine/Plugins/Runtime/NetworkPrediction/Source/NetworkPrediction/Public/NetworkPredictionDriver.h
@@ -67,26 +67,18 @@ struct FNetworkPredictionDriverBase
 	using InputType = typename StateTypes::InputType;
 	using SyncType = typename StateTypes::SyncType;
 	using AuxType = typename StateTypes::AuxType;
-	using PhysicsState = typename ModelDef::PhysicsState;
 
 	static constexpr bool HasNpState() { return !std::is_void_v<InputType> || !std::is_void_v<SyncType> || !std::is_void_v<AuxType>; }
 	static constexpr bool HasDriver() { return !std::is_void_v<DriverType>; }
 	static constexpr bool HasSimulation() { return !std::is_void_v<Simulation>; }
 	static constexpr bool HasInput() { return !std::is_void_v<InputType>; }
-	static constexpr bool HasPhysics() { return !std::is_void_v<PhysicsState>; }
 
 	// Defines what the ModelDef can do. This is a compile time thing only.
 	static constexpr FNetworkPredictionModelDefCapabilities GetCapabilities()
 	{
 		FNetworkPredictionModelDefCapabilities Capabilities;
-		if (HasPhysics())
-		{
-			// If you have physics state, you must use fixed tick
-			// (except if you want to dual support physics/non physics within the same simulation type, this will have to be a customization override)
-			Capabilities.SupportedTickingPolicies = ENetworkPredictionTickingPolicy::Fixed;
-		}
 
-		if (HasSimulation() == false && HasPhysics() == false)
+		if (HasSimulation() == false)
 		{
 			// We have nothing to tick, so no SimExtrapolation or ForwardPrediction
 			Capabilities.FixedNetworkLODs.AP = ENetworkLOD::Interpolated;
@@ -98,14 +90,9 @@ struct FNetworkPredictionDriverBase
 	}
 
 	// Defines the default settings for a spawn instance.
-	static bool GetDefaultArchetype(FNetworkPredictionInstanceArchetype& Archetype, ENetworkPredictionTickingPolicy PreferredTickingPolicy, bool bCanPhysicsTickPhysics)
+	static bool GetDefaultArchetype(FNetworkPredictionInstanceArchetype& Archetype, ENetworkPredictionTickingPolicy PreferredTickingPolicy)
 	{
 		static constexpr FNetworkPredictionModelDefCapabilities Capabilities = FNetworkPredictionDriver<ModelDef>::GetCapabilities();
-		if (HasPhysics() && bCanPhysicsTickPhysics == false && !EnumHasAnyFlags(Capabilities.SupportedTickingPolicies, ENetworkPredictionTickingPolicy::Independent))
-		{
-			// We require fix ticking physics to function and its not enabled. Return false.
-			return false;
-		}
 
 		// Use preferred ticking policy if we support it
 		if (EnumHasAnyFlags(Capabilities.SupportedTickingPolicies, PreferredTickingPolicy))
@@ -351,7 +338,7 @@ struct FNetworkPredictionDriverBase
 	template<bool HasFunc=HasFinalizeFrame>
 	static typename TEnableIf<!HasFunc>::Type CallFinalizeFrameMemberFunc(DriverType* Driver, const SyncType* SyncState, const AuxType* AuxState)
 	{
-		// This is only fatal if we have NpState. If we just had physics for example, a FinalizeFrame call isn't needed.
+		// This is only fatal if we have NpState. If we had a simulation without NpState, a FinalizeFrame call isn't needed.
 		npCheckf(!HasNpState(), TEXT("No FinalizeFrame implementation found. Implement DriverType::FinalizeFrame or ModelDef::FinalizeFrame"));
 	}
 
@@ -393,15 +380,10 @@ struct FNetworkPredictionDriverBase
 	//	Called prior to beginning rollback frames. This instance should put itself in whatever state it needs to be in for resimulation to
 	//	run. In practice this should mean getting right collision+component states in sync so that any scene queries will get the correct
 	//	data.
-	//
-	//	This can be automated for physics (the PhysicsState can generically marshal the data). We can't do this for kinematic sims.
-	//	
-	//	
 	// -----------------------------------------------------------------------------------------------------------------------------------
 
 	static void RestoreFrame(DriverType* Driver, const SyncType* SyncState, const AuxType* AuxState)
 	{
-		FNetworkPredictionDriver<ModelDef>::MarshalPhysicsToComponent(Driver);
 		FNetworkPredictionDriver<ModelDef>::CallRestoreFrameMemberFunc(Driver, SyncState, AuxState);
 	}
 
@@ -615,295 +597,6 @@ struct FNetworkPredictionDriverBase
 		Driver->GetOwner()->SetActorHiddenInGame(bHide);
 	}	
 
-	// -----------------------------------------------------------------------------------------------------------------------------------
-	//	Physics
-	//
-	//	Note: a ModelDef can support physics but if an instance is running in independent ticking mode, physics reconcile/net serialization will
-	//	effectively be off. (It doesn't make sense to reconcile physics state that is ticked out of step from the independent simulation tick).
-	//	So below will only be "on" when in fixed tick mode.
-	//
-	//	UPrimitiveComponent::SyncComponentToRBPhysics
-	// -----------------------------------------------------------------------------------------------------------------------------------
-
-	// ------------------------------------------
-	//	GetPhysicsPrimitiveComponent: returns UPrimitiveComponent to forward to PhysicsState function calls.
-	//	(it would be nice to not require this but the interpolatin path currently requires calls to kinematic bodies to go through the component, not PhysicsActorHandle)
-	// ------------------------------------------
-	
-	struct CGetPhysicsPrimitiveComponentMemberFuncable
-	{
-		template <typename InDriverType>
-		auto Requires(InDriverType* Driver) -> decltype(Driver->GetPhysicsPrimitiveComponent());
-	};
-
-	static constexpr bool HasGetPrimitiveComp = TModels_V<CGetPhysicsPrimitiveComponentMemberFuncable, DriverType>;
-	
-	template<bool HasFunc=HasGetPrimitiveComp>
-	static typename TEnableIf<HasFunc, UPrimitiveComponent*>::Type GetPhysicsPrimitiveComponent(DriverType* Driver)
-	{
-		npCheckSlow(Driver);
-		return Driver->GetPhysicsPrimitiveComponent();
-	}
-	
-	template<bool HasFunc=HasGetPrimitiveComp>
-	static typename TEnableIf<!HasFunc, UPrimitiveComponent*>::Type GetPhysicsPrimitiveComponent(DriverType* Driver)
-	{
-		return SafeCastDriverToPrimitiveComponent(Driver);
-	}	
-
-	template<typename InDriverType>
-	static UPrimitiveComponent* SafeCastDriverToPrimitiveComponent(InDriverType* Driver)
-	{
-		//	If you are landing here with a nullptr Driver:
-		//	You are using interpolated physics, we need to get a UPrimitiveComponent from your ModelDef Driver. 3 choices:
-		//		-A) Your ModelDef::Driver could be a UPrimitiveComponent (or inherit from one). Nothing else required.
-		//		-B) Your ModelDef::Driver class can implement a UPrimitiveComponent* GetPhysicsPrimitiveComponent() function and return it. Good for "UpdatedPrimitive pattern".
-		//		-C) You can specialize FNetworkPredictionDriver<ModelDef>::GetPhysicsPrimitiveComponent to do something else. Good for non intrusive case.
-		npCheckf(false, TEXT("Could not obtain UPrimitiveComponent from driver. See notes above in NetworkPredictionDriver.h"));
-		return nullptr;
-	}
-
-	static UPrimitiveComponent* SafeCastDriverToPrimitiveComponent(UPrimitiveComponent* Driver)
-	{
-		return Driver;
-	}
-
-	// ------------------------------------------
-	//	GetPhysicsBodyInstance
-	// ------------------------------------------
-
-	struct CGetPhysicsBodyInstanceMemberFuncable
-	{
-		template <typename InDriverType>
-		auto Requires(InDriverType* Driver) -> decltype(Driver->GetPhysicsBodyInstance());
-	};
-
-	static constexpr bool HasGetBodyInstance = TModels_V<CGetPhysicsBodyInstanceMemberFuncable, DriverType>;
-
-	template<bool HasFunc=HasGetBodyInstance>
-	static typename TEnableIf<HasFunc, FBodyInstance*>::Type GetPhysicsBodyInstance(DriverType* Driver)
-	{
-		npCheckSlow(Driver);
-		return Driver->GetPhysicsBodyInstance();
-	}
-
-	template<bool HasFunc=HasGetBodyInstance>
-	static typename TEnableIf<!HasFunc, FBodyInstance*>::Type GetPhysicsBodyInstance(DriverType* Driver)
-	{
-		// Non explicit version: get the primitive component and get the body instance off of that
-		UPrimitiveComponent* PrimitiveComponent = FNetworkPredictionDriver<ModelDef>::GetPhysicsPrimitiveComponent(Driver);
-		npCheckSlow(PrimitiveComponent);
-		return PrimitiveComponent->GetBodyInstance();
-	}
-
-	// ------------------------------------------
-	//	ShouldReconcilePhysics
-	// ------------------------------------------
-	template<bool bEnable=HasPhysics()>
-	static typename TEnableIf<bEnable, bool>::Type ShouldReconcilePhysics(int32 PhysicsFrame, Chaos::FRewindData* RewindData, DriverType* Driver, TConditionalState<PhysicsState>& RecvState)
-	{
-		return PhysicsState::ShouldReconcile(PhysicsFrame, RewindData, FNetworkPredictionDriver<ModelDef>::GetPhysicsBodyInstance(Driver), RecvState);
-	}
-
-	template<bool bEnable=HasPhysics()>
-	static typename TEnableIf<!bEnable, bool>::Type ShouldReconcilePhysics(int32 PhysicsFrame, Chaos::FRewindData* RewindData, DriverType* Driver, TConditionalState<PhysicsState>& RecvState)
- 	{
-		return false;
-	}
-
-	// ------------------------------------------
-	//	PerformPhysicsRollback
-	// ------------------------------------------
-
-	template<bool bEnable=HasPhysics()>
-	static typename TEnableIf<bEnable>::Type PerformPhysicsRollback(DriverType* Driver, TConditionalState<PhysicsState>& RecvState)
-	{
-		PhysicsState::PerformRollback(FNetworkPredictionDriver<ModelDef>::GetPhysicsBodyInstance(Driver), RecvState);
-	}
-
-	template<bool bEnable=HasPhysics()>
-	static typename TEnableIf<!bEnable>::Type PerformPhysicsRollback(DriverType* Driver, TConditionalState<PhysicsState>& RecvState) {}
-
-	// ------------------------------------------
-	//	PostPhysicsResimulate
-	// ------------------------------------------
-	template<bool bEnable=HasPhysics()>
-	static typename TEnableIf<bEnable>::Type PostPhysicsResimulate(DriverType* Driver)
-	{
-		PhysicsState::PostResimulate(FNetworkPredictionDriver<ModelDef>::GetPhysicsPrimitiveComponent(Driver), FNetworkPredictionDriver<ModelDef>::GetPhysicsBodyInstance(Driver));
-	}
-
-	template<bool bEnable=HasPhysics()>
-	static typename TEnableIf<!bEnable>::Type PostPhysicsResimulate(DriverType* Driver)
-	{
-		npCheck(false); // Not expected to be called in non physics based sims
-	}
-
-	// ------------------------------------------
-	//	MarshalPhysicsToComponent
-	// ------------------------------------------
-	template<bool bEnable=HasPhysics()>
-	static typename TEnableIf<bEnable>::Type MarshalPhysicsToComponent(DriverType* Driver)
-	{
-		npCheckSlow(Driver);
-		PhysicsState::MarshalPhysicsToComponent(FNetworkPredictionDriver<ModelDef>::GetPhysicsPrimitiveComponent(Driver), FNetworkPredictionDriver<ModelDef>::GetPhysicsBodyInstance(Driver));
-	}
-
-	template<bool bEnable=HasPhysics()>
-	static typename TEnableIf<!bEnable>::Type MarshalPhysicsToComponent(DriverType* Driver) { }
-
-	// ------------------------------------------
-	//	InterpolatePhysics
-	// ------------------------------------------
-	template<bool bEnable=HasPhysics()>
-	static typename TEnableIf<bEnable>::Type InterpolatePhysics(TConditionalState<PhysicsState>& FromState, TConditionalState<PhysicsState>& ToState, float PCT, TConditionalState<PhysicsState>& OutState)
-	{
-		PhysicsState::Interpolate(FromState, ToState, PCT, OutState);
-	}
-
-	template<bool bEnable=HasPhysics()>
-	static typename TEnableIf<!bEnable>::Type InterpolatePhysics(TConditionalState<PhysicsState>& FromState, TConditionalState<PhysicsState>& ToState, float PCT, TConditionalState<PhysicsState>& OutState) {}
-
-	// ------------------------------------------
-	//	FinalizeInterpolatedPhysics
-	//	Only called when needing to push NP interpolated physics state to the driver.
-	//	E.g, the physics sim is not running
-	// ------------------------------------------
-	template<bool bEnable=HasPhysics()>
-	static typename TEnableIf<bEnable>::Type FinalizeInterpolatedPhysics(DriverType* Driver, TConditionalState<PhysicsState>& Physics)
-	{
-		PhysicsState::FinalizeInterpolatedPhysics(FNetworkPredictionDriver<ModelDef>::GetPhysicsPrimitiveComponent(Driver), Physics);
-	}
-
-	template<bool bEnable=HasPhysics()>
-	static typename TEnableIf<!bEnable>::Type FinalizeInterpolatedPhysics(DriverType* Driver, TConditionalState<PhysicsState>& Physics) {}
-
-	// ------------------------------------------
-	//	BeginInterpolatedPhysics - turn off physics simulation so that NP can interpolate the physics state
-	// ------------------------------------------
-	template<bool bEnable=HasPhysics()>
-	static typename TEnableIf<bEnable>::Type BeginInterpolatedPhysics(DriverType* Driver)
-	{
-		PhysicsState::BeginInterpolation(FNetworkPredictionDriver<ModelDef>::GetPhysicsPrimitiveComponent(Driver));
-	}
-
-	template<bool bEnable=HasPhysics()>
-	static typename TEnableIf<!bEnable>::Type BeginInterpolatedPhysics(DriverType* Driver) {}
-
-	// ------------------------------------------
-	//	EndInterpolatedPhysics - turn on physics simulation
-	// ------------------------------------------
-	template<bool bEnable=HasPhysics()>
-	static typename TEnableIf<bEnable>::Type EndInterpolatedPhysics(DriverType* Driver)
-	{
-		PhysicsState::EndInterpolation(FNetworkPredictionDriver<ModelDef>::GetPhysicsPrimitiveComponent(Driver));
-	}
-
-	template<bool bEnable=HasPhysics()>
-	static typename TEnableIf<!bEnable>::Type EndInterpolatedPhysics(DriverType* Driver) {}
-
-	// ------------------------------------------
-	//	PhysicsNetSend
-	// ------------------------------------------
-	template<bool bEnable=HasPhysics()>
-	static typename TEnableIf<bEnable>::Type PhysicsNetSend(const FNetSerializeParams& P, DriverType* Driver)
-	{
-		PhysicsState::NetSend(P, FNetworkPredictionDriver<ModelDef>::GetPhysicsBodyInstance(Driver));
-	}
-
-	template<bool bEnable=HasPhysics()>
-	static typename TEnableIf<!bEnable>::Type PhysicsNetSend(const FNetSerializeParams& P, DriverType* Driver) {}
-	
-	// ------------------------------------------
-	//	PhysicsNetRecv
-	// ------------------------------------------
-	template<bool bEnable=HasPhysics()>
-	static typename TEnableIf<bEnable>::Type PhysicsNetRecv(const FNetSerializeParams& P, TConditionalState<PhysicsState>& RecvState)
-	{
-		PhysicsState::NetRecv(P, RecvState);
-	}
-
-	template<bool bEnable=HasPhysics()>
-	static typename TEnableIf<!bEnable>::Type PhysicsNetRecv(const FNetSerializeParams& P, TConditionalState<PhysicsState>& RecvState) {}
-
-	// ------------------------------------------
-	//	PhysicsStateIsConsistent
-	//	Only used in debugging - checks if UPrimitiveComponent and FPhysicsBody state are in sync with each other
-	// ------------------------------------------
-	template<bool bEnable=HasPhysics()>
-	static typename TEnableIf<bEnable, bool>::Type PhysicsStateIsConsistent(DriverType* Driver)
-	{
-		return PhysicsState::StateIsConsistent(FNetworkPredictionDriver<ModelDef>::GetPhysicsPrimitiveComponent(Driver),
-												FNetworkPredictionDriver<ModelDef>::GetPhysicsBodyInstance(Driver));
-	}
-
-	template<bool bEnable=HasPhysics()>
-	static typename TEnableIf<!bEnable, bool>::Type PhysicsStateIsConsistent(DriverType* Driver)
-	{
-		npCheck(false);
-		return false;
-	}
-
-	// ------------------------------------------
-	//	LogPhysicsState
-	// ------------------------------------------
-	template<bool bEnable=HasPhysics()>
-	static typename TEnableIf<bEnable>::Type LogPhysicsState(TConditionalState<PhysicsState>& RecvState, FOutputDevice& Ar=*GLog)
-	{
-		TAnsiStringBuilder<256> Builder;
-		PhysicsState::ToString(RecvState, Builder);
-
-		Ar.Log(StringCast<TCHAR>(Builder.ToString()).Get());
-	}
-
-	template<bool bEnable=HasPhysics()>
-	static typename TEnableIf<!bEnable>::Type LogPhysicsState(TConditionalState<PhysicsState>& RecvState, FOutputDevice& Ar=*GLog) { }
-
-	template<bool bEnable=HasPhysics()>
-	static typename TEnableIf<bEnable>::Type LogPhysicsState(int32 PhysicsFrame, Chaos::FRewindData* RewindData, DriverType* Driver, FOutputDevice& Ar=*GLog)
-	{		
-		TAnsiStringBuilder<256> Builder;
-		PhysicsState::ToString(PhysicsFrame, RewindData, FNetworkPredictionDriver<ModelDef>::GetPhysicsBodyInstance(Driver), Builder);
-
-		Ar.Log(StringCast<TCHAR>(Builder.ToString()).Get());
-	}
-
-	template<bool bEnable=HasPhysics()>
-	static typename TEnableIf<!bEnable>::Type LogPhysicsState(int32 PhysicsFrame, Chaos::FRewindData* RewindData, DriverType* Driver, FOutputDevice& Ar=*GLog) {}
-
-	// ------------------------------------------
-	//	TracePhysicsState
-	// ------------------------------------------
-	
-	// Received physics state
-	template<bool bEnable=HasPhysics()>
-	static typename TEnableIf<bEnable>::Type TracePhysicsStateRecv(const TConditionalState<PhysicsState>& RecvState, FAnsiStringBuilderBase& Builder)
-	{
-		PhysicsState::ToString(RecvState.Get(), Builder);
-	}
-
-	template<bool bEnable=HasPhysics()>
-	static typename TEnableIf<!bEnable>::Type TracePhysicsStateRecv(const TConditionalState<PhysicsState>& RecvState, FAnsiStringBuilderBase& Builder) { }
-
-	// Current latest physics state
-	template<bool bEnable=HasPhysics()>
-	static typename TEnableIf<bEnable>::Type TracePhysicsState(DriverType* Driver, FAnsiStringBuilderBase& Builder)
-	{
-		PhysicsState::ToString(FNetworkPredictionDriver<ModelDef>::GetPhysicsBodyInstance(Driver), Builder);
-	}
-
-	template<bool bEnable=HasPhysics()>
-	static typename TEnableIf<!bEnable>::Type TracePhysicsState(DriverType* Driver, FAnsiStringBuilderBase& Builder) {}
-
-	// Local historic physics state
-	template<bool bEnable=HasPhysics()>
-	static typename TEnableIf<bEnable>::Type TracePhysicsState(int32 PhysicsFrame, Chaos::FRewindData* RewindData, DriverType* Driver, FAnsiStringBuilderBase& Builder)
-	{
-		PhysicsState::ToString(PhysicsFrame, RewindData, FNetworkPredictionDriver<ModelDef>::GetPhysicsBodyInstance(Driver), Builder);
-	}
-
-	template<bool bEnable=HasPhysics()>
-	static typename TEnableIf<!bEnable>::Type TracePhysicsState(int32 PhysicsFrame, Chaos::FRewindData* RewindData, DriverType* Driver, FAnsiStringBuilderBase& Builder) {}
 
 	// -----------------------------------------------------------------------------------------------------------------------------------
 	//	ToString
diff --git a/Engine/Plugins/Runtime/NetworkPrediction/Source/NetworkPrediction/Public/NetworkPredictionModelDef.h b/Engine/Plugins/Runtime/NetworkPrediction/Source/NetworkPrediction/Public/NetworkPredictionModelDef.h
--- a/Engine/Plugins/Runtime/NetworkPrediction/Source/NetworkPrediction/Public/NetworkPredictionModelDef.h
+++ b/Engine/Plugins/Runtime/NetworkPrediction/Source/NetworkPrediction/Public/NetworkPredictionModelDef.h
@@ -12,7 +12,7 @@ enum class ENetworkPredictionSortPriority : uint8
 	PreKinematicMovers = 50,
 	KinematicMovers = 75,
 	PostKinematicMovers = 100,
-	Physics = 125,	// Note this not where physics itself *ticks*! Just a priority value for physics having definitions to be sorted in the various services.
+
 	Last = 250
 };
 
@@ -35,14 +35,10 @@ struct FNetworkPredictionModelDef
 
 	// Object class that can take output from prediction system. E.g AActor, AMyPawn.
 	// See notes in NetworkPredictionDriver.h
-	// Requires: StateTypes || PhysicsState
+	// Requires: StateTypes
 	// Enables: Finalize, Cues
 	using Driver = void;
 
-	// Physics state. Void = no physics, FNetworkPredictionPhysicsState is the default physics state that synchronizes X,R,V,W.
-	// Enables: Reconcile, Finalize
-	using PhysicsState = void;
-
 	static const TCHAR* GetName() { return nullptr; }
 
 	static constexpr int32 GetSortPriority() { return (int32)ENetworkPredictionSortPriority::Last; }
@@ -73,7 +69,6 @@ struct TNetworkPredictionModelInfo
 {
 	using SimulationType = typename ModelDef::Simulation;
 	using DriverType = typename ModelDef::Driver;
-	using PhysicsState = typename ModelDef::PhysicsState;
 
 	FConditionalSimulationPtr<ModelDef> Simulation;			// Object that ticks this instance
 	DriverType* Driver = nullptr;							// Object that handles input/out
diff --git a/Engine/Plugins/Runtime/NetworkPrediction/Source/NetworkPrediction/Public/NetworkPredictionPhysics.h b/Engine/Plugins/Runtime/NetworkPrediction/Source/NetworkPrediction/Public/NetworkPredictionPhysics.h
deleted file mode 100644
diff --git a/Engine/Plugins/Runtime/NetworkPrediction/Source/NetworkPrediction/Public/NetworkPredictionPhysicsComponent.h b/Engine/Plugins/Runtime/NetworkPrediction/Source/NetworkPrediction/Public/NetworkPredictionPhysicsComponent.h
deleted file mode 100644
diff --git a/Engine/Plugins/Runtime/NetworkPrediction/Source/NetworkPrediction/Public/NetworkPredictionProxyInit.h b/Engine/Plugins/Runtime/NetworkPrediction/Source/NetworkPrediction/Public/NetworkPredictionProxyInit.h
--- a/Engine/Plugins/Runtime/NetworkPrediction/Source/NetworkPrediction/Public/NetworkPredictionProxyInit.h
+++ b/Engine/Plugins/Runtime/NetworkPrediction/Source/NetworkPrediction/Public/NetworkPredictionProxyInit.h
@@ -25,9 +25,8 @@ void FNetworkPredictionProxy::Init(const FInitParams<ModelDef>& Params)
 	if (ID.IsValid() == false)
 	{
 		// Brand new registration. Initialize the default Archetype for this ModelDef
-		if (!FNetworkPredictionDriver<ModelDef>::GetDefaultArchetype(CachedArchetype, WorldManager->PreferredDefaultTickingPolicy(), WorldManager->CanPhysicsFixTick()))
+		if (!FNetworkPredictionDriver<ModelDef>::GetDefaultArchetype(CachedArchetype, WorldManager->PreferredDefaultTickingPolicy()))
 		{
-			// This can fail, for example the ModelDef requires fixed tick physics but the engine is not set to fix ticking
 			UE_LOG(LogNetworkPrediction, Error, TEXT("Unable to initialize DefaultArchetype. Skipping registration with NetworkPrediction"));
 			return;
 		}
diff --git a/Engine/Plugins/Runtime/NetworkPrediction/Source/NetworkPrediction/Public/NetworkPredictionSerialization.h b/Engine/Plugins/Runtime/NetworkPrediction/Source/NetworkPrediction/Public/NetworkPredictionSerialization.h
--- a/Engine/Plugins/Runtime/NetworkPrediction/Source/NetworkPrediction/Public/NetworkPredictionSerialization.h
+++ b/Engine/Plugins/Runtime/NetworkPrediction/Source/NetworkPrediction/Public/NetworkPredictionSerialization.h
@@ -345,7 +345,7 @@ private:
 // ---------------------------------------------------------------------------------------------------------------------------
 //	Server -> AP Client
 //
-//	The Fixed/Independent ticking cases differ a bit but still share the same the core payload: Sync/Aux/Physics/Cues.
+//	The Fixed/Independent ticking cases differ a bit but still share the same the core payload: Sync/Aux/Cues.
 //	Fixed tick sends last consumed client input frame # ANd the server frame in order to correlate client/server frame numbers.
 //	Independent tick sends last consumed client input frame # + TotalSimTime in order to detect dropped frames.
 //
@@ -364,12 +364,9 @@ public:
 
 		FNetworkPredictionDriver<ModelDef>::NetSerialize(ClientRecvState.SyncState, P);	// 1. Sync
 		FNetworkPredictionDriver<ModelDef>::NetSerialize(ClientRecvState.AuxState, P);	// 2. Aux
-
-		FNetworkPredictionDriver<ModelDef>::PhysicsNetRecv(P, ClientRecvState.Physics); // 3. Physics
 		
 		UE_NP_TRACE_USER_STATE_SYNC(ModelDef, ClientRecvState.SyncState.Get());
 		UE_NP_TRACE_USER_STATE_AUX(ModelDef, ClientRecvState.AuxState.Get());
-		UE_NP_TRACE_PHYSICS_STATE_RECV(ModelDef, ClientRecvState.Physics);
 	}
 
 	static void NetSend(const FNetSerializeParams& P, TInstanceData<ModelDef>& InstanceData, typename TInstanceFrameState<ModelDef>::FFrame& FrameData)
@@ -378,8 +375,6 @@ public:
 
 		FNetworkPredictionDriver<ModelDef>::NetSerialize(FrameData.SyncState, P);	// 1. Sync
 		FNetworkPredictionDriver<ModelDef>::NetSerialize(FrameData.AuxState, P);	// 2. Aux
-
-		FNetworkPredictionDriver<ModelDef>::PhysicsNetSend(P, InstanceData.Info.Driver); // 3. Physics
 	}
 };
 
@@ -520,7 +515,7 @@ public:
 
 // ---------------------------------------------------------------------------------------------------------------------------
 //	Server -> SP Client
-//	Like the AP case, the core payload is the same between fixed and independent: Input/Sync/Aux/Physics/Cues
+//	Like the AP case, the core payload is the same between fixed and independent: Input/Sync/Aux/Cues
 //
 //	There are actually 3 cases to consider here:
 //	1. Fixed Tick: only sends server frame #
@@ -546,12 +541,9 @@ public:
 		FNetworkPredictionDriver<ModelDef>::NetSerialize(ClientRecvState.SyncState, P);	// 2. Sync
 		FNetworkPredictionDriver<ModelDef>::NetSerialize(ClientRecvState.AuxState, P);	// 3. Aux
 
-		FNetworkPredictionDriver<ModelDef>::PhysicsNetRecv(P, ClientRecvState.Physics); // 4. Physics
-
 		UE_NP_TRACE_USER_STATE_INPUT(ModelDef, ClientRecvState.InputCmd.Get());
 		UE_NP_TRACE_USER_STATE_SYNC(ModelDef, ClientRecvState.SyncState.Get());
 		UE_NP_TRACE_USER_STATE_AUX(ModelDef, ClientRecvState.AuxState.Get());
-		UE_NP_TRACE_PHYSICS_STATE_RECV(ModelDef, ClientRecvState.Physics);
 	}
 	
 	static void NetSend(const FNetSerializeParams& P, FNetworkPredictionID ID, TModelDataStore<ModelDef>* DataStore, TInstanceData<ModelDef>* InstanceData, int32 PendingFrame)
@@ -568,8 +560,6 @@ public:
 		FNetworkPredictionDriver<ModelDef>::NetSerialize(FrameData.InputCmd, P);	// 1. Input
 		FNetworkPredictionDriver<ModelDef>::NetSerialize(FrameData.SyncState, P);	// 2. Sync
 		FNetworkPredictionDriver<ModelDef>::NetSerialize(FrameData.AuxState, P);	// 3. Aux
-
-		FNetworkPredictionDriver<ModelDef>::PhysicsNetSend(P, InstanceData->Info.Driver); // 4. Physics
 	}
 };
 
diff --git a/Engine/Plugins/Runtime/NetworkPrediction/Source/NetworkPrediction/Public/NetworkPredictionSettings.h b/Engine/Plugins/Runtime/NetworkPrediction/Source/NetworkPrediction/Public/NetworkPredictionSettings.h
--- a/Engine/Plugins/Runtime/NetworkPrediction/Source/NetworkPrediction/Public/NetworkPredictionSettings.h
+++ b/Engine/Plugins/Runtime/NetworkPrediction/Source/NetworkPrediction/Public/NetworkPredictionSettings.h
@@ -15,7 +15,6 @@ struct FNetworkPredictionSettings
 	GENERATED_BODY()
 
 	// Which ticking policy to use in cases where both are supported by the underlying simulation.
-	// Leave this on Fixed if you intend to use physics based simulations.
 	UPROPERTY(config, EditAnywhere, Category = Global)
 	ENetworkPredictionTickingPolicy PreferredTickingPolicy = ENetworkPredictionTickingPolicy::Fixed;
 
@@ -29,11 +28,6 @@ struct FNetworkPredictionSettings
 	UPROPERTY(config, EditAnywhere, Category = FixedTick)
 	int32 FixedTickFrameRate = 60;
 
-	// Forces the engine to run in fixed tick mode when a NP physics simulation is running.
-	// This is the same as settings UEngine::bUseFixedFrameRate / FixedFrameRate manually.
-	UPROPERTY(config, EditAnywhere, Category = FixedTick)
-	bool bForceEngineFixTickForcePhysics = true;
-
 	// Default NetworkLOD for simulated proxy simulations.
 	UPROPERTY(config, EditAnywhere, Category = FixedTick)
 	ENetworkLOD SimulatedProxyNetworkLOD = ENetworkLOD::ForwardPredict;
diff --git a/Engine/Plugins/Runtime/NetworkPrediction/Source/NetworkPrediction/Public/NetworkPredictionTickState.h b/Engine/Plugins/Runtime/NetworkPrediction/Source/NetworkPrediction/Public/NetworkPredictionTickState.h
--- a/Engine/Plugins/Runtime/NetworkPrediction/Source/NetworkPrediction/Public/NetworkPredictionTickState.h
+++ b/Engine/Plugins/Runtime/NetworkPrediction/Source/NetworkPrediction/Public/NetworkPredictionTickState.h
@@ -38,15 +38,7 @@ struct FServiceTimeStep
 // (Global) Tick state for fixed tick services
 //
 // Notes about Fix Ticking in NetworkPrediction:
-//	1, This is the only mode that supports physics rewind/resim.
-//	2. For now, this requires the engine to run in a fixed step mode on both server and clients
-//		a. GEngine->bUseFixedFrameRate / GEngine->FixedFrameRate
-//	3. If you aren't using Physics, FixedTick mode will accumulate real time and run 0-N Sim frames per engine frame.
-//	4. If you are using Physics, FixedTick will always run 1 sim frame per engine frame to stay in sync with physics.
-//		a. Since NP uses int32 ms and the engine will use float DeltaTimeSeconds, NP will slowly lose time compared to the rest of the engine/physics
-//
-//	The long term goal is supporting fixed physics steps without engine fix ticking. There is no guarantee or timeline
-//	this will happen though.
+//	1. FixedTick mode will accumulate real time and run 0-N Sim frames per engine frame.
 //
 // ---------------------------------------------------------------------------
 struct FFixedTickState
@@ -56,7 +48,6 @@ struct FFixedTickState
 
 	// Realtime steps. That is, one of these = one FixedStepMS in simulation time.
 	// This means sim time ticks slightly slower than real time.
-	// This is seemingly the best choice while fixed tick physics has to be 1:1.
 	float FixedStepRealTimeMS = 1/30.f;
 
 	// Next frame to be ticked (used as input to generated PendingFrame+1)
@@ -70,14 +61,9 @@ struct FFixedTickState
 	// "LocalFrame" = "ServerFrame" - Offset. 
 	int32 Offset=0;
 
-	// Accumulates raw delta time into our fixed steps, Note this not used when physics recording is enabled since NP must stay in sync with physics tick.
+	// Accumulates raw delta time into our fixed steps
 	float UnspentTimeMS = 0.f;
 
-	// Maps PendingFrame to PhysicsRewindData()->CurrentFrame().
-	// Local PendingFrame = PhysicsFrame - Offset
-	int32 PhysicsOffset = 0;
-	Chaos::FRewindData* PhysicsRewindData = nullptr;
-
 	struct FInterpolationState
 	{
 		float AccumulatedTimeMS = 0.f; // accumulated real time
diff --git a/Engine/Plugins/Runtime/NetworkPrediction/Source/NetworkPrediction/Public/NetworkPredictionTrace.h b/Engine/Plugins/Runtime/NetworkPrediction/Source/NetworkPrediction/Public/NetworkPredictionTrace.h
--- a/Engine/Plugins/Runtime/NetworkPrediction/Source/NetworkPrediction/Public/NetworkPredictionTrace.h
+++ b/Engine/Plugins/Runtime/NetworkPrediction/Source/NetworkPrediction/Public/NetworkPredictionTrace.h
@@ -86,10 +86,6 @@ struct FNetworkPredictionInstanceConfig;
 #define UE_NP_TRACE_USER_STATE_SYNC(ModelDef, UserState) FNetworkPredictionTrace::TraceUserState<ModelDef>(UserState, FNetworkPredictionTrace::ETraceUserState::Sync)
 #define UE_NP_TRACE_USER_STATE_AUX(ModelDef, UserState) FNetworkPredictionTrace::TraceUserState<ModelDef>(UserState, FNetworkPredictionTrace::ETraceUserState::Aux)
 
-#define UE_NP_TRACE_PHYSICS_STATE_CURRENT(ModelDef, Driver) FNetworkPredictionTrace::TracePhysicsStateCurrent<ModelDef>(Driver)
-#define UE_NP_TRACE_PHYSICS_STATE_AT_FRAME(ModelDef, Frame, RewindData, Driver) FNetworkPredictionTrace::TracePhysicsStateAtFrame<ModelDef>(Frame, RewindData, Driver)
-#define UE_NP_TRACE_PHYSICS_STATE_RECV(ModelDef, NpPhysicsState) FNetworkPredictionTrace::TracePhysicsStateRecv<ModelDef>(NpPhysicsState)
-
 #else
 
 #include "NetworkPredictionLog.h"
@@ -119,9 +115,6 @@ struct FNetworkPredictionInstanceConfig;
 #define UE_NP_TRACE_USER_STATE_INPUT(...)
 #define UE_NP_TRACE_USER_STATE_SYNC(...)
 #define UE_NP_TRACE_USER_STATE_AUX(...)
-#define UE_NP_TRACE_PHYSICS_STATE_CURRENT(...)
-#define UE_NP_TRACE_PHYSICS_STATE_AT_FRAME(...)
-#define UE_NP_TRACE_PHYSICS_STATE_RECV(...)
 
 #endif // UE_NP_TRACE_ENABLED
 
@@ -178,7 +171,6 @@ public:
 		Input,
 		Sync,
 		Aux,
-		Physics
 	};
 
 	template<typename ModelDef, typename StateType>
@@ -201,66 +193,6 @@ public:
 #endif
 	}
 
-	template<typename ModelDef, typename DriverType>
-	static void TracePhysicsStateCurrent(DriverType* Driver)
-	{
-#if UE_NP_TRACE_USER_STATES_ENABLED
-
-		if (!FNetworkPredictionDriver<ModelDef>::HasPhysics())
-		{
-			return;
-		}
-
-		if (UE_TRACE_CHANNELEXPR_IS_ENABLED(NetworkPredictionChannel))
-		{
-			TAnsiStringBuilder<512> Builder;
-			FNetworkPredictionDriver<ModelDef>::TracePhysicsState(Driver, Builder);
-			TraceUserState_Internal(ETraceUserState::Physics, Builder);
-		}
-#endif
-	}
-
-	template<typename ModelDef, typename DriverType>
-	static void TracePhysicsStateAtFrame(int32 PhysicsFrame, Chaos::FRewindData* RewindData, DriverType* Driver)
-	{
-#if UE_NP_TRACE_USER_STATES_ENABLED
-		
-		
-		if (!FNetworkPredictionDriver<ModelDef>::HasPhysics())
-		{
-			return;
-		}
-
-		if (UE_TRACE_CHANNELEXPR_IS_ENABLED(NetworkPredictionChannel))
-		{
-			TAnsiStringBuilder<512> Builder;
-			FNetworkPredictionDriver<ModelDef>::TracePhysicsState(PhysicsFrame, RewindData, Driver, Builder);
-			TraceUserState_Internal(ETraceUserState::Physics, Builder);
-		}
-#endif
-	}
-
-	template<typename ModelDef, typename PhysicsStateType>
-	static void TracePhysicsStateRecv(const PhysicsStateType& State)
-	{
-#if UE_NP_TRACE_USER_STATES_ENABLED
-
-		if (!FNetworkPredictionDriver<ModelDef>::HasPhysics())
-		{
-			return;
-		}
-
-		if (UE_TRACE_CHANNELEXPR_IS_ENABLED(NetworkPredictionChannel))
-		{
-			npCheckSlow(State.Get());
-
-			TAnsiStringBuilder<512> Builder;
-			FNetworkPredictionDriver<ModelDef>::TracePhysicsStateRecv(State, Builder);
-			TraceUserState_Internal(ETraceUserState::Physics, Builder);
-		}
-#endif
-	}
-
 private:
 
 	static UE_API void TraceSimulationCreated_Internal(FNetworkPredictionID ID, FStringBuilderBase& Builder);
diff --git a/Engine/Plugins/Runtime/NetworkPrediction/Source/NetworkPrediction/Public/NetworkPredictionUtil.h b/Engine/Plugins/Runtime/NetworkPrediction/Source/NetworkPrediction/Public/NetworkPredictionUtil.h
--- a/Engine/Plugins/Runtime/NetworkPrediction/Source/NetworkPrediction/Public/NetworkPredictionUtil.h
+++ b/Engine/Plugins/Runtime/NetworkPrediction/Source/NetworkPrediction/Public/NetworkPredictionUtil.h
@@ -84,14 +84,3 @@ void NpResizeForIndex(ArrayType& Array, int32 Index)
 	}
 }
 
-class AActor;
-
-namespace Chaos
-{
-	class FSingleParticlePhysicsProxy;
-}
-
-namespace UE_NP
-{
-	NETWORKPREDICTION_API Chaos::FSingleParticlePhysicsProxy* FindBestPhysicsProxy(AActor* Owning, FName NamedComponent = NAME_None);
-}
diff --git a/Engine/Plugins/Runtime/NetworkPrediction/Source/NetworkPrediction/Public/NetworkPredictionWorldManager.h b/Engine/Plugins/Runtime/NetworkPrediction/Source/NetworkPrediction/Public/NetworkPredictionWorldManager.h
--- a/Engine/Plugins/Runtime/NetworkPrediction/Source/NetworkPrediction/Public/NetworkPredictionWorldManager.h
+++ b/Engine/Plugins/Runtime/NetworkPrediction/Source/NetworkPrediction/Public/NetworkPredictionWorldManager.h
@@ -86,9 +86,6 @@ public:
 	template<typename ModelDef>
 	void ConfigureInstance(FNetworkPredictionID ID, const FNetworkPredictionInstanceArchetype& Archetype, const FNetworkPredictionInstanceConfig& Config, FReplicationProxySet RepProxies, ENetRole Role, bool bHasNetConnection);
 		
-	// Is the engine currently setup to fix tick physics
-	UE_API bool CanPhysicsFixTick() const;
-
 	UE_API ENetworkPredictionTickingPolicy PreferredDefaultTickingPolicy() const;
 
 	UE_API void SyncNetworkPredictionSettings(const UNetworkPredictionSettingsObject* Settings);
@@ -100,25 +97,6 @@ public:
 	
 private:
 
-	struct FPhysics
-	{
-		// Set when we first have a registered instance of a ModelDef with a valid physics state.
-		// This does not mean physics recording has been enabled, just that the system has to account
-		// for physics ticking
-		bool bUsingPhysics = false;
-		bool bRecordingEnabled = false;
-
-		FChaosSolversModule* Module = nullptr;
-		Chaos::FPBDRigidsSolver* Solver = nullptr;
-	};
-	
-	FPhysics Physics;
-
-	UE_API void SetUsingPhysics();
-	UE_API void InitPhysicsCapture();
-	UE_API void AdvancePhysicsResimFrame(int32& PhysicsFrame);
-	UE_API void EnsurePhysicsGTSync(const TCHAR* Context) const;
-
 	FNetworkPredictionSettings Settings;
 
 	FFixedTickState FixedTickState;
@@ -211,11 +189,6 @@ void UNetworkPredictionWorldManager::ConfigureInstance(FNetworkPredictionID ID,
 	const bool bNewInstance = (InstanceData.NetRole == ROLE_None);
 	InstanceData.NetRole = Role;
 
-	if (FNetworkPredictionDriver<ModelDef>::HasPhysics())
-	{
-		SetUsingPhysics();
-	}
-
 	TRemoteInputService<ModelDef>::SetMaxFaultLimit(Settings.MaximumRemoteInputFaultLimit);
 
 	ENetworkPredictionService ServiceMask = ENetworkPredictionService::None;
@@ -315,8 +288,6 @@ void UNetworkPredictionWorldManager::ConfigureInstance(FNetworkPredictionID ID,
 			&FrameData.Buffer[FixedTickState.PendingFrame].SyncState, 
 			&FrameData.Buffer[FixedTickState.PendingFrame].AuxState);
 
-		bool bUsePhysicsRecording = false;
-
 		// Bind NetSend/Recv and role-dependent services
 		switch (Role)
 		{
@@ -397,7 +368,6 @@ void UNetworkPredictionWorldManager::ConfigureInstance(FNetworkPredictionID ID,
 					}
 				}
 
-				bUsePhysicsRecording = true;
 				break;
 
 			case ENetworkLOD::Interpolated:
@@ -405,16 +375,6 @@ void UNetworkPredictionWorldManager::ConfigureInstance(FNetworkPredictionID ID,
 				break;
 			}
 		}
-
-		// Physics
-		if (FNetworkPredictionDriver<ModelDef>::HasPhysics())
-		{
-			ServiceMask |= ENetworkPredictionService::FixedPhysics;
-			if (bUsePhysicsRecording && !Physics.bRecordingEnabled)
-			{
-				InitPhysicsCapture();
-			}
-		}
 	}
 
 	// Net Cues: set which replicated cues we should accept based on if we are FP or interpolated
diff --git a/Engine/Plugins/Runtime/NetworkPrediction/Source/NetworkPrediction/Public/Services/NetworkPredictionInstanceData.h b/Engine/Plugins/Runtime/NetworkPrediction/Source/NetworkPrediction/Public/Services/NetworkPredictionInstanceData.h
--- a/Engine/Plugins/Runtime/NetworkPrediction/Source/NetworkPrediction/Public/Services/NetworkPredictionInstanceData.h
+++ b/Engine/Plugins/Runtime/NetworkPrediction/Source/NetworkPrediction/Public/Services/NetworkPredictionInstanceData.h
@@ -26,8 +26,7 @@ enum class ENetworkPredictionService : uint32
 	FixedTick					= 1 << 6,
 	FixedSmoothing				= 1 << 7,
 	FixedFinalize				= 1 << 8,
-	FixedPhysics				= 1 << 9,
-	MAX_FIXED					= FixedPhysics,
+	MAX_FIXED					= FixedFinalize,
 
 	// Services exclusive to independent tick mode
 	IndependentRollback		= 1 << 10,
@@ -65,7 +64,6 @@ struct TInstanceData
 };
 
 // Frame data that instances with StateTypes will have.
-// (a pure physics object would not have this for example)
 template<typename ModelDef=FNetworkPredictionModelDef>
 struct TInstanceFrameState
 {
@@ -95,7 +93,6 @@ struct TClientRecvData
 	using InputType = typename StateTypes::InputType;
 	using SyncType = typename StateTypes::SyncType;
 	using AuxType = typename StateTypes::AuxType;
-	using PhysicsState = typename ModelDef::PhysicsState;
 
 	int32 ServerFrame; // Fixed tick || Independent AP only
 	int32 SimTimeMS; // Independent tick only
@@ -104,8 +101,6 @@ struct TClientRecvData
 	TConditionalState<SyncType>	SyncState;
 	TConditionalState<AuxType>	AuxState;
 
-	TConditionalState<PhysicsState> Physics;
-
 	// Acceleration data.
 	int32 TraceID = INDEX_NONE;
 	int32 InstanceIdx = INDEX_NONE;	// Index into TModelDataStore::Instances
diff --git a/Engine/Plugins/Runtime/NetworkPrediction/Source/NetworkPrediction/Public/Services/NetworkPredictionServiceRegistry.h b/Engine/Plugins/Runtime/NetworkPrediction/Source/NetworkPrediction/Public/Services/NetworkPredictionServiceRegistry.h
--- a/Engine/Plugins/Runtime/NetworkPrediction/Source/NetworkPrediction/Public/Services/NetworkPredictionServiceRegistry.h
+++ b/Engine/Plugins/Runtime/NetworkPrediction/Source/NetworkPrediction/Public/Services/NetworkPredictionServiceRegistry.h
@@ -9,7 +9,6 @@
 #include "NetworkPredictionService_Input.inl"
 #include "NetworkPredictionService_Finalize.inl"
 #include "NetworkPredictionService_ServerRPC.inl"
-#include "NetworkPredictionService_Physics.inl"
 #include "NetworkPredictionService_Smooth.inl"
 
 // Services do the actual system work on batches of registered instances. UNetworkPredictionWorldManager orchestrates them.
@@ -62,7 +61,7 @@
 //
 // Future Improvements:
 //	-All services template classes are currently instantiated for all ModelDefs (the templated code is generated for each ModelDef)
-//	-Even if ModelDefs are incompatible for a given service. E.g, a pure physics ModelDef with no Simulation still has the TLocalInputService<PhysicsModelDef> code instantiated.
+//	-Even if ModelDefs are incompatible for a given service.
 //	-To be clear: it does not instantiate an instance of TLocalInputService at runtime, but the code is generated and almost certainly cannot be culled by the compiler.
 //	-Concepts could be used to determine if ModelDefs are compatible with Services and we could avoid the template code instantiation.
 //	-This would add more template machinery, and basically requires all TService<ModelDefs> be hidden behind SFINAE walls. E.g, Factories for instantiating and CallOrNot for Register/Unregister.
@@ -171,7 +170,6 @@ public:
 	NP_DECLARE_SERVICE(FixedInputRemote,			IInputService);
 	NP_DECLARE_SERVICE(FixedTick,					ILocalTickService);
 	NP_DECLARE_SERVICE(FixedFinalize,				IFinalizeService);
-	NP_DECLARE_SERVICE(FixedPhysics,				IPhysicsService);
 	NP_DECLARE_SERVICE(FixedSmoothing,				IFixedSmoothingService);
 
 	NP_DECLARE_SERVICE(IndependentRollback,			IIndependentRollbackService);
@@ -198,7 +196,6 @@ private:
 			NP_DEFINE_SERVICE_CALL(FixedInputRemote,			TRemoteInputService);
 			NP_DEFINE_SERVICE_CALL(FixedTick,					TLocalTickService);
 			NP_DEFINE_SERVICE_CALL(FixedFinalize,				TFinalizeService);
-			NP_DEFINE_SERVICE_CALL(FixedPhysics,				TPhysicsService);
 			NP_DEFINE_SERVICE_CALL(FixedSmoothing,				TFixedSmoothingService);
 		}
 		else if (EnumHasAnyFlags(Mask, ENetworkPredictionService::ANY_INDEPENDENT))
diff --git a/Engine/Plugins/Runtime/NetworkPrediction/Source/NetworkPrediction/Public/Services/NetworkPredictionService_Interpolate.inl b/Engine/Plugins/Runtime/NetworkPrediction/Source/NetworkPrediction/Public/Services/NetworkPredictionService_Interpolate.inl
--- a/Engine/Plugins/Runtime/NetworkPrediction/Source/NetworkPrediction/Public/Services/NetworkPredictionService_Interpolate.inl
+++ b/Engine/Plugins/Runtime/NetworkPrediction/Source/NetworkPrediction/Public/Services/NetworkPredictionService_Interpolate.inl
@@ -68,9 +68,6 @@ public:
 	using SyncType = typename StateTypes::SyncType;
 	using AuxType = typename StateTypes::AuxType;
 	using SyncAuxType = TSyncAuxPair<StateTypes>;
-	using PhysicsState = typename ModelDef::PhysicsState;
-
-	static constexpr bool bHasPhysics = FNetworkPredictionDriver<ModelDef>::HasPhysics();
 
 	TFixedInterpolateService(TModelDataStore<ModelDef>* InDataStore)
 		: DataStore(InDataStore) { }
@@ -90,11 +87,6 @@ public:
 		// Point the PresentationView to our managed state. Note this only has to be done once
 		FInstance* InternalInstance = (FInstance*)AllocInfo.Pointer;
 		InstanceData.Info.View->UpdatePresentationView(InternalInstance->SyncState, InternalInstance->AuxState);
-
-		if (bHasPhysics)
-		{
-			FNetworkPredictionDriver<ModelDef>::BeginInterpolatedPhysics(InstanceData.Info.Driver);
-		}
 	}
 
 	void UnregisterInstance(FNetworkPredictionID ID)
@@ -105,11 +97,6 @@ public:
 		TInstanceData<ModelDef>& InstanceData = DataStore->Instances.GetByIndexChecked(ClientRecvData.InstanceIdx);
 		InstanceData.Info.View->ClearPresentationView();
 
-		if (bHasPhysics)
-		{
-			FNetworkPredictionDriver<ModelDef>::EndInterpolatedPhysics(InstanceData.Info.Driver);
-		}
-		
 		FInstance& InterpolationData = Instances[ClientRecvIdx];
 		if (InterpolationData.bTwoValidFrames == false)
 		{
@@ -166,11 +153,6 @@ public:
 				}
 			}
 
-			if (bHasPhysics)
-			{
-				InterpolationData.PhysicsBuffer[LocalFrame] = ClientRecvData.Physics;
-			}			
-
 			// Look for gaps between the current interpolation frame and what we just wrote.
 			if (ToFrame != INDEX_NONE)
 			{
@@ -229,11 +211,6 @@ public:
 
 				FNetworkPredictionDriver<ModelDef>::FinalizeFrame(InstanceData.Info.Driver, ClientRecvData.SyncState, ClientRecvData.AuxState);
 
-				if (bHasPhysics)
-				{
-					FNetworkPredictionDriver<ModelDef>::FinalizeInterpolatedPhysics(InstanceData.Info.Driver, ClientRecvData.Physics);
-				}
-
 				InstanceData.Info.View->UpdateView(ToFrame, ClientRecvData.SimTimeMS, &Frames.Buffer[TickState->PendingFrame].InputCmd, &ClientRecvData.SyncState, &ClientRecvData.AuxState);
 
 
@@ -285,14 +262,6 @@ public:
 				TInstanceData<ModelDef>& InstanceData = DataStore->Instances.GetByIndexChecked(Instance.InstanceIdx);
 				FNetworkPredictionDriver<ModelDef>::FinalizeFrame(InstanceData.Info.Driver, Instance.SyncState, Instance.AuxState);
 
-				if (bHasPhysics)
-				{
-					TConditionalState<PhysicsState> PhysicsState;
-
-					FNetworkPredictionDriver<ModelDef>::InterpolatePhysics(Instance.PhysicsBuffer[FromFrame], Instance.PhysicsBuffer[ToFrame], PCT, PhysicsState);
-					FNetworkPredictionDriver<ModelDef>::FinalizeInterpolatedPhysics(InstanceData.Info.Driver, PhysicsState);
-				}
-
 				// Update SimulationView with the frame we are interpolating to
 				// FIXME: so this is kind of bad, we need to do '&Frames.Buffer[TickState->PendingFrame].InputCmd'  to support non-predictive (interpolated) client controlled simulations.
 				// This ensures that View->PendingInputCmd (where ProduceInput service writes to) is always Frames.Buffer[TickState->PendingFrame].InputCmd (where Replicators send cmds from).
@@ -309,7 +278,7 @@ private:
 	struct FInstance
 	{
 		FInstance(int32 InTraceID, int32 InInstanceIdx, int32 InFramesIdx)
-			: TraceID(InTraceID), InstanceIdx(InInstanceIdx), FramesIdx(InFramesIdx), PhysicsBuffer(64) {} // fixme
+			: TraceID(InTraceID), InstanceIdx(InInstanceIdx), FramesIdx(InFramesIdx) {} // fixme
 
 		int32 TraceID;
 		int32 InstanceIdx;
@@ -320,12 +289,6 @@ private:
 		// Last interpolated values. Stored here so that we can maintain FNetworkPredictionStateView to them
 		TConditionalState<SyncType> SyncState;
 		TConditionalState<AuxType> AuxState;
-
-		// If we are doing interpolated physics, we need a place to store it
-		// (since the physics engine itself handles saving physics state, we have no NP managed buffers of physics state)
-		// Given current implementation, it is tempting to manage all interpolated state within the service,
-		// though this goes against the general guidelines of services only storing acceleration/bookkeeping data
-		TNetworkPredictionBuffer<TConditionalState<PhysicsState>> PhysicsBuffer;
 	};
 
 	TSparseArray<FInstance> Instances; // Indices are shared with DataStore->ClientRecv
@@ -342,11 +305,6 @@ private:
 
 		DestFrameData.SyncState = SourceFrameData.SyncState;
 		DestFrameData.AuxState = SourceFrameData.AuxState;
-
-		if (bHasPhysics)
-		{
-			Instance.PhysicsBuffer[DestFrame] = Instance.PhysicsBuffer[SourceFrame];
-		}
 	}
 
 	static void Interpolate(const int32 FromSourceFrame, const int32 ToSourceFrame, const int32 DestFrame, TInstanceFrameState<ModelDef>& Frames, FInstance& Instance, const int32 FixedStepMS)
@@ -371,11 +329,6 @@ private:
 
 		FNetworkPredictionDriver<ModelDef>::Interpolate(SyncAuxType{FromFrameData.SyncState, FromFrameData.AuxState}, SyncAuxType{ToFrameData.SyncState, ToFrameData.AuxState}, 
 			PCT, DestFrameData.SyncState, DestFrameData.AuxState);
-
-		if (bHasPhysics)
-		{
-			FNetworkPredictionDriver<ModelDef>::InterpolatePhysics(Instance.PhysicsBuffer[FromSourceFrame], Instance.PhysicsBuffer[ToSourceFrame], PCT, Instance.PhysicsBuffer[DestFrame]);
-		}
 	}
 };
 
diff --git a/Engine/Plugins/Runtime/NetworkPrediction/Source/NetworkPrediction/Public/Services/NetworkPredictionService_Physics.inl b/Engine/Plugins/Runtime/NetworkPrediction/Source/NetworkPrediction/Public/Services/NetworkPredictionService_Physics.inl
deleted file mode 100644
diff --git a/Engine/Plugins/Runtime/NetworkPrediction/Source/NetworkPrediction/Public/Services/NetworkPredictionService_Rollback.inl b/Engine/Plugins/Runtime/NetworkPrediction/Source/NetworkPrediction/Public/Services/NetworkPredictionService_Rollback.inl
--- a/Engine/Plugins/Runtime/NetworkPrediction/Source/NetworkPrediction/Public/Services/NetworkPredictionService_Rollback.inl
+++ b/Engine/Plugins/Runtime/NetworkPrediction/Source/NetworkPrediction/Public/Services/NetworkPredictionService_Rollback.inl
@@ -118,30 +118,6 @@ public:
 					FNetworkPredictionDriver<ModelDef>::LogUserStates({LocalFrameData.InputCmd, LocalFrameData.SyncState, LocalFrameData.AuxState });
 				}
 			}
-			else if (FNetworkPredictionDriver<ModelDef>::HasPhysics())
-			{
-				// LocalFrame = PhysicsFrame - Offset;
-				const int32 PhysicsFrame = LocalFrame + TickState->PhysicsOffset;
-
-				// Maybe consider putting a copy of PhysicsActor on TClientRecvData if this lookup ever shows up in profiler
-
-				if (FNetworkPredictionDriver<ModelDef>::ShouldReconcilePhysics(PhysicsFrame, TickState->PhysicsRewindData, InstanceData.Info.Driver, ClientRecvData.Physics))
-				{
-					UE_NP_TRACE_SHOULD_RECONCILE(ClientRecvData.TraceID);
-					bDoRollback = true;
-
-					if (NetworkPredictionCVars::PrintReconciles())
-					{
-						UE_LOG(LogNetworkPrediction, Warning, TEXT("Reconcile required due to Physics mismatch. LocalFrame: %d. Recv Frame: %d. Offset: %d. Idx: %d"), LocalFrame, ClientRecvData.ServerFrame, Offset, LocalFrame % Frames.Buffer.Capacity());
-
-						UE_LOG(LogNetworkPrediction, Warning, TEXT("Received:"));
-						FNetworkPredictionDriver<ModelDef>::LogPhysicsState(ClientRecvData.Physics);
-
-						UE_LOG(LogNetworkPrediction, Warning, TEXT("Local:"));
-						FNetworkPredictionDriver<ModelDef>::LogPhysicsState(PhysicsFrame, TickState->PhysicsRewindData, InstanceData.Info.Driver);
-					}
-				}
-			}
 
 			if (bDoRollback && !NetworkPredictionCVars::SkipReconcile())
 			{
@@ -234,11 +210,6 @@ private:
 
 				TInstanceData<ModelDef>& InstanceData = DataStore->Instances.GetByIndexChecked(ClientRecvData.InstanceIdx);
 
-				if (FNetworkPredictionDriver<ModelDef>::HasPhysics())
-				{
-					FNetworkPredictionDriver<ModelDef>::PerformPhysicsRollback(InstanceData.Info.Driver, ClientRecvData.Physics);
-				}
-
 				// Copy input cmd if SP
 				if (ClientRecvData.NetRole == ROLE_SimulatedProxy)
 				{
diff --git a/Engine/Plugins/Runtime/NetworkPrediction/Source/NetworkPrediction/Public/Services/NetworkPredictionService_Ticking.inl b/Engine/Plugins/Runtime/NetworkPrediction/Source/NetworkPrediction/Public/Services/NetworkPredictionService_Ticking.inl
--- a/Engine/Plugins/Runtime/NetworkPrediction/Source/NetworkPrediction/Public/Services/NetworkPredictionService_Ticking.inl
+++ b/Engine/Plugins/Runtime/NetworkPrediction/Source/NetworkPrediction/Public/Services/NetworkPredictionService_Ticking.inl
@@ -45,8 +45,6 @@ struct TTickUtil
 		// Fixme: should only trace aux if it changed
 		UE_NP_TRACE_USER_STATE_SYNC(ModelDef, OutputFrameData.SyncState.Get());
 		UE_NP_TRACE_USER_STATE_AUX(ModelDef, OutputFrameData.AuxState.Get());
-
-		UE_NP_TRACE_PHYSICS_STATE_CURRENT(ModelDef, Instance.Info.Driver);
 	}
 
 	template<typename SimulationType = typename ModelDef::Simulation>
diff --git a/Engine/Plugins/Runtime/NetworkPredictionExtras/Source/NetworkPredictionExtras/Private/CharacterMotionComponent.cpp b/Engine/Plugins/Runtime/NetworkPredictionExtras/Source/NetworkPredictionExtras/Private/CharacterMotionComponent.cpp
--- a/Engine/Plugins/Runtime/NetworkPredictionExtras/Source/NetworkPredictionExtras/Private/CharacterMotionComponent.cpp
+++ b/Engine/Plugins/Runtime/NetworkPredictionExtras/Source/NetworkPredictionExtras/Private/CharacterMotionComponent.cpp
@@ -11,7 +11,6 @@
 #include "Services/NetworkPredictionService_Finalize.inl"
 #include "Services/NetworkPredictionService_Input.inl"
 #include "Services/NetworkPredictionService_Interpolate.inl"
-#include "Services/NetworkPredictionService_Physics.inl"
 #include "Services/NetworkPredictionService_Rollback.inl"
 #include "Services/NetworkPredictionService_ServerRPC.inl"
 
diff --git a/Engine/Plugins/Runtime/NetworkPredictionExtras/Source/NetworkPredictionExtras/Private/FlyingMovementComponent.cpp b/Engine/Plugins/Runtime/NetworkPredictionExtras/Source/NetworkPredictionExtras/Private/FlyingMovementComponent.cpp
--- a/Engine/Plugins/Runtime/NetworkPredictionExtras/Source/NetworkPredictionExtras/Private/FlyingMovementComponent.cpp
+++ b/Engine/Plugins/Runtime/NetworkPredictionExtras/Source/NetworkPredictionExtras/Private/FlyingMovementComponent.cpp
@@ -11,7 +11,6 @@
 #include "Services/NetworkPredictionService_Finalize.inl"
 #include "Services/NetworkPredictionService_Input.inl"
 #include "Services/NetworkPredictionService_Interpolate.inl"
-#include "Services/NetworkPredictionService_Physics.inl"
 #include "Services/NetworkPredictionService_Rollback.inl"
 #include "Services/NetworkPredictionService_ServerRPC.inl"
 
diff --git a/Engine/Plugins/Runtime/NetworkPredictionExtras/Source/NetworkPredictionExtras/Private/MockAbilitySimulation.cpp b/Engine/Plugins/Runtime/NetworkPredictionExtras/Source/NetworkPredictionExtras/Private/MockAbilitySimulation.cpp
--- a/Engine/Plugins/Runtime/NetworkPredictionExtras/Source/NetworkPredictionExtras/Private/MockAbilitySimulation.cpp
+++ b/Engine/Plugins/Runtime/NetworkPredictionExtras/Source/NetworkPredictionExtras/Private/MockAbilitySimulation.cpp
@@ -9,7 +9,6 @@
 #include "Services/NetworkPredictionService_Finalize.inl"
 #include "Services/NetworkPredictionService_Input.inl"
 #include "Services/NetworkPredictionService_Interpolate.inl"
-#include "Services/NetworkPredictionService_Physics.inl"
 #include "Services/NetworkPredictionService_Rollback.inl"
 #include "Services/NetworkPredictionService_ServerRPC.inl"
 
diff --git a/Engine/Plugins/Runtime/NetworkPredictionExtras/Source/NetworkPredictionExtras/Private/MockCharacterAbilitySimulation.cpp b/Engine/Plugins/Runtime/NetworkPredictionExtras/Source/NetworkPredictionExtras/Private/MockCharacterAbilitySimulation.cpp
--- a/Engine/Plugins/Runtime/NetworkPredictionExtras/Source/NetworkPredictionExtras/Private/MockCharacterAbilitySimulation.cpp
+++ b/Engine/Plugins/Runtime/NetworkPredictionExtras/Source/NetworkPredictionExtras/Private/MockCharacterAbilitySimulation.cpp
@@ -8,7 +8,6 @@
 #include "Services/NetworkPredictionService_Finalize.inl"
 #include "Services/NetworkPredictionService_Input.inl"
 #include "Services/NetworkPredictionService_Interpolate.inl"
-#include "Services/NetworkPredictionService_Physics.inl"
 #include "Services/NetworkPredictionService_Rollback.inl"
 #include "Services/NetworkPredictionService_ServerRPC.inl"
 
diff --git a/Engine/Plugins/Runtime/NetworkPredictionExtras/Source/NetworkPredictionExtras/Private/MockNetworkSimulation.cpp b/Engine/Plugins/Runtime/NetworkPredictionExtras/Source/NetworkPredictionExtras/Private/MockNetworkSimulation.cpp
--- a/Engine/Plugins/Runtime/NetworkPredictionExtras/Source/NetworkPredictionExtras/Private/MockNetworkSimulation.cpp
+++ b/Engine/Plugins/Runtime/NetworkPredictionExtras/Source/NetworkPredictionExtras/Private/MockNetworkSimulation.cpp
@@ -15,7 +15,6 @@
 #include "Services/NetworkPredictionService_Finalize.inl"
 #include "Services/NetworkPredictionService_Input.inl"
 #include "Services/NetworkPredictionService_Interpolate.inl"
-#include "Services/NetworkPredictionService_Physics.inl"
 #include "Services/NetworkPredictionService_Rollback.inl"
 #include "Services/NetworkPredictionService_ServerRPC.inl"
 
diff --git a/Engine/Plugins/Runtime/NetworkPredictionExtras/Source/NetworkPredictionExtras/Private/MockPhysicsComponent.cpp b/Engine/Plugins/Runtime/NetworkPredictionExtras/Source/NetworkPredictionExtras/Private/MockPhysicsComponent.cpp
deleted file mode 100644
diff --git a/Engine/Plugins/Runtime/NetworkPredictionExtras/Source/NetworkPredictionExtras/Private/MockPhysicsGrenade.cpp b/Engine/Plugins/Runtime/NetworkPredictionExtras/Source/NetworkPredictionExtras/Private/MockPhysicsGrenade.cpp
deleted file mode 100644
diff --git a/Engine/Plugins/Runtime/NetworkPredictionExtras/Source/NetworkPredictionExtras/Private/MockPhysicsSimulation.cpp b/Engine/Plugins/Runtime/NetworkPredictionExtras/Source/NetworkPredictionExtras/Private/MockPhysicsSimulation.cpp
deleted file mode 100644
diff --git a/Engine/Plugins/Runtime/NetworkPredictionExtras/Source/NetworkPredictionExtras/Private/MockRootMotionComponent.cpp b/Engine/Plugins/Runtime/NetworkPredictionExtras/Source/NetworkPredictionExtras/Private/MockRootMotionComponent.cpp
--- a/Engine/Plugins/Runtime/NetworkPredictionExtras/Source/NetworkPredictionExtras/Private/MockRootMotionComponent.cpp
+++ b/Engine/Plugins/Runtime/NetworkPredictionExtras/Source/NetworkPredictionExtras/Private/MockRootMotionComponent.cpp
@@ -12,7 +12,6 @@
 #include "Services/NetworkPredictionService_Finalize.inl"
 #include "Services/NetworkPredictionService_Input.inl"
 #include "Services/NetworkPredictionService_Interpolate.inl"
-#include "Services/NetworkPredictionService_Physics.inl"
 #include "Services/NetworkPredictionService_Rollback.inl"
 #include "Services/NetworkPredictionService_ServerRPC.inl"
 
diff --git a/Engine/Plugins/Runtime/NetworkPredictionExtras/Source/NetworkPredictionExtras/Public/MockPhysicsComponent.h b/Engine/Plugins/Runtime/NetworkPredictionExtras/Source/NetworkPredictionExtras/Public/MockPhysicsComponent.h
deleted file mode 100644
diff --git a/Engine/Plugins/Runtime/NetworkPredictionExtras/Source/NetworkPredictionExtras/Public/MockPhysicsGrenade.h b/Engine/Plugins/Runtime/NetworkPredictionExtras/Source/NetworkPredictionExtras/Public/MockPhysicsGrenade.h
deleted file mode 100644
diff --git a/Engine/Plugins/Runtime/NetworkPredictionExtras/Source/NetworkPredictionExtras/Public/MockPhysicsSimulation.h b/Engine/Plugins/Runtime/NetworkPredictionExtras/Source/NetworkPredictionExtras/Public/MockPhysicsSimulation.h
deleted file mode 100644
diff --git a/Engine/Plugins/Runtime/NetworkPredictionExtras/Source/NetworkPredictionExtras/Public/ParametricMovement.cpp b/Engine/Plugins/Runtime/NetworkPredictionExtras/Source/NetworkPredictionExtras/Public/ParametricMovement.cpp
--- a/Engine/Plugins/Runtime/NetworkPredictionExtras/Source/NetworkPredictionExtras/Public/ParametricMovement.cpp
+++ b/Engine/Plugins/Runtime/NetworkPredictionExtras/Source/NetworkPredictionExtras/Public/ParametricMovement.cpp
@@ -7,7 +7,6 @@
 #include "Services/NetworkPredictionService_Finalize.inl"
 #include "Services/NetworkPredictionService_Input.inl"
 #include "Services/NetworkPredictionService_Interpolate.inl"
-#include "Services/NetworkPredictionService_Physics.inl"
 #include "Services/NetworkPredictionService_Rollback.inl"
 #include "Services/NetworkPredictionService_ServerRPC.inl"
 
