NPP: integrating time dilation feature to help shrink target input buffer size once network conditions improve

#Synced-CL 51151953

diff --git a/Engine/Plugins/Runtime/NetworkPrediction/Source/NetworkPrediction/Private/NetworkPredictionSerialization.cpp b/Engine/Plugins/Runtime/NetworkPrediction/Source/NetworkPrediction/Private/NetworkPredictionSerialization.cpp
new file mode 100644
--- /dev/null
+++ b/Engine/Plugins/Runtime/NetworkPrediction/Source/NetworkPrediction/Private/NetworkPredictionSerialization.cpp
@@ -0,0 +1,45 @@
+// Copyright Epic Games, Inc. All Rights Reserved.
+
+
+#include "NetworkPredictionSerialization.h"
+
+
+namespace NetworkPredictionCVars
+{
+	TAutoConsoleVariable<bool> EnableTimeDilation(
+		TEXT("np.TimeDilation.Enable"),
+		true,
+		TEXT("Time Dilation affects Autonomous Proxy clients in Fixed Tick mode, trying to minimize the amount of time the client is predicting ahead of the server. The server suggests to slow down or speed up the frequency of the client's sim frames, to encourage the smallest amount of input buffering without starvation."));
+
+	TAutoConsoleVariable<float> TimeDilationAmount(
+		TEXT("np.TimeDilation.Amount"),
+		0.01f,
+		TEXT("Affects how quickly dilation occurs. Server-side CVar, Disable TimeDilation by setting to 0 | Default: 0.01 | Value is in percent where 0.01 = 1% dilation. Example: 1.0/0.01 = 100, meaning that over the time it usually takes to tick 100 fixed simulation steps we will tick 99 or 101 depending on if we dilate up or down."));
+
+	TAutoConsoleVariable<bool> AllowTimeDilationEscalation(
+		TEXT("np.TimeDilation.AllowEscalation"),
+		true,
+		TEXT("Enables aggressive dilation when input buffer surplus or deficit grow. Server-side CVar. If enabled, dilate the time even more depending on how many ticks we need to adjust. When set to false, we use the set TimeDilationAmount as perform correct the offset slowly. When set to true, we multiply the TimeDilationEscalationDecay with the buffer offset count and resolve faster. Capped by TimeDilationEscalationDecayMax."));
+
+	TAutoConsoleVariable<float> TimeDilationEscalationDecay(
+		TEXT("np.TimeDilation.EscalationDecay"),
+		0.05f,
+		TEXT("Value is a multiplier, Default: 0.05. For each frame offset from the ideal amount, also decay by this much. Disable by setting to 0"));
+
+	TAutoConsoleVariable<float> TimeDilationEscalationDecayMax(
+		TEXT("np.TimeDilation.EscalationDecayMax"),
+		0.5f,
+		TEXT("Default: 0.5. The max decay value for escalated time dilation, which caps decay when input buffer offset is large. Lower value means higher decay."));
+
+	TAutoConsoleVariable<float> TimeDilationMax(
+		TEXT("np.TimeDilation.Max"),
+		1.1f,
+		TEXT("Max value of the time dilation multiplier, with 1.0 meaning no dilation"));
+
+	TAutoConsoleVariable<float> TimeDilationMin(
+		TEXT("np.TimeDilation.Min"),
+		0.9f,
+		TEXT("Min value of the time dilation multiplier, with 1.0 meaning no dilation"));
+
+
+}
\ No newline at end of file
diff --git a/Engine/Plugins/Runtime/NetworkPrediction/Source/NetworkPrediction/Private/NetworkPredictionWorldManager.cpp b/Engine/Plugins/Runtime/NetworkPrediction/Source/NetworkPrediction/Private/NetworkPredictionWorldManager.cpp
--- a/Engine/Plugins/Runtime/NetworkPrediction/Source/NetworkPrediction/Private/NetworkPredictionWorldManager.cpp
+++ b/Engine/Plugins/Runtime/NetworkPrediction/Source/NetworkPrediction/Private/NetworkPredictionWorldManager.cpp
@@ -95,6 +95,16 @@ void UNetworkPredictionWorldManager::OnWorldPreTick_Internal(float InDeltaSecond
 	FixedTickState.FixedStepRealTimeMS = (1.f /  InFixedFrameRate) * 1000.f;
 	FixedTickState.FixedStepMS = (int32)FixedTickState.FixedStepRealTimeMS;
 
+	// Time Dilation (Fixed Tick only) : happens only on autonomous clients as a throttle up/down of the client's simulation frequency in order to minimize server-side input buffering.
+	// Higher latency network conditions require more input buffering, and the buffer should grow/shrink as network latency changes.
+	// This is calculated by the server based on its input buffer count and sent to autonomous proxy via TFixedTickReplicator_AP NetSend().
+	FixedTickState.FixedStepDilatedRealTimeMS = FixedTickState.FixedStepRealTimeMS;
+	
+	if (NetworkPredictionCVars::EnableTimeDilation->GetBool())
+	{
+		FixedTickState.FixedStepDilatedRealTimeMS = FixedTickState.FixedStepRealTimeMS * FixedTickState.SuggestedTimeDilation;
+	}
+
 	ActiveInstance = this;
 }
 
@@ -261,9 +271,9 @@ void UNetworkPredictionWorldManager::BeginNewSimulationFrame_Internal(float Delt
 
 		FixedTickState.UnspentTimeMS += fEngineFrameDeltaTimeMS;
 
-		while ((FixedTickState.UnspentTimeMS + KINDA_SMALL_NUMBER) >= FixedTickState.FixedStepRealTimeMS)
+		while ((FixedTickState.UnspentTimeMS + KINDA_SMALL_NUMBER) >= FixedTickState.FixedStepDilatedRealTimeMS)
 		{
-			FixedTickState.UnspentTimeMS -= FixedTickState.FixedStepRealTimeMS;
+			FixedTickState.UnspentTimeMS -= FixedTickState.FixedStepDilatedRealTimeMS;
 			if (FMath::IsNearlyZero(FixedTickState.UnspentTimeMS))
 			{
 				FixedTickState.UnspentTimeMS = 0.f;
diff --git a/Engine/Plugins/Runtime/NetworkPrediction/Source/NetworkPrediction/Public/NetworkPredictionSerialization.h b/Engine/Plugins/Runtime/NetworkPrediction/Source/NetworkPrediction/Public/NetworkPredictionSerialization.h
--- a/Engine/Plugins/Runtime/NetworkPrediction/Source/NetworkPrediction/Public/NetworkPredictionSerialization.h
+++ b/Engine/Plugins/Runtime/NetworkPrediction/Source/NetworkPrediction/Public/NetworkPredictionSerialization.h
@@ -10,6 +10,9 @@
 #include "Services/NetworkPredictionInstanceData.h"
 #include "UObject/CoreNet.h"
 
+#define UE_API NETWORKPREDICTION_API
+
+
 #define NETSIM_ENABLE_CHECKSUMS !(UE_BUILD_SHIPPING || UE_BUILD_TEST)
 #if NETSIM_ENABLE_CHECKSUMS 
 #define NETSIM_CHECKSUM(Ser) SerializeChecksum(Ser,0xA186A384, false);
@@ -24,6 +27,16 @@
 namespace NetworkPredictionCVars
 {
 	NETSIM_DEVCVAR_SHIPCONST_INT(ForceSendDefaultInputCommands, 0, "np.ForceSendDefaultInputCommands", "While enabled on a client, it will send default input cmds to the server, rather than the locally-produced input. This is a way to introduce de-syncs/rollbacks for debugging.");
+
+	// Options governing fixed-tick time dilation. See CVar definitions for descriptions.
+	extern UE_API TAutoConsoleVariable<bool> EnableTimeDilation;
+	extern UE_API TAutoConsoleVariable<float> TimeDilationAmount;
+	extern UE_API TAutoConsoleVariable<bool> AllowTimeDilationEscalation;
+	extern UE_API TAutoConsoleVariable<float> TimeDilationEscalationDecay;
+	extern UE_API TAutoConsoleVariable<float> TimeDilationEscalationDecayMax;
+	extern UE_API TAutoConsoleVariable<float> TimeDilationMax;
+	extern UE_API TAutoConsoleVariable<float> TimeDilationMin;
+
 }
 
 struct FNetworkPredictionSerialization
@@ -144,8 +157,16 @@ public:
 				}
 				else
 				{
-					UE_NP_TRACE_SYSTEM_FAULT("Recovering from input stream starvation on server. Advancing over dropped client frames [%d-%d]", ServerRecvData.LastConsumedFrame+1, StartFrame-1);
-					ServerRecvData.LastConsumedFrame = StartFrame-1;
+
+					if (!NetworkPredictionCVars::EnableTimeDilation->GetBool())
+					{
+						UE_NP_TRACE_SYSTEM_FAULT("Recovering from input stream starvation on server. Advancing over dropped client frames [%d-%d]", ServerRecvData.LastConsumedFrame + 1, StartFrame - 1);
+						ServerRecvData.LastConsumedFrame = StartFrame - 1;
+					}
+					else
+					{
+						UE_NP_TRACE_SYSTEM_FAULT("Recovering from input stream starvation on server. Advancing over dropped client frame %d", ServerRecvData.LastConsumedFrame + 1);
+					}
 				}
 
 				npEnsure(Frame >= 0);
@@ -378,6 +399,8 @@ public:
 	}
 };
 
+static constexpr float TimeDilationNetSerializationScale = 10000.f;	// Used when net-serializing a float into an int16
+
 template<typename InModelDef>
 class TFixedTickReplicator_AP
 {
@@ -395,6 +418,18 @@ public:
 		const int32 LastConsumedInputFrame = FNetworkPredictionSerialization::ReadCompressedFrame(Ar, TickState->PendingFrame); // 1. Last Consumed (Client) Input Frame
 		const int32 ServerFrame = FNetworkPredictionSerialization::ReadCompressedFrame(Ar, TickState->PendingFrame + TickState->Offset); // 2. Server's Frame
 
+		Ar.SerializeBits(&TickState->bHasTimeDilation, 1); // 3. Whether We have Time Dilation Or Not. Most Of The time We Do Not.. We Don't Need to Send int16 over network
+		if (TickState->bHasTimeDilation)
+		{
+			int16 NetworkTimeDilation = 1;
+			Ar << NetworkTimeDilation; // 4. Time Dilation
+			TickState->SuggestedTimeDilation = FMath::Clamp((float)NetworkTimeDilation / TimeDilationNetSerializationScale, NetworkPredictionCVars::TimeDilationMin->GetFloat(), NetworkPredictionCVars::TimeDilationMax->GetFloat());
+		}
+		else
+		{
+			TickState->SuggestedTimeDilation = 1.f;
+		}
+
 		npEnsure(LastConsumedInputFrame <= TickState->PendingFrame);
 		npEnsure(ServerFrame >= 0);
 
@@ -433,9 +468,9 @@ public:
 		npEnsureSlow(ClientRecvState.InstanceIdx >= 0);
 		TInstanceData<ModelDef>& InstanceData = DataStore->Instances.GetByIndexChecked(ClientRecvState.InstanceIdx);
 
-		TCommonReplicator_AP<ModelDef>::NetRecv(P, InstanceData, ClientRecvState); // 3. Common
+		TCommonReplicator_AP<ModelDef>::NetRecv(P, InstanceData, ClientRecvState); // 5. Common
 
-		InstanceData.CueDispatcher->NetRecvSavedCues(P.Ar, true, ServerFrame, 0); // 4. NetSimCues
+		InstanceData.CueDispatcher->NetRecvSavedCues(P.Ar, true, ServerFrame, 0); // 6. NetSimCues
 	}
 
 	// ------------------------------------------------------------------------------------------------------------
@@ -452,10 +487,12 @@ public:
 		npCheckSlow(Frames);
 
 		int32 LastConsumedFrame = INDEX_NONE;
+		int32 LastReceivedFrame = INDEX_NONE;
 		TServerRecvData_Fixed<ModelDef>* ServerRecvState = DataStore->ServerRecv.Find(ID);
 		if (ServerRecvState)
 		{
 			LastConsumedFrame = ServerRecvState->LastConsumedFrame;
+			LastReceivedFrame = ServerRecvState->LastRecvFrame;
 		}
 
 		const int32 PendingFrame = TickState->PendingFrame;
@@ -464,9 +501,45 @@ public:
 		FNetworkPredictionSerialization::WriteCompressedFrame(Ar, LastConsumedFrame); // 1. Last Consumed Input Frame (Client's frame)
 		FNetworkPredictionSerialization::WriteCompressedFrame(Ar, PendingFrame); // 2. PendingFrame (Server's frame)
 
-		TCommonReplicator_AP<ModelDef>::NetSend(P, *Instance, Frames->Buffer[PendingFrame]); // 3. Common
+		int32 BufferOffset = LastConsumedFrame - LastReceivedFrame + 2;
 
-		Instance->CueDispatcher->NetSendSavedCues(P.Ar, ENetSimCueReplicationTarget::AutoProxy, true); // 4. NetSimCues
+		if (!NetworkPredictionCVars::AllowTimeDilationEscalation->GetBool())
+		{
+			BufferOffset = FMath::Clamp(BufferOffset, -1, 1);
+		}
+
+		if (LastReceivedFrame == INDEX_NONE || LastConsumedFrame == INDEX_NONE)
+		{
+			BufferOffset = 0;
+		}
+
+		// Calculate desired dilation and send to client
+		float CalculatedTimeDilation = 1.f;
+		bool HasTimeDilation = false;
+
+		if (BufferOffset != 0)
+		{
+			const float TimeDilationDecay = FMath::Clamp(1.0f - (NetworkPredictionCVars::TimeDilationEscalationDecay->GetFloat() * FMath::Abs(BufferOffset)), NetworkPredictionCVars::TimeDilationEscalationDecayMax->GetFloat(), 1.0f);
+			CalculatedTimeDilation = 1.0f + ((NetworkPredictionCVars::TimeDilationAmount->GetFloat() * -BufferOffset) * TimeDilationDecay);
+			CalculatedTimeDilation = FMath::Clamp(CalculatedTimeDilation, NetworkPredictionCVars::TimeDilationMin->GetFloat(), NetworkPredictionCVars::TimeDilationMax->GetFloat());
+			HasTimeDilation = true;
+		}
+		else
+		{
+			CalculatedTimeDilation = 0.f;
+		}
+
+		Ar.SerializeBits(&HasTimeDilation, 1); // 3. Whether we have active time dilation. Most of the time we don't, so we can avoid sending the int16 over the network.
+		
+		if (HasTimeDilation)
+		{
+			int16 NetworkTimeDilation = FMath::RoundToInt(CalculatedTimeDilation * TimeDilationNetSerializationScale);
+			Ar << NetworkTimeDilation; // 4. Time Dilation
+		}
+
+		TCommonReplicator_AP<ModelDef>::NetSend(P, *Instance, Frames->Buffer[PendingFrame]); // 5. Common
+
+		Instance->CueDispatcher->NetSendSavedCues(P.Ar, ENetSimCueReplicationTarget::AutoProxy, true); // 6. NetSimCues
 	}
 };
 
@@ -706,3 +779,6 @@ private:
 };
 
 // ---------------------------------------------------------------------------------------------------------------------------
+
+
+#undef UE_API
diff --git a/Engine/Plugins/Runtime/NetworkPrediction/Source/NetworkPrediction/Public/NetworkPredictionTickState.h b/Engine/Plugins/Runtime/NetworkPrediction/Source/NetworkPrediction/Public/NetworkPredictionTickState.h
--- a/Engine/Plugins/Runtime/NetworkPrediction/Source/NetworkPrediction/Public/NetworkPredictionTickState.h
+++ b/Engine/Plugins/Runtime/NetworkPrediction/Source/NetworkPrediction/Public/NetworkPredictionTickState.h
@@ -48,7 +48,10 @@ struct FFixedTickState
 
 	// Realtime steps. That is, one of these = one FixedStepMS in simulation time.
 	// This means sim time ticks slightly slower than real time.
-	float FixedStepRealTimeMS = 1/30.f;
+	float FixedStepRealTimeMS = 1.0f /30.0f;
+
+	// If time dilation is enabled, one of these = one FixedStepMS in simulation time. This will shrink or grow slightly as network latency changes.
+	float FixedStepDilatedRealTimeMS = 1.0f / 30.0f;
 
 	// Next frame to be ticked (used as input to generated PendingFrame+1)
 	int32 PendingFrame=0;
@@ -64,6 +67,12 @@ struct FFixedTickState
 	// Accumulates raw delta time into our fixed steps
 	float UnspentTimeMS = 0.f;
 
+	// Signals whether any time dilation is active. If so, use @SuggestedTimeDilation
+	bool bHasTimeDilation = false;
+	
+	// Only valid when bHasTimeDilation is true. Represents a percentage of time speed up/down.  1.0 = no dilation.
+	float SuggestedTimeDilation = 1.f;
+
 	struct FInterpolationState
 	{
 		float AccumulatedTimeMS = 0.f; // accumulated real time
diff --git a/Engine/Plugins/Runtime/NetworkPrediction/Source/NetworkPrediction/Public/Services/NetworkPredictionService_Input.inl b/Engine/Plugins/Runtime/NetworkPrediction/Source/NetworkPrediction/Public/Services/NetworkPredictionService_Input.inl
--- a/Engine/Plugins/Runtime/NetworkPrediction/Source/NetworkPrediction/Public/Services/NetworkPredictionService_Input.inl
+++ b/Engine/Plugins/Runtime/NetworkPrediction/Source/NetworkPrediction/Public/Services/NetworkPredictionService_Input.inl
@@ -155,7 +155,7 @@ public:
 			else if (NumBufferedInputCmds == 0)
 			{
 				// No Cmds to process, enter fault state. Increment FaultLimit each time this happens.
-				// TODO: We should have something to bring this back down (which means skipping frames) we don't want temporary poor conditions to cause permanent high input buffering
+				// Note: time dilation will gradually bring any excessive input buffering back down as net latency improves. When dilation is enabled, the FaultLimit project setting can be kept to 1 safely.
 				Remote.bFault = true;
 				Remote.FaultLimit = FMath::Min(Remote.FaultLimit + 1, FMath::Min(MaxFaultLimit, ServerRecvData.InputBuffer.Capacity()));
 				
